Editgrating RigLogic plugin from Dev-CharacterTech to UE5/Main @ CL 13891876.

diff --git a/Engine/Plugins/Animation/RigLogic/LICENCE.txt b/Engine/Plugins/Animation/RigLogic/LICENCE.txt
new file mode 100644
--- /dev/null
+++ b/Engine/Plugins/Animation/RigLogic/LICENCE.txt
@@ -0,0 +1,11 @@
+Copyright (c) 2017-2018, 3Lateral. All rights reserved.
+Owner: 3Lateral - www.3lateral.com.
+License:
+    The use of this source code and its functionality is regulated by the licensing terms defined by 3Lateral.
+    The source code has been delivered for the use only for this project and under strict conditions that were defined by 3Lateral.
+
+    Please note that this infrastructure, code and its functionality are subject to Confidentiality Agreement and that you
+    should not do or let anyone else do any decompilation, reverse engineering or any other attempt to learn the inner workings
+    of the system. In case you wish to apply this infrastructure, code and its functionality on any other project
+    (both commercial or non-commercial) please consult 3Lateral regarding the license terms and fees. Furthermore,
+    please consult and confirm with 3Lateral when it comes to allowing any other third party to access this infrastructure or code.
\ No newline at end of file
diff --git a/Engine/Plugins/Animation/RigLogic/RigLogic.uplugin b/Engine/Plugins/Animation/RigLogic/RigLogic.uplugin
new file mode 100644
--- /dev/null
+++ b/Engine/Plugins/Animation/RigLogic/RigLogic.uplugin
@@ -0,0 +1,64 @@
+{
+	"FileVersion": 3,
+	"Version": 5,
+	"VersionName": "5.1.0",
+	"FriendlyName": "RigLogic Plugin v5.1.0",
+	"Description": "3Lateral RigLogic Plugin for Facial Animation v5.1.0",
+	"Category": "Animation",
+	"CreatedBy": "3Lateral",
+	"CreatedByURL": "http://www.3lateral.com/",
+	"DocsURL": "",
+	"MarketplaceURL": "",
+	"SupportURL": "",
+	"CanContainContent": false,
+	"IsBetaVersion": true,
+	"Installed": true,
+	"Modules": [
+		{
+			"Name": "RigLogicLib",
+			"Type": "Runtime",
+			"LoadingPhase": "PreDefault",
+			"WhitelistPlatforms": [
+				"Win32",
+				"Win64",
+				"Linux",
+				"Mac",
+				"PS4",
+				"XboxOne"
+			]
+		},
+		{
+			"Name": "RigLogicLibTest",
+			"Type": "Runtime",
+			"LoadingPhase": "PreDefault",
+			"WhitelistPlatforms": [
+				"Win32",
+				"Win64",
+				"Linux",
+				"Mac"
+			]
+		},
+		{
+			"Name": "RigLogicModule",
+			"Type": "Runtime",
+			"LoadingPhase": "Default",
+			"WhitelistPlatforms": [
+				"Win64"
+			]
+		},
+		{
+			"Name": "RigLogicEditor",
+			"Type": "Editor",
+			"LoadingPhase": "Default",
+			"WhitelistPlatforms": [
+				"Win64"
+			]
+		}
+	],
+	"Plugins": [
+		{
+			"Name": "ControlRig",
+			"Enabled": true
+		}
+	]
+}
\ No newline at end of file
diff --git a/Engine/Plugins/Animation/RigLogic/Source/RigLogicEditor/Private/DNAAssetImportFactory.cpp b/Engine/Plugins/Animation/RigLogic/Source/RigLogicEditor/Private/DNAAssetImportFactory.cpp
new file mode 100644
--- /dev/null
+++ b/Engine/Plugins/Animation/RigLogic/Source/RigLogicEditor/Private/DNAAssetImportFactory.cpp
@@ -0,0 +1,187 @@
+// Copyright 1998-2019 Epic Games, Inc. All Rights Reserved.
+
+#include "DNAAssetImportFactory.h"
+#include "DNAAsset.h"
+
+#include "Runtime/Engine/Classes/Engine/SkeletalMesh.h"
+#include "Runtime/Engine/Public/ReferenceSkeleton.h"
+#include "Runtime/Engine/Public/Rendering/SkeletalMeshModel.h"
+#include "Runtime/Engine/Public/Rendering/SkeletalMeshLODModel.h"
+#include "Runtime/Engine/Public/Rendering/SkeletalMeshRenderData.h"
+#include "Runtime/Engine/Public/Rendering/SkeletalMeshLODRenderData.h"
+#include "Runtime/Engine/Public/ComponentReregisterContext.h"
+#include "Runtime/CoreUObject/Public/UObject/ConstructorHelpers.h"
+
+#include "ActorFactories/ActorFactoryStaticMesh.h"
+#include "ProfilingDebugging/ScopedTimers.h"
+
+#include "Engine/StaticMesh.h"
+#include "Misc/Paths.h"
+
+#include "AssetImportTask.h"
+#include "Runtime/Core/Public/Misc/FileHelper.h"
+
+#include "SkelMeshDNAUtils.h"
+#include "DNAAssetImportUI.h"
+#include "DNAImporter.h"
+
+DEFINE_LOG_CATEGORY(LogDNAImportFactory);
+#define LOCTEXT_NAMESPACE "DNAAssetImportFactory"
+
+UDNAAssetImportFactory::UDNAAssetImportFactory(const FObjectInitializer& ObjectInitializer)
+	: Super(ObjectInitializer)
+{
+	bCreateNew = false;
+	bEditAfterNew = true;
+	SupportedClass = UDNAAsset::StaticClass();
+
+	bEditorImport = true;
+	bText = false;
+
+	Formats.Add(TEXT("dna;Character DNA file"));
+}
+
+void UDNAAssetImportFactory::PostInitProperties()
+{
+	Super::PostInitProperties();
+	bEditorImport = true;
+	bText = false;
+
+	ImportUI = NewObject<UDNAAssetImportUI>(this, NAME_None, RF_NoFlags);
+}
+
+UObject* UDNAAssetImportFactory::FactoryCreateFile(UClass* InClass, UObject* InParent, FName InName, EObjectFlags Flags, const FString& Filename, const TCHAR* Parms, FFeedbackContext* Warn, bool& bOutOperationCanceled)
+{
+	// this one prints all messages that are stored in FFbxImporter
+	FDNAImporter* DNAImporter = FDNAImporter::GetInstance();
+	FDNAAssetImportOptions* ImportOptions = DNAImporter->GetImportOptions();
+	//Clean up the options
+	FDNAAssetImportOptions::ResetOptions(ImportOptions);	
+
+	struct FRestoreImportUI
+	{
+		FRestoreImportUI(UDNAAssetImportFactory* InDNAAssetImportFactory)
+			: DNAAssetImportFactory(InDNAAssetImportFactory)
+		{
+			ensure(DNAAssetImportFactory->OriginalImportUI == nullptr);
+			DNAAssetImportFactory->OriginalImportUI = DNAAssetImportFactory->ImportUI;
+		}
+
+		~FRestoreImportUI()
+		{
+			DNAAssetImportFactory->ImportUI = DNAAssetImportFactory->OriginalImportUI;
+			DNAAssetImportFactory->OriginalImportUI = nullptr;
+		}
+
+	private:
+		UDNAAssetImportFactory* DNAAssetImportFactory;
+	};
+	FRestoreImportUI RestoreImportUI(this);
+	UDNAAssetImportUI* OverrideImportUI = AssetImportTask ? Cast<UDNAAssetImportUI>(AssetImportTask->Options) : nullptr;
+	if (OverrideImportUI)
+	{
+		ImportUI = OverrideImportUI;
+	}
+	//We are not re-importing
+	ImportUI->bIsReimport = false;
+	ImportUI->ReimportMesh = nullptr;
+
+	// Show the import dialog only when not in a "yes to all" state or when automating import
+	bool bIsAutomated = IsAutomatedImport();
+	bool bShowImportDialog = !bIsAutomated;
+
+	ImportOptions = GetImportOptions(DNAImporter, ImportUI, bShowImportDialog, bIsAutomated, InParent->GetPathName(), bOutOperationCanceled, UFactory::CurrentFilename);
+
+	if (bOutOperationCanceled)
+	{
+		// User cancelled, clean up and return
+		DNAImporter->PartialCleanUp();
+		return nullptr;
+	}
+
+	if (ImportOptions)
+	{
+		USkeletalMesh* SkelMesh = ImportOptions->SkeletalMesh;
+		if (SkelMesh)
+		{
+			FString SkelMeshFullPath(SkelMesh->GetPathName());
+			FString SkelMeshJustPath(FPaths::GetPath(*SkelMeshFullPath));
+			FString SkelMeshFileName(FPaths::GetBaseFilename(*SkelMeshFullPath));
+			FString DNAAssetFileName(TEXT("DNA"));
+			FString DNAAssetFullName(SkelMeshJustPath);
+			DNAAssetFullName.PathAppend(*SkelMeshFileName, SkelMeshFileName.Len());
+			DNAAssetFullName.Append(TEXT("."));
+			DNAAssetFullName.Append(DNAAssetFileName);
+
+			GEditor->GetEditorSubsystem<UImportSubsystem>()->BroadcastAssetPreImport(this, InClass, InParent, *DNAAssetFullName, TEXT("dna"));
+			
+			UDNAAsset* DNAAsset = NewObject< UDNAAsset >(SkelMesh, FName(*DNAAssetFileName)); //SkelMesh has to be its outer, otherwise DNAAsset won't be saved			
+			if (DNAAsset->Init(*Filename))
+			{
+				UAssetUserData* DNAAssetUserData = Cast<UAssetUserData>(DNAAsset);
+				SkelMesh->AddAssetUserData(DNAAssetUserData);
+
+				GEditor->GetEditorSubsystem<UImportSubsystem>()->BroadcastAssetPostImport(this, DNAAsset);
+				DNAImporter->PartialCleanUp();
+				return DNAAsset;
+			}
+		}
+		else
+		{
+			// Skeletal mesh not found.
+			const FText Message = LOCTEXT("DNAImportFailed", "DNA file can not be uploaded without Skeletal Mesh.");
+			FMessageDialog::Open(EAppMsgType::Ok, Message);
+			UE_LOG(LogDNAImportFactory, Warning, TEXT("%s"), *Message.ToString());
+		}
+	}
+
+	// Failed to load file or create DNA stream. Clean and return nullptr.
+	DNAImporter->PartialCleanUp();	
+	GEditor->GetEditorSubsystem<UImportSubsystem>()->BroadcastAssetPostImport(this, nullptr);
+	return nullptr;
+}
+
+bool UDNAAssetImportFactory::FactoryCanImport(const FString& Filename)
+{
+	const FString Extension = FPaths::GetExtension(Filename);
+
+	if (Extension == TEXT("dna"))
+	{
+		return true;
+	}
+
+	return false;
+}
+
+void UDNAAssetImportFactory::CleanUp()
+{
+}
+
+// DNA Asset Import UI Implementation.
+UDNAAssetImportUI::UDNAAssetImportUI(const FObjectInitializer& ObjectInitializer)
+	: Super(ObjectInitializer)
+{
+	bIsReimport = false;
+	ReimportMesh = nullptr;
+	//Make sure we are transactional to allow undo redo
+	this->SetFlags(RF_Transactional);
+}
+
+
+bool UDNAAssetImportUI::CanEditChange(const FProperty* InProperty) const
+{
+	bool bIsMutable = Super::CanEditChange(InProperty);
+	return bIsMutable;
+}
+
+void UDNAAssetImportUI::ParseFromJson(TSharedRef<class FJsonObject> ImportSettingsJson)
+{
+}
+
+void UDNAAssetImportUI::ResetToDefault()
+{
+	ReloadConfig();
+}
+// END of DNA Asset Import UI Implementation.
+
+#undef LOCTEXT_NAMESPACE
diff --git a/Engine/Plugins/Animation/RigLogic/Source/RigLogicEditor/Private/DNAImporter.cpp b/Engine/Plugins/Animation/RigLogic/Source/RigLogicEditor/Private/DNAImporter.cpp
new file mode 100644
--- /dev/null
+++ b/Engine/Plugins/Animation/RigLogic/Source/RigLogicEditor/Private/DNAImporter.cpp
@@ -0,0 +1,193 @@
+// Copyright 1998-2019 Epic Games, Inc. All Rights Reserved.
+
+/*=============================================================================
+	Main implementation of FDNAImporter : import DNA data to Skeletal Mesh
+=============================================================================*/
+
+#include "DNAImporter.h"
+#include "CoreMinimal.h"
+#include "Modules/ModuleManager.h"
+#include "AssetRegistryModule.h"
+#include "Interfaces/IMainFrameModule.h"
+#include "ARFilter.h"
+#include "Widgets/SWindow.h"
+#include "Framework/Application/SlateApplication.h"
+#include "HAL/PlatformApplicationMisc.h"
+#include "DNAAssetImportWindow.h"
+
+TSharedPtr<FDNAImporter> FDNAImporter::StaticInstance;
+TSharedPtr<FDNAImporter> FDNAImporter::StaticPreviewInstance;
+
+FDNAAssetImportOptions* GetImportOptions(FDNAImporter * DNAImporter, UDNAAssetImportUI * ImportUI, bool bShowOptionDialog, bool bIsAutomated, const FString & FullPath, bool & OutOperationCanceled, const FString & InFilename)
+{
+	OutOperationCanceled = false;
+
+	if (bShowOptionDialog)
+	{
+		FDNAAssetImportOptions* ImportOptions = DNAImporter->GetImportOptions();
+
+		// if SkeletalMesh was set by outside, please make sure copy back to UI
+		if (ImportOptions->SkeletalMesh)
+		{
+			ImportUI->SkeletalMesh = ImportOptions->SkeletalMesh;
+		}
+		else
+		{
+			// Look in the current target directory to see if we have a skeleton
+			FARFilter Filter;
+			Filter.PackagePaths.Add(*FPaths::GetPath(FullPath));
+			Filter.ClassNames.Add(USkeletalMesh::StaticClass()->GetFName());
+
+			IAssetRegistry& AssetRegistry = FModuleManager::LoadModuleChecked<FAssetRegistryModule>("AssetRegistry").Get();
+			TArray<FAssetData> SkeletalMeshAssets;
+			AssetRegistry.GetAssets(Filter, SkeletalMeshAssets);
+			if (SkeletalMeshAssets.Num() > 0)
+			{
+				ImportUI->SkeletalMesh = CastChecked<USkeletalMesh>(SkeletalMeshAssets[0].GetAsset());
+			}
+			else
+			{
+				ImportUI->SkeletalMesh = NULL;
+			}
+		}
+
+		//This option must always be the same value has the skeletalmesh one.
+
+		//////////////////////////////////////////////////////////////////////////
+		// Set the information section data
+
+		//Make sure the file is open to be able to read the header before showing the options
+		//If the file is already open it will simply return false.
+		// do analytics on getting DNA data
+		ImportUI->FileVersion = DNAImporter->GetDNAFileVersion();
+		ImportUI->FileCreator = DNAImporter->GetDNAFileCreator();
+
+		TSharedPtr<SWindow> ParentWindow;
+
+		if (FModuleManager::Get().IsModuleLoaded("MainFrame"))
+		{
+			IMainFrameModule& MainFrame = FModuleManager::LoadModuleChecked<IMainFrameModule>("MainFrame");
+			ParentWindow = MainFrame.GetParentWindow();
+		}
+
+		// Compute centered window position based on max window size, which include when all categories are expanded
+		const float DNAImportWindowWidth = 410.0f;
+		const float DNAImportWindowHeight = 750.0f;
+		FVector2D DNAImportWindowSize = FVector2D(DNAImportWindowWidth, DNAImportWindowHeight); // Max window size it can get based on current slate
+
+
+		FSlateRect WorkAreaRect = FSlateApplicationBase::Get().GetPreferredWorkArea();
+		FVector2D DisplayTopLeft(WorkAreaRect.Left, WorkAreaRect.Top);
+		FVector2D DisplaySize(WorkAreaRect.Right - WorkAreaRect.Left, WorkAreaRect.Bottom - WorkAreaRect.Top);
+
+		float ScaleFactor = FPlatformApplicationMisc::GetDPIScaleFactorAtPoint(DisplayTopLeft.X, DisplayTopLeft.Y);
+		DNAImportWindowSize *= ScaleFactor;
+
+		FVector2D WindowPosition = (DisplayTopLeft + (DisplaySize - DNAImportWindowSize) / 2.0f) / ScaleFactor;
+
+
+		TSharedRef<SWindow> Window = SNew(SWindow)
+			.Title(NSLOCTEXT("UnrealEd", "DNAImportOpionsTitle", "DNA Import Options"))
+			.SizingRule(ESizingRule::Autosized)
+			.AutoCenter(EAutoCenter::None)
+			.ClientSize(DNAImportWindowSize)
+			.ScreenPosition(WindowPosition);
+
+		TSharedPtr<SDNAAssetImportWindow> DNAImportWindow;
+		Window->SetContent
+		(
+			SAssignNew(DNAImportWindow, SDNAAssetImportWindow)
+			.ImportUI(ImportUI)
+			.WidgetWindow(Window)
+			.FullPath(FText::FromString(FullPath))
+			.MaxWindowHeight(DNAImportWindowHeight)
+			.MaxWindowWidth(DNAImportWindowWidth)
+		);
+
+		// @todo: we can make this slow as showing progress bar later
+		FSlateApplication::Get().AddModalWindow(Window, ParentWindow, false);
+
+		ImportUI->SaveConfig();
+
+		if (DNAImportWindow->ShouldImport())
+		{
+			// open dialog
+			// see if it's canceled
+			ApplyImportUIToImportOptions(ImportUI, *ImportOptions);
+
+			return ImportOptions;
+		}
+		else
+		{
+			OutOperationCanceled = true;
+		}
+	}
+	else if (bIsAutomated)
+	{
+		//Automation tests set ImportUI settings directly.  Just copy them over
+		FDNAAssetImportOptions* ImportOptions = DNAImporter->GetImportOptions();
+		//Clean up the options
+		FDNAAssetImportOptions::ResetOptions(ImportOptions);
+		ApplyImportUIToImportOptions(ImportUI, *ImportOptions);
+		return ImportOptions;
+	}
+	else
+	{
+		return DNAImporter->GetImportOptions();
+	}
+
+	return NULL;
+}
+
+void ApplyImportUIToImportOptions(UDNAAssetImportUI* ImportUI, FDNAAssetImportOptions& InOutImportOptions)
+{
+	check(ImportUI);
+	InOutImportOptions.SkeletalMesh = ImportUI->SkeletalMesh;
+}
+
+FDNAImporter::~FDNAImporter()
+{
+	CleanUp();
+}
+
+FDNAImporter* FDNAImporter::GetInstance()
+{
+	if (!StaticInstance.IsValid())
+	{
+		StaticInstance = MakeShareable(new FDNAImporter());
+	}
+	return StaticInstance.Get();
+}
+
+void FDNAImporter::DeleteInstance()
+{
+	StaticInstance.Reset();
+}
+
+FDNAAssetImportOptions* FDNAImporter::GetImportOptions() const
+{
+	return ImportOptions;
+}
+
+FDNAImporter::FDNAImporter()
+	: ImportOptions(NULL)
+{
+	ImportOptions = new FDNAAssetImportOptions();
+	FMemory::Memzero(*ImportOptions);
+
+	CurPhase = NOTSTARTED;
+}
+
+void FDNAImporter::CleanUp()
+{
+	PartialCleanUp();
+
+	delete ImportOptions;
+	ImportOptions = NULL;
+}
+
+void FDNAImporter::PartialCleanUp()
+{
+	// reset
+	CurPhase = NOTSTARTED;
+}
diff --git a/Engine/Plugins/Animation/RigLogic/Source/RigLogicEditor/Private/RigLogicEditor.cpp b/Engine/Plugins/Animation/RigLogic/Source/RigLogicEditor/Private/RigLogicEditor.cpp
new file mode 100644
--- /dev/null
+++ b/Engine/Plugins/Animation/RigLogic/Source/RigLogicEditor/Private/RigLogicEditor.cpp
@@ -0,0 +1,29 @@
+// Copyright 1998-2019 Epic Games, Inc. All Rights Reserved.
+
+#pragma once
+
+#include "RigLogicEditor.h"
+#include "Features/IModularFeatures.h"
+#include "Engine/Engine.h"
+#include "Misc/Paths.h"
+#include "Modules/ModuleManager.h"
+#include "Logging/LogMacros.h"
+
+#include "RigUnit_RigLogic.h"
+
+IMPLEMENT_MODULE(FRigLogicEditor, RigLogicEditor)
+
+DEFINE_LOG_CATEGORY_STATIC(LogRigLogicEditor, Log, All);
+
+FRigLogicEditor::FRigLogicEditor()
+{
+}
+
+void FRigLogicEditor::StartupModule()
+{
+}
+
+void FRigLogicEditor::ShutdownModule()
+{
+}
+
diff --git a/Engine/Plugins/Animation/RigLogic/Source/RigLogicEditor/Public/DNAAssetImportFactory.h b/Engine/Plugins/Animation/RigLogic/Source/RigLogicEditor/Public/DNAAssetImportFactory.h
new file mode 100644
--- /dev/null
+++ b/Engine/Plugins/Animation/RigLogic/Source/RigLogicEditor/Public/DNAAssetImportFactory.h
@@ -0,0 +1,36 @@
+// Copyright 1998-2019 Epic Games, Inc. All Rights Reserved.
+
+#pragma once 
+
+#include "Factories/Factory.h"
+#include "Factories/ImportSettings.h"
+#include "EditorReimportHandler.h"
+#include "DNAAsset.h"
+#include "Runtime/Engine/Classes/Engine/SkeletalMesh.h"
+#include "DNAAssetImportFactory.generated.h"
+
+DECLARE_LOG_CATEGORY_EXTERN(LogDNAImportFactory, Log, All);
+
+/** Factory responsible for importing DNA file and attaching DNA data into SkeletalMesh
+ */
+UCLASS(transient)
+class UDNAAssetImportFactory: public UFactory
+{ 
+	GENERATED_UCLASS_BODY()
+
+public:
+	UPROPERTY()
+		class UDNAAssetImportUI* ImportUI;
+
+	/** Prevent garbage collection of original when overriding ImportUI property */
+	UPROPERTY()
+		class UDNAAssetImportUI* OriginalImportUI;
+
+	/** UObject properties */
+	virtual void PostInitProperties() override;
+
+	/** UFactory interface */
+	virtual UObject* FactoryCreateFile(UClass* InClass, UObject* InParent, FName InName, EObjectFlags Flags, const FString& Filename, const TCHAR* Parms, FFeedbackContext* Warn, bool& bOutOperationCanceled) override;
+	virtual bool FactoryCanImport(const FString& Filename) override;
+	virtual void CleanUp() override;
+};
diff --git a/Engine/Plugins/Animation/RigLogic/Source/RigLogicEditor/Public/DNAAssetImportUI.h b/Engine/Plugins/Animation/RigLogic/Source/RigLogicEditor/Public/DNAAssetImportUI.h
new file mode 100644
--- /dev/null
+++ b/Engine/Plugins/Animation/RigLogic/Source/RigLogicEditor/Public/DNAAssetImportUI.h
@@ -0,0 +1,54 @@
+// Copyright 1998-2019 Epic Games, Inc. All Rights Reserved.
+
+/**
+ * DNA Asset Importer UI options.
+ */
+
+#pragma once
+
+#include "CoreMinimal.h"
+#include "Engine/SkeletalMesh.h"
+#include "UObject/ObjectMacros.h"
+#include "UObject/Object.h"
+#include "Factories/ImportSettings.h"
+#include "DNAAssetImportUI.generated.h"
+
+UCLASS(config = EditorPerProjectUserSettings, AutoExpandCategories = (Mesh), MinimalAPI)
+class UDNAAssetImportUI : public UObject, public IImportSettingsParser
+{
+	GENERATED_UCLASS_BODY()
+
+public:
+
+	/** Skeletal mesh to use for imported DNA asset. When importing DNA, leaving this as "None" will generate new skeletal mesh. */
+	UPROPERTY(EditAnywhere, BlueprintReadWrite, Category = Mesh, meta = (ImportType = "DNAAsset"))
+		class USkeletalMesh* SkeletalMesh;
+
+	UFUNCTION(BlueprintCallable, Category = Miscellaneous)
+		void ResetToDefault();
+
+	/** UObject Interface */
+	virtual bool CanEditChange(const FProperty* InProperty) const override;
+
+	/** IImportSettings Interface */
+	virtual void ParseFromJson(TSharedRef<class FJsonObject> ImportSettingsJson) override;
+
+	/* Whether this UI is construct for a reimport */
+	bool bIsReimport;
+
+	/* When we are reimporting, we need the current object to preview skeletal mesh match issues. */
+	UObject* ReimportMesh;
+
+
+	//////////////////////////////////////////////////////////////////////////
+		// DNA Asset file informations
+		// Transient value that are set everytime we show the options dialog. These are information only and should be string.
+
+	/* The DNA file version */
+	UPROPERTY(VisibleAnywhere, Transient, Category = DNAFileInformation, meta = (ImportType = "DNAAsset", DisplayName = "File Version"))
+		FString FileVersion;
+
+	/* The file creator information */
+	UPROPERTY(VisibleAnywhere, Transient, Category = DNAFileInformation, meta = (ImportType = "DNAAsset", DisplayName = "File Creator"))
+		FString FileCreator;
+};
diff --git a/Engine/Plugins/Animation/RigLogic/Source/RigLogicEditor/Public/DNAImporter.h b/Engine/Plugins/Animation/RigLogic/Source/RigLogicEditor/Public/DNAImporter.h
new file mode 100644
--- /dev/null
+++ b/Engine/Plugins/Animation/RigLogic/Source/RigLogicEditor/Public/DNAImporter.h
@@ -0,0 +1,78 @@
+// Copyright 1998-2019 Epic Games, Inc. All Rights Reserved.
+
+#pragma once
+
+#include "CoreMinimal.h"
+#include "UObject/WeakObjectPtr.h"
+#include "Misc/SecureHash.h"
+#include "DNAAssetImportUI.h"
+#include "Logging/TokenizedMessage.h"
+#include "MeshBuild.h"
+
+class USkeletalMesh;
+
+struct FDNAAssetImportOptions
+{
+	// General options
+	bool bCanShowDialog;	
+	// Skeletal Mesh options
+	USkeletalMesh* SkeletalMesh;
+
+	static void ResetOptions(FDNAAssetImportOptions *OptionsToReset)
+	{
+		check(OptionsToReset != nullptr);
+		*OptionsToReset = FDNAAssetImportOptions();
+	}
+};
+
+FDNAAssetImportOptions* GetImportOptions(class FDNAImporter* DNAImporter, UDNAAssetImportUI* ImportUI, bool bShowOptionDialog, bool bIsAutomated, const FString& FullPath, bool& OutOperationCanceled, const FString& InFilename);
+void ApplyImportUIToImportOptions(UDNAAssetImportUI* ImportUI, FDNAAssetImportOptions& InOutImportOptions);
+
+class FDNAImporter
+{
+public:
+	~FDNAImporter();
+	/**
+	 * Returns the importer singleton. It will be created on the first request.
+	 */
+	static FDNAImporter* GetInstance();
+	static void DeleteInstance();
+
+	FDNAAssetImportOptions* ImportOptions;
+	/**
+	 * Get the object of import options
+	 *
+	 * @return FDNAAssetImportOptions
+	 */
+	FDNAAssetImportOptions* GetImportOptions() const;
+	void PartialCleanUp();
+	FString GetDNAFileVersion() { return DNAFileVersion; }
+	FString GetDNAFileCreator() { return DNAFileCreator; }
+protected:
+	enum IMPORTPHASE
+	{
+		NOTSTARTED,
+		FILEOPENED,
+		IMPORTED,
+		FIXEDANDCONVERTED,
+	};
+
+	static TSharedPtr<FDNAImporter> StaticInstance;
+	static TSharedPtr<FDNAImporter> StaticPreviewInstance;
+
+	// scene management
+	IMPORTPHASE CurPhase;
+	FString ErrorMessage;
+	// base path of DNA file
+	FString FileBasePath;
+	TWeakObjectPtr<UObject> Parent;
+	FString DNAFileVersion;
+	FString DNAFileCreator;
+
+	FDNAImporter();
+
+	/**
+	 * Clean up to reset the import.
+	 */
+	void CleanUp();
+};
\ No newline at end of file
diff --git a/Engine/Plugins/Animation/RigLogic/Source/RigLogicEditor/Public/RigLogicEditor.h b/Engine/Plugins/Animation/RigLogic/Source/RigLogicEditor/Public/RigLogicEditor.h
new file mode 100644
--- /dev/null
+++ b/Engine/Plugins/Animation/RigLogic/Source/RigLogicEditor/Public/RigLogicEditor.h
@@ -0,0 +1,18 @@
+// Copyright 1998-2019 Epic Games, Inc. All Rights Reserved.
+
+#pragma once
+
+#include "CoreMinimal.h"
+#include "Modules/ModuleInterface.h"
+
+struct RIGLOGICEDITOR_API FRigLogicEditor : IModuleInterface 
+{
+public:
+	virtual void StartupModule() override;
+	virtual void ShutdownModule() override;
+
+	FRigLogicEditor();
+
+private:
+	TArray<TSharedRef<class IAssetTypeActions>> AssetTypeActions;
+};
diff --git a/Engine/Plugins/Animation/RigLogic/Source/RigLogicEditor/RigLogicEditor.Build.cs b/Engine/Plugins/Animation/RigLogic/Source/RigLogicEditor/RigLogicEditor.Build.cs
new file mode 100644
--- /dev/null
+++ b/Engine/Plugins/Animation/RigLogic/Source/RigLogicEditor/RigLogicEditor.Build.cs
@@ -0,0 +1,53 @@
+ // Copyright 1998-2019 Epic Games, Inc. All Rights Reserved.
+
+using System.IO;
+
+namespace UnrealBuildTool.Rules
+{
+    public class RigLogicEditor : ModuleRules
+    {
+        private string ModulePath
+        {
+            get { return ModuleDirectory; }
+        }
+
+        private string ThirdPartyPath
+        {
+            get { return Path.GetFullPath(Path.Combine(ModulePath, "../Source/ThirdParty/")); }
+        }
+
+        public RigLogicEditor(ReadOnlyTargetRules Target) : base(Target)
+        {
+            PublicDefinitions.Add("RL_SHARED=1"); //used instead of #define so cpp files can do conditional compilation
+
+            PublicDependencyModuleNames.AddRange(
+                new string[]
+                {
+                    "Core",
+                    "CoreUObject",
+                    "Engine",
+                    "ControlRig",
+					"UnrealEd",
+					"MainFrame",
+					"RigLogicModule",
+					"RigLogicLib",
+					"PropertyEditor",
+					"SlateCore",
+					"ApplicationCore",
+					"Slate",
+					"EditorStyle",
+					"InputCore"
+                }
+            );
+
+            PrivateIncludePathModuleNames.AddRange(
+				new string[]
+				{
+                    "PropertyEditor",
+                    "AssetTools"
+				}
+			);
+
+        }
+    }
+}
diff --git a/Engine/Plugins/Animation/RigLogic/Source/RigLogicEditor/RigLogicEditor_UPL.xml b/Engine/Plugins/Animation/RigLogic/Source/RigLogicEditor/RigLogicEditor_UPL.xml
new file mode 100644
--- /dev/null
+++ b/Engine/Plugins/Animation/RigLogic/Source/RigLogicEditor/RigLogicEditor_UPL.xml
@@ -0,0 +1,32 @@
+<?xml version="1.0" encoding="utf-8"?>
+<!--SpeechGraphics plugin additions-->
+<root xmlns:android="http://schemas.android.com/apk/res/android">
+	<!-- init section is always evaluated once per architecture -->
+	<init>
+		<log text="RigLogic init"/>
+		<setBool result="bSupported" value="false"/>
+		<isArch arch="arm64-v8a">
+			<setBool result="bSupported" value="true"/>
+		</isArch>
+		<if condition="bSupported">
+			<false>
+				<log text="Unsupported architecture: $S(Architecture)"/>
+			</false>
+		</if>
+	</init>
+
+	<!-- optional files or directories to copy to Intermediate/Android/APK -->
+	<resourceCopies>
+		<if condition="bSupported">
+			<true>					
+				<copyFile src="$S(AbsPluginDir)/ThirdParty/UFC/Lib/Android/arm64/libufc.so"
+							dst="$S(BuildDir)/libs/$S(Architecture)/libufc.so" />						
+				<copyFile src="$S(AbsPluginDir)/ThirdParty/UFC/Lib/Android/arm64/libc++_shared.so"
+							dst="$S(BuildDir)/libs/$S(Architecture)/libc++_shared.so" />
+				<copyFile src="$S(AbsPluginDir)/ThirdParty/UFC/Lib/Android/arm64/liblog.so"
+							dst="$S(BuildDir)/libs/$S(Architecture)/liblog.so" />																			
+				<log text="Copying libufc.so"/>
+				</true>
+		</if>
+	</resourceCopies>
+</root>
diff --git a/Engine/Plugins/Animation/RigLogic/Source/RigLogicLib/Private/RigLogicLib.cpp b/Engine/Plugins/Animation/RigLogic/Source/RigLogicLib/Private/RigLogicLib.cpp
new file mode 100644
--- /dev/null
+++ b/Engine/Plugins/Animation/RigLogic/Source/RigLogicLib/Private/RigLogicLib.cpp
@@ -0,0 +1,26 @@
+// Copyright 1998-2019 Epic Games, Inc. All Rights Reserved.
+
+#ifndef RIGLOGIC_MODULE_DISCARD
+
+#include "RigLogicLib.h"
+
+#include "Runtime/Core/Public/Modules/ModuleInterface.h"
+#include "Runtime/Core/Public/Modules/ModuleManager.h"
+
+#define LOCTEXT_NAMESPACE "FRigLogicLib"
+
+DEFINE_LOG_CATEGORY(LogRigLogicLib);
+
+void FRigLogicLib::StartupModule()
+{
+}
+
+void FRigLogicLib::ShutdownModule()
+{
+}
+
+#undef LOCTEXT_NAMESPACE
+
+IMPLEMENT_MODULE(FRigLogicLib, RigLogicLib)
+
+#endif  // RIGLOGIC_MODULE_DISCARD
diff --git a/Engine/Plugins/Animation/RigLogic/Source/RigLogicLib/Private/dna/DNA.h b/Engine/Plugins/Animation/RigLogic/Source/RigLogicLib/Private/dna/DNA.h
new file mode 100644
--- /dev/null
+++ b/Engine/Plugins/Animation/RigLogic/Source/RigLogicLib/Private/dna/DNA.h
@@ -0,0 +1,722 @@
+#pragma once
+
+#include "dna/LODMapping.h"
+#include "dna/SurjectiveMapping.h"
+#include "dna/TypeDefs.h"
+
+#include <terse/utils/ArchiveOffset.h>
+
+#ifdef _MSC_VER
+    #pragma warning(push)
+    #pragma warning(disable : 4365 4987)
+#endif
+#include <algorithm>
+#include <array>
+#include <cassert>
+#include <cstddef>
+#include <cstdint>
+#include <iterator>
+#include <memory>
+#include <tuple>
+#ifdef _MSC_VER
+    #pragma warning(pop)
+#endif
+
+namespace dna {
+
+template<typename T>
+using AlignedVector = Vector<T, PolyAllocator<T, 64ul, AlignedMemoryResource> >;
+
+template<typename TFrom, typename TTo = TFrom>
+struct RawSurjectiveMapping : public SurjectiveMapping<TFrom, TTo> {
+    using SurjectiveMapping<TFrom, TTo>::SurjectiveMapping;
+
+    template<class Archive>
+    void serialize(Archive& archive) {
+        archive(this->from, this->to);
+    }
+
+};
+
+template<typename T>
+struct ExpectedValue {
+    T expected;
+    T got;
+
+    explicit ExpectedValue(const T& value) : expected{value}, got{} {
+    }
+
+    template<class Archive>
+    void load(Archive& archive) {
+        archive(got);
+    }
+
+    template<class Archive>
+    void save(Archive& archive) {
+        archive(expected);
+    }
+
+    bool matches() const {
+        return (expected == got);
+    }
+
+};
+
+template<std::size_t Size>
+struct Signature {
+    using SignatureValueType = std::array<char, Size>;
+
+    ExpectedValue<SignatureValueType> value;
+
+    explicit Signature(SignatureValueType bytes) : value{bytes} {
+    }
+
+    template<class Archive>
+    void serialize(Archive& archive) {
+        archive(value);
+    }
+
+    bool matches() const {
+        return value.matches();
+    }
+
+};
+
+struct Version {
+    ExpectedValue<std::uint16_t> generation;
+    ExpectedValue<std::uint16_t> version;
+
+    Version(std::uint16_t generation_, std::uint16_t version_) :
+        generation{generation_},
+        version{version_} {
+    }
+
+    template<class Archive>
+    void serialize(Archive& archive) {
+        archive(generation, version);
+    }
+
+    bool matches() const {
+        return (generation.matches() && version.matches());
+    }
+
+};
+
+struct SectionLookupTable {
+    terse::ArchiveOffset<std::uint32_t> descriptor;
+    terse::ArchiveOffset<std::uint32_t> definition;
+    terse::ArchiveOffset<std::uint32_t> behavior;
+    terse::ArchiveOffset<std::uint32_t> controls;
+    terse::ArchiveOffset<std::uint32_t> joints;
+    terse::ArchiveOffset<std::uint32_t> blendShapeChannels;
+    terse::ArchiveOffset<std::uint32_t> animatedMaps;
+    terse::ArchiveOffset<std::uint32_t> geometry;
+
+    template<class Archive>
+    void serialize(Archive& archive) {
+        archive(descriptor,
+                definition,
+                behavior,
+                controls,
+                joints,
+                blendShapeChannels,
+                animatedMaps,
+                geometry);
+    }
+
+};
+
+struct RawCoordinateSystem {
+    std::uint16_t xAxis;
+    std::uint16_t yAxis;
+    std::uint16_t zAxis;
+
+    template<class Archive>
+    void serialize(Archive& archive) {
+        archive(xAxis, yAxis, zAxis);
+    }
+
+};
+
+struct RawLODMapping : public LODMapping {
+    using LODMapping::LODMapping;
+
+    template<class Archive>
+    void serialize(Archive& archive) {
+        archive(lods, indices);
+    }
+
+};
+
+struct RawDescriptor {
+    using StringPair = std::tuple<String<char>, String<char> >;
+
+    terse::ArchiveOffset<std::uint32_t>::Proxy marker;
+    String<char> name;
+    std::uint16_t archetype;
+    std::uint16_t gender;
+    std::uint16_t age;
+    Vector<StringPair> metadata;
+    std::uint16_t translationUnit;
+    std::uint16_t rotationUnit;
+    RawCoordinateSystem coordinateSystem;
+    std::uint16_t lodCount;
+    std::uint16_t maxLOD;
+    String<char> complexity;
+    String<char> dbName;
+
+    RawDescriptor(terse::ArchiveOffset<std::uint32_t>& markerTarget, MemoryResource* memRes) :
+        marker{markerTarget},
+        name{memRes},
+        archetype{},
+        gender{},
+        age{},
+        metadata{memRes},
+        translationUnit{},
+        rotationUnit{},
+        coordinateSystem{},
+        lodCount{},
+        maxLOD{},
+        complexity{memRes},
+        dbName{memRes} {
+    }
+
+    template<class Archive>
+    void serialize(Archive& archive) {
+        archive(marker,
+                name,
+                archetype,
+                gender,
+                age,
+                metadata,
+                translationUnit,
+                rotationUnit,
+                coordinateSystem,
+                lodCount,
+                maxLOD,
+                complexity,
+                dbName);
+    }
+
+};
+
+struct RawVector3Vector {
+    AlignedVector<float> xs;
+    AlignedVector<float> ys;
+    AlignedVector<float> zs;
+
+    explicit RawVector3Vector(MemoryResource* memRes) :
+        xs{memRes},
+        ys{memRes},
+        zs{memRes} {
+    }
+
+    RawVector3Vector(std::size_t size_, float initial, MemoryResource* memRes) :
+        xs{size_, initial, memRes},
+        ys{size_, initial, memRes},
+        zs{size_, initial, memRes} {
+    }
+
+    RawVector3Vector(ConstArrayView<float> xs_, ConstArrayView<float> ys_, ConstArrayView<float> zs_, MemoryResource* memRes) :
+        xs{xs_.begin(), xs_.end(), memRes},
+        ys{ys_.begin(), ys_.end(), memRes},
+        zs{zs_.begin(), zs_.end(), memRes} {
+    }
+
+    template<class Archive>
+    void serialize(Archive& archive) {
+        archive(xs, ys, zs);
+    }
+
+    std::size_t size() const {
+        assert(xs.size() == ys.size() && ys.size() == zs.size());
+        return xs.size();
+    }
+
+    void reserve(std::size_t count) {
+        xs.reserve(count);
+        ys.reserve(count);
+        zs.reserve(count);
+    }
+
+    void resize(std::size_t count) {
+        xs.resize(count);
+        ys.resize(count);
+        zs.resize(count);
+    }
+
+    void resize(std::size_t count, float value) {
+        xs.resize(count, value);
+        ys.resize(count, value);
+        zs.resize(count, value);
+    }
+
+    void clear() {
+        xs.clear();
+        ys.clear();
+        zs.clear();
+    }
+
+    template<typename Iterator>
+    void assign(Iterator start, Iterator end) {
+        clear();
+        reserve(static_cast<std::size_t>(std::distance(start, end)));
+        for (auto it = start; it != end; ++it) {
+            xs.push_back(it->x);
+            ys.push_back(it->y);
+            zs.push_back(it->z);
+        }
+    }
+
+};
+
+struct RawDefinition {
+    terse::ArchiveOffset<std::uint32_t>::Proxy marker;
+    RawLODMapping lodJointMapping;
+    RawLODMapping lodBlendShapeMapping;
+    RawLODMapping lodAnimatedMapMapping;
+    RawLODMapping lodMeshMapping;
+    Vector<String<char> > guiControlNames;
+    Vector<String<char> > rawControlNames;
+    Vector<String<char> > jointNames;
+    Vector<String<char> > blendShapeChannelNames;
+    Vector<String<char> > animatedMapNames;
+    Vector<String<char> > meshNames;
+    RawSurjectiveMapping<std::uint16_t> meshBlendShapeChannelMapping;
+    Vector<std::uint16_t> jointHierarchy;
+    RawVector3Vector neutralJointTranslations;
+    RawVector3Vector neutralJointRotations;
+
+    RawDefinition(terse::ArchiveOffset<std::uint32_t>& markerTarget, MemoryResource* memRes) :
+        marker{markerTarget},
+        lodJointMapping{memRes},
+        lodBlendShapeMapping{memRes},
+        lodAnimatedMapMapping{memRes},
+        lodMeshMapping{memRes},
+        guiControlNames{memRes},
+        rawControlNames{memRes},
+        jointNames{memRes},
+        blendShapeChannelNames{memRes},
+        animatedMapNames{memRes},
+        meshNames{memRes},
+        meshBlendShapeChannelMapping{memRes},
+        jointHierarchy{memRes},
+        neutralJointTranslations{memRes},
+        neutralJointRotations{memRes} {
+    }
+
+    template<class Archive>
+    void serialize(Archive& archive) {
+        archive(marker,
+                lodJointMapping,
+                lodBlendShapeMapping,
+                lodAnimatedMapMapping,
+                lodMeshMapping,
+                guiControlNames,
+                rawControlNames,
+                jointNames,
+                blendShapeChannelNames,
+                animatedMapNames,
+                meshNames,
+                meshBlendShapeChannelMapping,
+                jointHierarchy,
+                neutralJointTranslations,
+                neutralJointRotations);
+    }
+
+};
+
+struct RawConditionalTable {
+    Vector<std::uint16_t> inputIndices;
+    Vector<std::uint16_t> outputIndices;
+    Vector<float> fromValues;
+    Vector<float> toValues;
+    Vector<float> slopeValues;
+    Vector<float> cutValues;
+
+    explicit RawConditionalTable(MemoryResource* memRes) :
+        inputIndices{memRes},
+        outputIndices{memRes},
+        fromValues{memRes},
+        toValues{memRes},
+        slopeValues{memRes},
+        cutValues{memRes} {
+    }
+
+    template<class Archive>
+    void serialize(Archive& archive) {
+        archive(inputIndices,
+                outputIndices,
+                fromValues,
+                toValues,
+                slopeValues,
+                cutValues);
+    }
+
+};
+
+struct RawPSDMatrix {
+    Vector<std::uint16_t> rows;
+    Vector<std::uint16_t> columns;
+    Vector<float> values;
+
+    explicit RawPSDMatrix(MemoryResource* memRes) :
+        rows{memRes},
+        columns{memRes},
+        values{memRes} {
+    }
+
+    template<class Archive>
+    void serialize(Archive& archive) {
+        archive(rows, columns, values);
+    }
+
+};
+
+struct RawControls {
+    std::uint16_t psdCount;
+    RawConditionalTable conditionals;
+    RawPSDMatrix psds;
+
+    explicit RawControls(MemoryResource* memRes) :
+        psdCount{},
+        conditionals{memRes},
+        psds{memRes} {
+    }
+
+    template<class Archive>
+    void serialize(Archive& archive) {
+        archive(psdCount, conditionals, psds);
+    }
+
+};
+
+struct RawJointGroup {
+    // Row count of each LOD
+    // 12, 9, 3,
+    // |  |  + LOD-2 contains first 3 rows
+    // |  + LOD-1 contains first 9 rows
+    // + LOD-0 contains first 12 rows
+    Vector<std::uint16_t> lods;
+    // Sub-matrix col -> input vector
+    Vector<std::uint16_t> inputIndices;
+    // Sub-matrix row -> output vector
+    Vector<std::uint16_t> outputIndices;
+    // Non-zero values of all sub-matrices
+    AlignedVector<float> values;
+
+    Vector<std::uint16_t> jointIndices;
+
+    explicit RawJointGroup(MemoryResource* memRes) :
+        lods{memRes},
+        inputIndices{memRes},
+        outputIndices{memRes},
+        values{memRes},
+        jointIndices{memRes} {
+    }
+
+    template<class Archive>
+    void serialize(Archive& archive) {
+        archive(lods, inputIndices, outputIndices, values, jointIndices);
+    }
+
+};
+
+struct RawJoints {
+    std::uint16_t rowCount;
+    std::uint16_t colCount;
+    Vector<RawJointGroup> jointGroups;
+
+    explicit RawJoints(MemoryResource* memRes) :
+        rowCount{},
+        colCount{},
+        jointGroups{memRes} {
+    }
+
+    template<class Archive>
+    void serialize(Archive& archive) {
+        archive(rowCount, colCount, jointGroups);
+    }
+
+};
+
+struct RawBlendShapeChannels {
+    Vector<std::uint16_t> lods;
+    Vector<std::uint16_t> inputIndices;
+    Vector<std::uint16_t> outputIndices;
+
+    explicit RawBlendShapeChannels(MemoryResource* memRes) :
+        lods{memRes},
+        inputIndices{memRes},
+        outputIndices{memRes} {
+    }
+
+    template<class Archive>
+    void serialize(Archive& archive) {
+        archive(lods, inputIndices, outputIndices);
+    }
+
+};
+
+struct RawAnimatedMaps {
+    Vector<std::uint16_t> lods;
+    RawConditionalTable conditionals;
+
+    explicit RawAnimatedMaps(MemoryResource* memRes) :
+        lods{memRes},
+        conditionals{memRes} {
+    }
+
+    template<class Archive>
+    void serialize(Archive& archive) {
+        archive(lods, conditionals);
+    }
+
+};
+
+struct RawBehavior {
+    terse::ArchiveOffset<std::uint32_t>::Proxy marker;
+
+    terse::ArchiveOffset<std::uint32_t>::Proxy controlsMarker;
+    RawControls controls;
+
+    terse::ArchiveOffset<std::uint32_t>::Proxy jointsMarker;
+    RawJoints joints;
+
+    terse::ArchiveOffset<std::uint32_t>::Proxy blendShapeChannelsMarker;
+    RawBlendShapeChannels blendShapeChannels;
+
+    terse::ArchiveOffset<std::uint32_t>::Proxy animatedMapsMarker;
+    RawAnimatedMaps animatedMaps;
+
+    RawBehavior(terse::ArchiveOffset<std::uint32_t>& markerTarget,
+                terse::ArchiveOffset<std::uint32_t>& controlsMarkerTarget,
+                terse::ArchiveOffset<std::uint32_t>& jointsMarkerTarget,
+                terse::ArchiveOffset<std::uint32_t>& blendShapeChannelsMarkerTarget,
+                terse::ArchiveOffset<std::uint32_t>& animatedMapsMarkerTarget,
+                MemoryResource* memRes) :
+        marker{markerTarget},
+        controlsMarker{controlsMarkerTarget},
+        controls{memRes},
+        jointsMarker{jointsMarkerTarget},
+        joints{memRes},
+        blendShapeChannelsMarker{blendShapeChannelsMarkerTarget},
+        blendShapeChannels{memRes},
+        animatedMapsMarker{animatedMapsMarkerTarget},
+        animatedMaps{memRes} {
+    }
+
+    template<class Archive>
+    void serialize(Archive& archive) {
+        archive(marker,
+                controlsMarker,
+                controls,
+                jointsMarker,
+                joints,
+                blendShapeChannelsMarker,
+                blendShapeChannels,
+                animatedMapsMarker,
+                animatedMaps);
+    }
+
+};
+
+struct RawTextureCoordinateVector {
+    Vector<float> us;
+    Vector<float> vs;
+
+    explicit RawTextureCoordinateVector(MemoryResource* memRes) :
+        us{memRes},
+        vs{memRes} {
+    }
+
+    template<class Archive>
+    void serialize(Archive& archive) {
+        archive(us, vs);
+    }
+
+    std::size_t size() const {
+        assert(us.size() == vs.size());
+        return us.size();
+    }
+
+    void clear() {
+        us.clear();
+        vs.clear();
+    }
+
+};
+
+struct RawVertexLayoutVector {
+    Vector<std::uint32_t> positions;
+    Vector<std::uint32_t> textureCoordinates;
+    Vector<std::uint32_t> normals;
+
+    explicit RawVertexLayoutVector(MemoryResource* memRes) :
+        positions{memRes},
+        textureCoordinates{memRes},
+        normals{memRes} {
+    }
+
+    template<class Archive>
+    void serialize(Archive& archive) {
+        archive(positions, textureCoordinates, normals);
+    }
+
+    std::size_t size() const {
+        assert(positions.size() == textureCoordinates.size() && textureCoordinates.size() == normals.size());
+        return positions.size();
+    }
+
+    void clear() {
+        positions.clear();
+        textureCoordinates.clear();
+        normals.clear();
+    }
+
+};
+
+struct RawFace {
+    Vector<std::uint32_t> layoutIndices;
+
+    explicit RawFace(MemoryResource* memRes) :
+        layoutIndices{memRes} {
+    }
+
+    template<class Archive>
+    void serialize(Archive& archive) {
+        archive(layoutIndices);
+    }
+
+};
+
+struct RawVertexSkinWeights {
+    AlignedVector<float> weights;
+    Vector<std::uint16_t> jointIndices;
+
+    explicit RawVertexSkinWeights(MemoryResource* memRes) :
+        weights{memRes},
+        jointIndices{memRes} {
+    }
+
+    template<class Archive>
+    void serialize(Archive& archive) {
+        archive(weights, jointIndices);
+    }
+
+};
+
+struct RawBlendShapeTarget {
+    RawVector3Vector deltas;
+    Vector<std::uint32_t> vertexIndices;
+    std::uint16_t blendShapeChannelIndex;
+
+    explicit RawBlendShapeTarget(MemoryResource* memRes) :
+        deltas{memRes},
+        vertexIndices{memRes},
+        blendShapeChannelIndex{} {
+    }
+
+    template<class Archive>
+    void serialize(Archive& archive) {
+        archive(deltas, vertexIndices, blendShapeChannelIndex);
+    }
+
+};
+
+struct RawMesh {
+    terse::ArchiveOffset<std::uint32_t> offset;
+    RawVector3Vector positions;
+    RawTextureCoordinateVector textureCoordinates;
+    RawVector3Vector normals;
+    RawVertexLayoutVector layouts;
+    Vector<RawFace> faces;
+    std::uint16_t maximumInfluencePerVertex;
+    Vector<RawVertexSkinWeights> skinWeights;
+    Vector<RawBlendShapeTarget> blendShapeTargets;
+    terse::ArchiveOffset<std::uint32_t>::Proxy marker;
+
+    explicit RawMesh(MemoryResource* memRes) :
+        offset{},
+        positions{memRes},
+        textureCoordinates{memRes},
+        normals{memRes},
+        layouts{memRes},
+        faces{memRes},
+        maximumInfluencePerVertex{},
+        skinWeights{memRes},
+        blendShapeTargets{memRes},
+        marker{offset} {
+    }
+
+    template<class Archive>
+    void serialize(Archive& archive) {
+        archive(offset,
+                positions,
+                textureCoordinates,
+                normals,
+                layouts,
+                faces,
+                maximumInfluencePerVertex,
+                skinWeights,
+                blendShapeTargets,
+                marker);
+    }
+
+};
+
+struct RawGeometry {
+    terse::ArchiveOffset<std::uint32_t>::Proxy marker;
+    Vector<RawMesh> meshes;
+
+    RawGeometry(terse::ArchiveOffset<std::uint32_t>& markerTarget, MemoryResource* memRes) :
+        marker{markerTarget},
+        meshes{memRes} {
+    }
+
+    template<class Archive>
+    void serialize(Archive& archive) {
+        archive(marker, meshes);
+    }
+
+};
+
+struct DNA {
+    Signature<3> signature{{'D', 'N', 'A'}};
+    Version version{2, 1};
+    SectionLookupTable sections;
+    RawDescriptor descriptor;
+    RawDefinition definition;
+    RawBehavior behavior;
+    RawGeometry geometry;
+    Signature<3> eof{{'A', 'N', 'D'}};
+
+    explicit DNA(MemoryResource* memRes) :
+        sections{},
+        descriptor{sections.descriptor, memRes},
+        definition{sections.definition, memRes},
+        behavior{sections.behavior,
+                 sections.controls,
+                 sections.joints,
+                 sections.blendShapeChannels,
+                 sections.animatedMaps,
+                 memRes},
+        geometry{sections.geometry, memRes} {
+    }
+
+    template<class Archive>
+    void load(Archive& archive) {
+        archive(signature, version);
+        if (signature.matches() && version.matches()) {
+            archive(sections, descriptor, definition, behavior, geometry, eof);
+        }
+    }
+
+    template<class Archive>
+    void save(Archive& archive) {
+        archive(signature, version, sections, descriptor, definition, behavior, geometry, eof);
+    }
+
+};
+
+}  // namespace dna
diff --git a/Engine/Plugins/Animation/RigLogic/Source/RigLogicLib/Private/dna/DenormalizedData.h b/Engine/Plugins/Animation/RigLogic/Source/RigLogicLib/Private/dna/DenormalizedData.h
new file mode 100644
--- /dev/null
+++ b/Engine/Plugins/Animation/RigLogic/Source/RigLogicLib/Private/dna/DenormalizedData.h
@@ -0,0 +1,77 @@
+#pragma once
+
+#include "dna/LODMapping.h"
+#include "dna/types/Aliases.h"
+#include "dna/utils/Extd.h"
+
+#include <cassert>
+#include <cstdint>
+
+namespace dna {
+
+template<class Reader>
+struct DenormalizedData {
+    LODMapping jointVariableAttributeIndices;
+    LODMapping meshBlendShapeMappingIndices;
+
+    explicit DenormalizedData(MemoryResource* memRes) :
+        jointVariableAttributeIndices{memRes},
+        meshBlendShapeMappingIndices{memRes} {
+    }
+
+    void populate(const Reader* source) {
+        populateJointVariableAttributeIndices(source, jointVariableAttributeIndices);
+        populateMeshBlendShapeMappingIndices(source, meshBlendShapeMappingIndices);
+    }
+
+private:
+    void populateJointVariableAttributeIndices(const Reader* source, LODMapping& destination) {
+        // Prepare storage for all available LODs
+        const auto lodCount = source->getLODCount();
+        destination.setLODCount(lodCount);
+        // Concatenate all output indices for each LOD
+        for (std::uint16_t i = 0u; i < source->getJointGroupCount(); ++i) {
+            const auto outputIndices = source->getJointGroupOutputIndices(i);
+            const auto lodSizes = source->getJointGroupLODs(i);
+            assert(lodSizes.size() == lodCount);
+            for (std::uint16_t lod = 0u; lod < lodCount; ++lod) {
+                // In this case, each LOD has a distinct set of indices, so the LOD and Index parameters
+                // are the same for all LODs
+                destination.addIndices(lod, outputIndices.data(), lodSizes[lod]);
+                destination.associateLODWithIndices(lod, lod);
+            }
+        }
+    }
+
+    void populateMeshBlendShapeMappingIndices(const Reader* source, LODMapping& destination) {
+        // Prepare storage for all available LODs
+        const auto lodCount = source->getLODCount();
+        destination.setLODCount(lodCount);
+        // Include only those mapping indices which are present in the already filtered
+        // mesh and blendshape LOD mapping
+        for (std::uint16_t lod = 0u; lod < lodCount; ++lod) {
+            const auto meshIndices = source->getMeshIndicesForLOD(lod);
+            const auto blendShapeIndices = source->getBlendShapeChannelIndicesForLOD(lod);
+
+            auto isMappingNeeded = [&meshIndices, &blendShapeIndices](const MeshBlendShapeChannelMapping& mapping) {
+                const bool meshNeeded = extd::contains(meshIndices, mapping.meshIndex);
+                const bool blendShapeNeeded = extd::contains(blendShapeIndices, mapping.blendShapeChannelIndex);
+                return (meshNeeded && blendShapeNeeded);
+            };
+
+            for (std::uint16_t i = 0u; i < source->getMeshBlendShapeChannelMappingCount(); ++i) {
+                const auto mapping = source->getMeshBlendShapeChannelMapping(i);
+                if (isMappingNeeded(mapping)) {
+                    // In this case, each LOD has a distinct set of indices, so the LOD and Index parameters
+                    // are the same for all LODs
+                    destination.addIndices(lod, &i, static_cast<std::uint16_t>(1));
+                }
+            }
+
+            destination.associateLODWithIndices(lod, lod);
+        }
+    }
+
+};
+
+}  // namespace dna
diff --git a/Engine/Plugins/Animation/RigLogic/Source/RigLogicLib/Private/dna/LODConstraint.cpp b/Engine/Plugins/Animation/RigLogic/Source/RigLogicLib/Private/dna/LODConstraint.cpp
new file mode 100644
--- /dev/null
+++ b/Engine/Plugins/Animation/RigLogic/Source/RigLogicLib/Private/dna/LODConstraint.cpp
@@ -0,0 +1,51 @@
+#include "dna/LODConstraint.h"
+
+#include "dna/utils/Extd.h"
+
+#ifdef _MSC_VER
+    #pragma warning(push)
+    #pragma warning(disable : 4365 4987)
+#endif
+#include <cstddef>
+#include <cstdint>
+#include <numeric>
+#ifdef _MSC_VER
+    #pragma warning(pop)
+#endif
+
+namespace dna {
+
+LODConstraint::LODConstraint(std::uint16_t maxLOD, std::uint16_t minLOD, MemoryResource* memRes) : lods{memRes} {
+    assert(maxLOD <= minLOD);
+    lods.resize(static_cast<std::size_t>(minLOD - maxLOD + 1));
+    std::iota(lods.begin(), lods.end(), maxLOD);
+}
+
+LODConstraint::LODConstraint(ConstArrayView<std::uint16_t> lods_, MemoryResource* memRes) : lods{lods_.begin(), lods_.end(),
+                                                                                                 memRes} {
+    std::sort(lods.begin(), lods.end());
+}
+
+bool LODConstraint::hasImpactOn(std::uint16_t lodCount) const {
+    return (getMaxLOD() > 0u) || (getMinLOD() < static_cast<std::uint16_t>(lodCount - 1u));
+}
+
+std::uint16_t LODConstraint::getMaxLOD() const {
+    return (lods.empty() ? std::uint16_t{} : lods.front());
+}
+
+std::uint16_t LODConstraint::getMinLOD() const {
+    return (lods.empty() ? std::uint16_t{} : lods.back());
+}
+
+void LODConstraint::clampTo(std::uint16_t lodCount) {
+    extd::filter(lods, [lodCount](std::uint16_t lod, std::size_t  /*unused*/) {
+            return lod < lodCount;
+        });
+}
+
+void LODConstraint::applyTo(Vector<std::uint16_t>& unconstrainedLODs) const {
+    extd::filter(unconstrainedLODs, extd::byPosition(lods));
+}
+
+}  // namespace dna
diff --git a/Engine/Plugins/Animation/RigLogic/Source/RigLogicLib/Private/dna/ReaderImpl.hpp b/Engine/Plugins/Animation/RigLogic/Source/RigLogicLib/Private/dna/ReaderImpl.hpp
new file mode 100644
--- /dev/null
+++ b/Engine/Plugins/Animation/RigLogic/Source/RigLogicLib/Private/dna/ReaderImpl.hpp
@@ -0,0 +1,789 @@
+#include "dna/TypeDefs.h"
+
+#include <cstddef>
+#include <limits>
+#include <tuple>
+
+namespace dna {
+
+#ifdef _MSC_VER
+    #pragma warning(push)
+    #pragma warning(disable : 4589)
+#endif
+template<class TReaderBase>
+ReaderImpl<TReaderBase>::ReaderImpl(MemoryResource* memRes_) : BaseImpl{memRes_}, cache{memRes_} {
+}
+
+#ifdef _MSC_VER
+    #pragma warning(pop)
+#endif
+
+#ifdef _MSC_VER
+    #pragma warning(push)
+    #pragma warning(disable : 4505)
+#endif
+template<class TReaderBase>
+StringView ReaderImpl<TReaderBase>::getName() const {
+    return {dna.descriptor.name.data(), dna.descriptor.name.size()};
+}
+
+template<class TReaderBase>
+Archetype ReaderImpl<TReaderBase>::getArchetype() const {
+    return static_cast<Archetype>(dna.descriptor.archetype);
+}
+
+template<class TReaderBase>
+Gender ReaderImpl<TReaderBase>::getGender() const {
+    return static_cast<Gender>(dna.descriptor.gender);
+}
+
+template<class TReaderBase>
+std::uint16_t ReaderImpl<TReaderBase>::getAge() const {
+    return dna.descriptor.age;
+}
+
+template<class TReaderBase>
+std::uint32_t ReaderImpl<TReaderBase>::getMetaDataCount() const {
+    return static_cast<std::uint32_t>(dna.descriptor.metadata.size());
+}
+
+template<class TReaderBase>
+StringView ReaderImpl<TReaderBase>::getMetaDataKey(std::uint32_t index) const {
+    if (index < dna.descriptor.metadata.size()) {
+        const auto& key = std::get<0>(dna.descriptor.metadata[index]);
+        return {key.data(), key.size()};
+    }
+    return {};
+}
+
+template<class TReaderBase>
+StringView ReaderImpl<TReaderBase>::getMetaDataValue(const char* key) const {
+    for (const auto& data: dna.descriptor.metadata) {
+        if (std::get<0>(data) == key) {
+            const auto& value = std::get<1>(data);
+            return {value.data(), value.size()};
+        }
+    }
+    return {};
+}
+
+template<class TReaderBase>
+TranslationUnit ReaderImpl<TReaderBase>::getTranslationUnit() const {
+    return static_cast<TranslationUnit>(dna.descriptor.translationUnit);
+}
+
+template<class TReaderBase>
+RotationUnit ReaderImpl<TReaderBase>::getRotationUnit() const {
+    return static_cast<RotationUnit>(dna.descriptor.rotationUnit);
+}
+
+template<class TReaderBase>
+CoordinateSystem ReaderImpl<TReaderBase>::getCoordinateSystem() const {
+    return {
+        static_cast<Direction>(dna.descriptor.coordinateSystem.xAxis),
+        static_cast<Direction>(dna.descriptor.coordinateSystem.yAxis),
+        static_cast<Direction>(dna.descriptor.coordinateSystem.zAxis)
+    };
+}
+
+template<class TReaderBase>
+std::uint16_t ReaderImpl<TReaderBase>::getLODCount() const {
+    return dna.descriptor.lodCount;
+}
+
+template<class TReaderBase>
+std::uint16_t ReaderImpl<TReaderBase>::getDBMaxLOD() const {
+    return dna.descriptor.maxLOD;
+}
+
+template<class TReaderBase>
+StringView ReaderImpl<TReaderBase>::getDBComplexity() const {
+    return {dna.descriptor.complexity.data(), dna.descriptor.complexity.size()};
+}
+
+template<class TReaderBase>
+StringView ReaderImpl<TReaderBase>::getDBName() const {
+    return {dna.descriptor.dbName.data(), dna.descriptor.dbName.size()};
+}
+
+template<class TReaderBase>
+std::uint16_t ReaderImpl<TReaderBase>::getGUIControlCount() const {
+    return static_cast<std::uint16_t>(dna.definition.guiControlNames.size());
+}
+
+template<class TReaderBase>
+StringView ReaderImpl<TReaderBase>::getGUIControlName(std::uint16_t index) const {
+    if (index < dna.definition.guiControlNames.size()) {
+        const auto& guiControlName = dna.definition.guiControlNames[index];
+        return {guiControlName.data(), guiControlName.size()};
+    }
+    return {};
+}
+
+template<class TReaderBase>
+std::uint16_t ReaderImpl<TReaderBase>::getRawControlCount() const {
+    return static_cast<std::uint16_t>(dna.definition.rawControlNames.size());
+}
+
+template<class TReaderBase>
+StringView ReaderImpl<TReaderBase>::getRawControlName(std::uint16_t index) const {
+    if (index < dna.definition.rawControlNames.size()) {
+        const auto& rawControlName = dna.definition.rawControlNames[index];
+        return {rawControlName.data(), rawControlName.size()};
+    }
+    return {};
+}
+
+template<class TReaderBase>
+std::uint16_t ReaderImpl<TReaderBase>::getJointCount() const {
+    return static_cast<std::uint16_t>(dna.definition.jointNames.size());
+}
+
+template<class TReaderBase>
+StringView ReaderImpl<TReaderBase>::getJointName(std::uint16_t index) const {
+    if (index < dna.definition.jointNames.size()) {
+        const auto& jointName = dna.definition.jointNames[index];
+        return {jointName.data(), jointName.size()};
+    }
+    return {};
+}
+
+template<class TReaderBase>
+ConstArrayView<std::uint16_t> ReaderImpl<TReaderBase>::getJointIndicesForLOD(std::uint16_t lod) const {
+    return dna.definition.lodJointMapping.getIndices(lod);
+}
+
+template<class TReaderBase>
+std::uint16_t ReaderImpl<TReaderBase>::getJointParentIndex(std::uint16_t index) const {
+    if (index < dna.definition.jointHierarchy.size()) {
+        return dna.definition.jointHierarchy[index];
+    }
+    return std::numeric_limits<std::uint16_t>::max();
+}
+
+template<class TReaderBase>
+std::uint16_t ReaderImpl<TReaderBase>::getBlendShapeChannelCount() const {
+    return static_cast<std::uint16_t>(dna.definition.blendShapeChannelNames.size());
+}
+
+template<class TReaderBase>
+StringView ReaderImpl<TReaderBase>::getBlendShapeChannelName(std::uint16_t index) const {
+    if (index < dna.definition.blendShapeChannelNames.size()) {
+        const auto& blendShapeName = dna.definition.blendShapeChannelNames[index];
+        return {blendShapeName.data(), blendShapeName.size()};
+    }
+    return {};
+}
+
+template<class TReaderBase>
+ConstArrayView<std::uint16_t> ReaderImpl<TReaderBase>::getBlendShapeChannelIndicesForLOD(std::uint16_t lod) const {
+    return dna.definition.lodBlendShapeMapping.getIndices(lod);
+}
+
+template<class TReaderBase>
+std::uint16_t ReaderImpl<TReaderBase>::getAnimatedMapCount() const {
+    return static_cast<std::uint16_t>(dna.definition.animatedMapNames.size());
+}
+
+template<class TReaderBase>
+StringView ReaderImpl<TReaderBase>::getAnimatedMapName(std::uint16_t index) const {
+    if (index < dna.definition.animatedMapNames.size()) {
+        const auto& animatedMapName = dna.definition.animatedMapNames[index];
+        return {animatedMapName.data(), animatedMapName.size()};
+    }
+    return {};
+}
+
+template<class TReaderBase>
+ConstArrayView<std::uint16_t> ReaderImpl<TReaderBase>::getAnimatedMapIndicesForLOD(std::uint16_t lod) const {
+    return dna.definition.lodAnimatedMapMapping.getIndices(lod);
+}
+
+template<class TReaderBase>
+std::uint16_t ReaderImpl<TReaderBase>::getMeshCount() const {
+    return static_cast<std::uint16_t>(dna.definition.meshNames.size());
+}
+
+template<class TReaderBase>
+StringView ReaderImpl<TReaderBase>::getMeshName(std::uint16_t index) const {
+    if (index < dna.definition.meshNames.size()) {
+        const auto& meshName = dna.definition.meshNames[index];
+        return {meshName.data(), meshName.size()};
+    }
+    return {};
+}
+
+template<class TReaderBase>
+ConstArrayView<std::uint16_t> ReaderImpl<TReaderBase>::getMeshIndicesForLOD(std::uint16_t lod) const {
+    return dna.definition.lodMeshMapping.getIndices(lod);
+}
+
+template<class TReaderBase>
+std::uint16_t ReaderImpl<TReaderBase>::getMeshBlendShapeChannelMappingCount() const {
+    return static_cast<std::uint16_t>(dna.definition.meshBlendShapeChannelMapping.size());
+}
+
+template<class TReaderBase>
+MeshBlendShapeChannelMapping ReaderImpl<TReaderBase>::getMeshBlendShapeChannelMapping(std::uint16_t index) const {
+    const auto mapping = dna.definition.meshBlendShapeChannelMapping.get(index);
+    return {mapping.from, mapping.to};
+}
+
+template<class TReaderBase>
+ConstArrayView<std::uint16_t> ReaderImpl<TReaderBase>::getMeshBlendShapeChannelMappingIndicesForLOD(std::uint16_t lod) const {
+    if (cache.meshBlendShapeMappingIndices.getLODCount() == static_cast<std::uint16_t>(0)) {
+        cache.populate(this);
+    }
+    return cache.meshBlendShapeMappingIndices.getIndices(lod);
+}
+
+template<class TReaderBase>
+Vector3 ReaderImpl<TReaderBase>::getNeutralJointTranslation(std::uint16_t index) const {
+    const auto& translations = dna.definition.neutralJointTranslations;
+    if (index < translations.xs.size()) {
+        return {translations.xs[index], translations.ys[index], translations.zs[index]};
+    }
+    return {};
+}
+
+template<class TReaderBase>
+ConstArrayView<float> ReaderImpl<TReaderBase>::getNeutralJointTranslationXs() const {
+    const auto& xs = dna.definition.neutralJointTranslations.xs;
+    return {xs.data(), xs.size()};
+}
+
+template<class TReaderBase>
+ConstArrayView<float> ReaderImpl<TReaderBase>::getNeutralJointTranslationYs() const {
+    const auto& ys = dna.definition.neutralJointTranslations.ys;
+    return {ys.data(), ys.size()};
+}
+
+template<class TReaderBase>
+ConstArrayView<float> ReaderImpl<TReaderBase>::getNeutralJointTranslationZs() const {
+    const auto& zs = dna.definition.neutralJointTranslations.zs;
+    return {zs.data(), zs.size()};
+}
+
+template<class TReaderBase>
+Vector3 ReaderImpl<TReaderBase>::getNeutralJointRotation(std::uint16_t index) const {
+    const auto& rotations = dna.definition.neutralJointRotations;
+    if (index < rotations.size()) {
+        return {rotations.xs[index], rotations.ys[index], rotations.zs[index]};
+    }
+    return {};
+}
+
+template<class TReaderBase>
+ConstArrayView<float> ReaderImpl<TReaderBase>::getNeutralJointRotationXs() const {
+    const auto& xs = dna.definition.neutralJointRotations.xs;
+    return {xs.data(), xs.size()};
+}
+
+template<class TReaderBase>
+ConstArrayView<float> ReaderImpl<TReaderBase>::getNeutralJointRotationYs() const {
+    const auto& ys = dna.definition.neutralJointRotations.ys;
+    return {ys.data(), ys.size()};
+}
+
+template<class TReaderBase>
+ConstArrayView<float> ReaderImpl<TReaderBase>::getNeutralJointRotationZs() const {
+    const auto& zs = dna.definition.neutralJointRotations.zs;
+    return {zs.data(), zs.size()};
+}
+
+template<class TReaderBase>
+ConstArrayView<std::uint16_t> ReaderImpl<TReaderBase>::getGUIToRawInputIndices() const {
+    const auto& inputIndices = dna.behavior.controls.conditionals.inputIndices;
+    return {inputIndices.data(), inputIndices.size()};
+}
+
+template<class TReaderBase>
+ConstArrayView<std::uint16_t> ReaderImpl<TReaderBase>::getGUIToRawOutputIndices() const {
+    const auto& outputIndices = dna.behavior.controls.conditionals.outputIndices;
+    return {outputIndices.data(), outputIndices.size()};
+}
+
+template<class TReaderBase>
+ConstArrayView<float> ReaderImpl<TReaderBase>::getGUIToRawFromValues() const {
+    const auto& fromValues = dna.behavior.controls.conditionals.fromValues;
+    return {fromValues.data(), fromValues.size()};
+}
+
+template<class TReaderBase>
+ConstArrayView<float> ReaderImpl<TReaderBase>::getGUIToRawToValues() const {
+    const auto& toValues = dna.behavior.controls.conditionals.toValues;
+    return {toValues.data(), toValues.size()};
+}
+
+template<class TReaderBase>
+ConstArrayView<float> ReaderImpl<TReaderBase>::getGUIToRawSlopeValues() const {
+    const auto& slopeValues = dna.behavior.controls.conditionals.slopeValues;
+    return {slopeValues.data(), slopeValues.size()};
+}
+
+template<class TReaderBase>
+ConstArrayView<float> ReaderImpl<TReaderBase>::getGUIToRawCutValues() const {
+    const auto& cutValues = dna.behavior.controls.conditionals.cutValues;
+    return {cutValues.data(), cutValues.size()};
+}
+
+template<class TReaderBase>
+std::uint16_t ReaderImpl<TReaderBase>::getPSDCount() const {
+    return dna.behavior.controls.psdCount;
+}
+
+template<class TReaderBase>
+ConstArrayView<std::uint16_t> ReaderImpl<TReaderBase>::getPSDRowIndices() const {
+    const auto& rows = dna.behavior.controls.psds.rows;
+    return {rows.data(), rows.size()};
+}
+
+template<class TReaderBase>
+ConstArrayView<std::uint16_t> ReaderImpl<TReaderBase>::getPSDColumnIndices() const {
+    const auto& columns = dna.behavior.controls.psds.columns;
+    return {columns.data(), columns.size()};
+}
+
+template<class TReaderBase>
+ConstArrayView<float> ReaderImpl<TReaderBase>::getPSDValues() const {
+    const auto& values = dna.behavior.controls.psds.values;
+    return {values.data(), values.size()};
+}
+
+template<class TReaderBase>
+std::uint16_t ReaderImpl<TReaderBase>::getJointRowCount() const {
+    return dna.behavior.joints.rowCount;
+}
+
+template<class TReaderBase>
+std::uint16_t ReaderImpl<TReaderBase>::getJointColumnCount() const {
+    return dna.behavior.joints.colCount;
+}
+
+template<class TReaderBase>
+ConstArrayView<std::uint16_t> ReaderImpl<TReaderBase>::getJointVariableAttributeIndices(std::uint16_t lod) const {
+    if (cache.jointVariableAttributeIndices.getLODCount() == static_cast<std::uint16_t>(0)) {
+        cache.populate(this);
+    }
+    return cache.jointVariableAttributeIndices.getIndices(lod);
+}
+
+template<class TReaderBase>
+std::uint16_t ReaderImpl<TReaderBase>::getJointGroupCount() const {
+    return static_cast<std::uint16_t>(dna.behavior.joints.jointGroups.size());
+}
+
+template<class TReaderBase>
+ConstArrayView<std::uint16_t> ReaderImpl<TReaderBase>::getJointGroupLODs(std::uint16_t jointGroupIndex) const {
+    if (jointGroupIndex < dna.behavior.joints.jointGroups.size()) {
+        const auto& lods = dna.behavior.joints.jointGroups[jointGroupIndex].lods;
+        return {lods.data(), lods.size()};
+    }
+    return {};
+}
+
+template<class TReaderBase>
+ConstArrayView<std::uint16_t> ReaderImpl<TReaderBase>::getJointGroupInputIndices(std::uint16_t jointGroupIndex) const {
+    if (jointGroupIndex < dna.behavior.joints.jointGroups.size()) {
+        const auto& inputIndices = dna.behavior.joints.jointGroups[jointGroupIndex].inputIndices;
+        return {inputIndices.data(), inputIndices.size()};
+    }
+    return {};
+}
+
+template<class TReaderBase>
+ConstArrayView<std::uint16_t> ReaderImpl<TReaderBase>::getJointGroupOutputIndices(std::uint16_t jointGroupIndex) const {
+    if (jointGroupIndex < dna.behavior.joints.jointGroups.size()) {
+        const auto& outputIndices = dna.behavior.joints.jointGroups[jointGroupIndex].outputIndices;
+        return {outputIndices.data(), outputIndices.size()};
+    }
+    return {};
+}
+
+template<class TReaderBase>
+ConstArrayView<float> ReaderImpl<TReaderBase>::getJointGroupValues(std::uint16_t jointGroupIndex) const {
+    if (jointGroupIndex < dna.behavior.joints.jointGroups.size()) {
+        const auto& values = dna.behavior.joints.jointGroups[jointGroupIndex].values;
+        return {values.data(), values.size()};
+    }
+    return {};
+}
+
+template<class TReaderBase>
+ConstArrayView<std::uint16_t> ReaderImpl<TReaderBase>::getJointGroupJointIndices(std::uint16_t jointGroupIndex) const {
+    if (jointGroupIndex < dna.behavior.joints.jointGroups.size()) {
+        const auto& jointIndices = dna.behavior.joints.jointGroups[jointGroupIndex].jointIndices;
+        return {jointIndices.data(), jointIndices.size()};
+    }
+    return {};
+}
+
+template<class TReaderBase>
+ConstArrayView<std::uint16_t> ReaderImpl<TReaderBase>::getBlendShapeChannelLODs() const {
+    const auto& lods = dna.behavior.blendShapeChannels.lods;
+    return {lods.data(), lods.size()};
+}
+
+template<class TReaderBase>
+ConstArrayView<std::uint16_t> ReaderImpl<TReaderBase>::getBlendShapeChannelInputIndices() const {
+    const auto& inputIndices = dna.behavior.blendShapeChannels.inputIndices;
+    return {inputIndices.data(), inputIndices.size()};
+}
+
+template<class TReaderBase>
+ConstArrayView<std::uint16_t> ReaderImpl<TReaderBase>::getBlendShapeChannelOutputIndices() const {
+    const auto& outputIndices = dna.behavior.blendShapeChannels.outputIndices;
+    return {outputIndices.data(), outputIndices.size()};
+}
+
+template<class TReaderBase>
+ConstArrayView<std::uint16_t> ReaderImpl<TReaderBase>::getAnimatedMapLODs() const {
+    const auto& lods = dna.behavior.animatedMaps.lods;
+    return {lods.data(), lods.size()};
+}
+
+template<class TReaderBase>
+ConstArrayView<std::uint16_t> ReaderImpl<TReaderBase>::getAnimatedMapInputIndices() const {
+    const auto& inputIndices = dna.behavior.animatedMaps.conditionals.inputIndices;
+    return {inputIndices.data(), inputIndices.size()};
+}
+
+template<class TReaderBase>
+ConstArrayView<std::uint16_t> ReaderImpl<TReaderBase>::getAnimatedMapOutputIndices() const {
+    const auto& outputIndices = dna.behavior.animatedMaps.conditionals.outputIndices;
+    return {outputIndices.data(), outputIndices.size()};
+}
+
+template<class TReaderBase>
+ConstArrayView<float> ReaderImpl<TReaderBase>::getAnimatedMapFromValues() const {
+    const auto& fromValues = dna.behavior.animatedMaps.conditionals.fromValues;
+    return {fromValues.data(), fromValues.size()};
+}
+
+template<class TReaderBase>
+ConstArrayView<float> ReaderImpl<TReaderBase>::getAnimatedMapToValues() const {
+    const auto& toValues = dna.behavior.animatedMaps.conditionals.toValues;
+    return {toValues.data(), toValues.size()};
+}
+
+template<class TReaderBase>
+ConstArrayView<float> ReaderImpl<TReaderBase>::getAnimatedMapSlopeValues() const {
+    const auto& slopeValues = dna.behavior.animatedMaps.conditionals.slopeValues;
+    return {slopeValues.data(), slopeValues.size()};
+}
+
+template<class TReaderBase>
+ConstArrayView<float> ReaderImpl<TReaderBase>::getAnimatedMapCutValues() const {
+    const auto& cutValues = dna.behavior.animatedMaps.conditionals.cutValues;
+    return {cutValues.data(), cutValues.size()};
+}
+
+template<class TReaderBase>
+std::uint32_t ReaderImpl<TReaderBase>::getVertexPositionCount(std::uint16_t meshIndex) const {
+    if (meshIndex < dna.geometry.meshes.size()) {
+        return static_cast<std::uint32_t>(dna.geometry.meshes[meshIndex].positions.xs.size());
+    }
+    return 0u;
+}
+
+template<class TReaderBase>
+Position ReaderImpl<TReaderBase>::getVertexPosition(std::uint16_t meshIndex, std::uint32_t vertexIndex) const {
+    if (meshIndex < dna.geometry.meshes.size()) {
+        const auto& positions = dna.geometry.meshes[meshIndex].positions;
+        if (vertexIndex < positions.size()) {
+            return {positions.xs[vertexIndex], positions.ys[vertexIndex], positions.zs[vertexIndex]};
+        }
+    }
+    return {};
+}
+
+template<class TReaderBase>
+ConstArrayView<float> ReaderImpl<TReaderBase>::getVertexPositionXs(std::uint16_t meshIndex) const {
+    if (meshIndex < dna.geometry.meshes.size()) {
+        const auto& xPositions = dna.geometry.meshes[meshIndex].positions.xs;
+        return {xPositions.data(), xPositions.size()};
+    }
+    return {};
+}
+
+template<class TReaderBase>
+ConstArrayView<float> ReaderImpl<TReaderBase>::getVertexPositionYs(std::uint16_t meshIndex) const {
+    if (meshIndex < dna.geometry.meshes.size()) {
+        const auto& yPositions = dna.geometry.meshes[meshIndex].positions.ys;
+        return {yPositions.data(), yPositions.size()};
+    }
+    return {};
+}
+
+template<class TReaderBase>
+ConstArrayView<float> ReaderImpl<TReaderBase>::getVertexPositionZs(std::uint16_t meshIndex) const {
+    if (meshIndex < dna.geometry.meshes.size()) {
+        const auto& zPositions = dna.geometry.meshes[meshIndex].positions.zs;
+        return {zPositions.data(), zPositions.size()};
+    }
+    return {};
+}
+
+template<class TReaderBase>
+std::uint32_t ReaderImpl<TReaderBase>::getVertexTextureCoordinateCount(std::uint16_t meshIndex) const {
+    if (meshIndex < dna.geometry.meshes.size()) {
+        return static_cast<std::uint32_t>(dna.geometry.meshes[meshIndex].textureCoordinates.us.size());
+    }
+    return 0u;
+}
+
+template<class TReaderBase>
+TextureCoordinate ReaderImpl<TReaderBase>::getVertexTextureCoordinate(std::uint16_t meshIndex,
+                                                                      std::uint32_t textureCoordinateIndex) const {
+    if (meshIndex < dna.geometry.meshes.size()) {
+        const auto& textureCoordinates = dna.geometry.meshes[meshIndex].textureCoordinates;
+        if (textureCoordinateIndex < textureCoordinates.size()) {
+            return {textureCoordinates.us[textureCoordinateIndex], textureCoordinates.vs[textureCoordinateIndex]};
+        }
+    }
+    return {};
+}
+
+template<class TReaderBase>
+ConstArrayView<float> ReaderImpl<TReaderBase>::getVertexTextureCoordinateUs(std::uint16_t meshIndex) const {
+    const auto& uTextureCoordinates = dna.geometry.meshes[meshIndex].textureCoordinates.us;
+    return {uTextureCoordinates.data(), uTextureCoordinates.size()};
+}
+
+template<class TReaderBase>
+ConstArrayView<float> ReaderImpl<TReaderBase>::getVertexTextureCoordinateVs(std::uint16_t meshIndex) const {
+    const auto& vTextureCoordinates = dna.geometry.meshes[meshIndex].textureCoordinates.vs;
+    return {vTextureCoordinates.data(), vTextureCoordinates.size()};
+}
+
+template<class TReaderBase>
+std::uint32_t ReaderImpl<TReaderBase>::getVertexNormalCount(std::uint16_t meshIndex) const {
+    if (meshIndex < dna.geometry.meshes.size()) {
+        return static_cast<std::uint32_t>(dna.geometry.meshes[meshIndex].normals.xs.size());
+    }
+    return 0u;
+}
+
+template<class TReaderBase>
+Normal ReaderImpl<TReaderBase>::getVertexNormal(std::uint16_t meshIndex, std::uint32_t normalIndex) const {
+    if (meshIndex < dna.geometry.meshes.size()) {
+        const auto& normals = dna.geometry.meshes[meshIndex].normals;
+        if (normalIndex < normals.size()) {
+            return {normals.xs[normalIndex], normals.ys[normalIndex], normals.zs[normalIndex]};
+        }
+    }
+    return {};
+}
+
+template<class TReaderBase>
+ConstArrayView<float> ReaderImpl<TReaderBase>::getVertexNormalXs(std::uint16_t meshIndex) const {
+    if (meshIndex < dna.geometry.meshes.size()) {
+        const auto& xNormals = dna.geometry.meshes[meshIndex].normals.xs;
+        return {xNormals.data(), xNormals.size()};
+    }
+    return {};
+}
+
+template<class TReaderBase>
+ConstArrayView<float> ReaderImpl<TReaderBase>::getVertexNormalYs(std::uint16_t meshIndex) const {
+    if (meshIndex < dna.geometry.meshes.size()) {
+        const auto& yNormals = dna.geometry.meshes[meshIndex].normals.ys;
+        return {yNormals.data(), yNormals.size()};
+    }
+    return {};
+}
+
+template<class TReaderBase>
+ConstArrayView<float> ReaderImpl<TReaderBase>::getVertexNormalZs(std::uint16_t meshIndex) const {
+    if (meshIndex < dna.geometry.meshes.size()) {
+        const auto& zNormals = dna.geometry.meshes[meshIndex].normals.zs;
+        return {zNormals.data(), zNormals.size()};
+    }
+    return {};
+}
+
+template<class TReaderBase>
+std::uint32_t ReaderImpl<TReaderBase>::getFaceCount(std::uint16_t meshIndex) const {
+    if (meshIndex < dna.geometry.meshes.size()) {
+        return static_cast<std::uint32_t>(dna.geometry.meshes[meshIndex].faces.size());
+    }
+    return 0u;
+}
+
+template<class TReaderBase>
+ConstArrayView<std::uint32_t> ReaderImpl<TReaderBase>::getFaceVertexLayoutIndices(std::uint16_t meshIndex,
+                                                                                  std::uint32_t faceIndex) const {
+    const auto& meshes = dna.geometry.meshes;
+    if ((meshIndex < meshes.size()) && (faceIndex < meshes[meshIndex].faces.size())) {
+        const auto& layoutIndices = meshes[meshIndex].faces[faceIndex].layoutIndices;
+        return {layoutIndices.data(), layoutIndices.size()};
+    }
+    return {};
+}
+
+template<class TReaderBase>
+std::uint32_t ReaderImpl<TReaderBase>::getVertexLayoutCount(std::uint16_t meshIndex) const {
+    if (meshIndex < dna.geometry.meshes.size()) {
+        return static_cast<std::uint32_t>(dna.geometry.meshes[meshIndex].layouts.positions.size());
+    }
+    return 0u;
+}
+
+template<class TReaderBase>
+VertexLayout ReaderImpl<TReaderBase>::getVertexLayout(std::uint16_t meshIndex, std::uint32_t layoutIndex) const {
+    if (meshIndex < dna.geometry.meshes.size()) {
+        const auto& layouts = dna.geometry.meshes[meshIndex].layouts;
+        if (layoutIndex < layouts.size()) {
+            return {layouts.positions[layoutIndex], layouts.textureCoordinates[layoutIndex], layouts.normals[layoutIndex]};
+        }
+    }
+    return {};
+}
+
+template<class TReaderBase>
+ConstArrayView<std::uint32_t> ReaderImpl<TReaderBase>::getVertexLayoutPositionIndices(std::uint16_t meshIndex) const {
+    if (meshIndex < dna.geometry.meshes.size()) {
+        const auto& positions = dna.geometry.meshes[meshIndex].layouts.positions;
+        return {positions.data(), positions.size()};
+    }
+    return {};
+}
+
+template<class TReaderBase>
+ConstArrayView<std::uint32_t> ReaderImpl<TReaderBase>::getVertexLayoutTextureCoordinateIndices(std::uint16_t meshIndex) const {
+    if (meshIndex < dna.geometry.meshes.size()) {
+        const auto& textureCoordinated = dna.geometry.meshes[meshIndex].layouts.textureCoordinates;
+        return {textureCoordinated.data(), textureCoordinated.size()};
+    }
+    return {};
+}
+
+template<class TReaderBase>
+ConstArrayView<std::uint32_t> ReaderImpl<TReaderBase>::getVertexLayoutNormalIndices(std::uint16_t meshIndex) const {
+    if (meshIndex < dna.geometry.meshes.size()) {
+        const auto& normals = dna.geometry.meshes[meshIndex].layouts.normals;
+        return {normals.data(), normals.size()};
+    }
+    return {};
+}
+
+template<class TReaderBase>
+std::uint16_t ReaderImpl<TReaderBase>::getMaximumInfluencePerVertex(std::uint16_t meshIndex) const {
+    if (meshIndex < dna.geometry.meshes.size()) {
+        return dna.geometry.meshes[meshIndex].maximumInfluencePerVertex;
+    }
+    return {};
+}
+
+template<class TReaderBase>
+ConstArrayView<float> ReaderImpl<TReaderBase>::getSkinWeightsValues(std::uint16_t meshIndex, std::uint32_t vertexIndex) const {
+    const auto& meshes = dna.geometry.meshes;
+    if ((meshIndex < meshes.size()) && (vertexIndex < meshes[meshIndex].skinWeights.size())) {
+        const auto& weights = meshes[meshIndex].skinWeights[vertexIndex].weights;
+        return {weights.data(), weights.size()};
+    }
+    return {};
+}
+
+template<class TReaderBase>
+ConstArrayView<std::uint16_t> ReaderImpl<TReaderBase>::getSkinWeightsJointIndices(std::uint16_t meshIndex,
+                                                                                  std::uint32_t vertexIndex) const {
+    const auto& meshes = dna.geometry.meshes;
+    if ((meshIndex < meshes.size()) && (vertexIndex < meshes[meshIndex].skinWeights.size())) {
+        const auto& jointIndices = meshes[meshIndex].skinWeights[vertexIndex].jointIndices;
+        return {jointIndices.data(), jointIndices.size()};
+    }
+    return {};
+}
+
+template<class TReaderBase>
+std::uint16_t ReaderImpl<TReaderBase>::getBlendShapeTargetCount(std::uint16_t meshIndex) const {
+    if (meshIndex < dna.geometry.meshes.size()) {
+        return static_cast<std::uint16_t>(dna.geometry.meshes[meshIndex].blendShapeTargets.size());
+    }
+    return {};
+}
+
+template<class TReaderBase>
+std::uint16_t ReaderImpl<TReaderBase>::getBlendShapeChannelIndex(std::uint16_t meshIndex,
+                                                                 std::uint16_t blendShapeTargetIndex) const {
+    const auto& meshes = dna.geometry.meshes;
+    if ((meshIndex < meshes.size()) && (blendShapeTargetIndex < meshes[meshIndex].blendShapeTargets.size())) {
+        return meshes[meshIndex].blendShapeTargets[blendShapeTargetIndex].blendShapeChannelIndex;
+    }
+    return {};
+}
+
+template<class TReaderBase>
+std::uint32_t ReaderImpl<TReaderBase>::getBlendShapeTargetDeltaCount(std::uint16_t meshIndex,
+                                                                     std::uint16_t blendShapeTargetIndex) const {
+    const auto& meshes = dna.geometry.meshes;
+    if ((meshIndex < meshes.size()) && (blendShapeTargetIndex < meshes[meshIndex].blendShapeTargets.size())) {
+        return static_cast<std::uint32_t>(meshes[meshIndex].blendShapeTargets[blendShapeTargetIndex].deltas.xs.size());
+    }
+    return {};
+}
+
+template<class TReaderBase>
+Delta ReaderImpl<TReaderBase>::getBlendShapeTargetDelta(std::uint16_t meshIndex,
+                                                        std::uint16_t blendShapeTargetIndex,
+                                                        std::uint32_t deltaIndex) const {
+    const auto& meshes = dna.geometry.meshes;
+    if ((meshIndex < meshes.size()) && (blendShapeTargetIndex < meshes[meshIndex].blendShapeTargets.size()) &&
+        (deltaIndex < meshes[meshIndex].blendShapeTargets[blendShapeTargetIndex].deltas.size())) {
+        const auto& deltas = meshes[meshIndex].blendShapeTargets[blendShapeTargetIndex].deltas;
+        return {deltas.xs[deltaIndex], deltas.ys[deltaIndex], deltas.zs[deltaIndex]};
+    }
+    return {};
+}
+
+template<class TReaderBase>
+ConstArrayView<float> ReaderImpl<TReaderBase>::getBlendShapeTargetDeltaXs(std::uint16_t meshIndex,
+                                                                          std::uint16_t blendShapeTargetIndex) const {
+    const auto& meshes = dna.geometry.meshes;
+    if ((meshIndex < meshes.size()) && (blendShapeTargetIndex < meshes[meshIndex].blendShapeTargets.size())) {
+        const auto& xDeltas = meshes[meshIndex].blendShapeTargets[blendShapeTargetIndex].deltas.xs;
+        return {xDeltas.data(), xDeltas.size()};
+    }
+    return {};
+}
+
+template<class TReaderBase>
+ConstArrayView<float> ReaderImpl<TReaderBase>::getBlendShapeTargetDeltaYs(std::uint16_t meshIndex,
+                                                                          std::uint16_t blendShapeTargetIndex) const {
+    const auto& meshes = dna.geometry.meshes;
+    if ((meshIndex < meshes.size()) && (blendShapeTargetIndex < meshes[meshIndex].blendShapeTargets.size())) {
+        const auto& yDeltas = meshes[meshIndex].blendShapeTargets[blendShapeTargetIndex].deltas.ys;
+        return {yDeltas.data(), yDeltas.size()};
+    }
+    return {};
+}
+
+template<class TReaderBase>
+ConstArrayView<float> ReaderImpl<TReaderBase>::getBlendShapeTargetDeltaZs(std::uint16_t meshIndex,
+                                                                          std::uint16_t blendShapeTargetIndex) const {
+    const auto& meshes = dna.geometry.meshes;
+    if ((meshIndex < meshes.size()) && (blendShapeTargetIndex < meshes[meshIndex].blendShapeTargets.size())) {
+        const auto& zDeltas = meshes[meshIndex].blendShapeTargets[blendShapeTargetIndex].deltas.zs;
+        return {zDeltas.data(), zDeltas.size()};
+    }
+    return {};
+}
+
+template<class TReaderBase>
+ConstArrayView<std::uint32_t> ReaderImpl<TReaderBase>::getBlendShapeTargetVertexIndices(std::uint16_t meshIndex,
+                                                                                        std::uint16_t blendShapeTargetIndex) const
+{
+    const auto& meshes = dna.geometry.meshes;
+    if ((meshIndex < meshes.size()) && (blendShapeTargetIndex < meshes[meshIndex].blendShapeTargets.size())) {
+        const auto& vertexIndices = meshes[meshIndex].blendShapeTargets[blendShapeTargetIndex].vertexIndices;
+        return {vertexIndices.data(), vertexIndices.size()};
+    }
+    return {};
+}
+
+#ifdef _MSC_VER
+    #pragma warning(pop)
+#endif
+
+}  // namespace dna
