Editgrating RigLogic plugin from Dev-CharacterTech to UE5/Main @ CL 13891876.

diff --git a/Engine/Plugins/Animation/RigLogic/Source/RigLogicLib/Private/dna/WriterImpl.hpp b/Engine/Plugins/Animation/RigLogic/Source/RigLogicLib/Private/dna/WriterImpl.hpp
new file mode 100644
--- /dev/null
+++ b/Engine/Plugins/Animation/RigLogic/Source/RigLogicLib/Private/dna/WriterImpl.hpp
@@ -0,0 +1,476 @@
+#include "dna/TypeDefs.h"
+
+#include <cassert>
+#include <cstddef>
+#include <cstring>
+#include <tuple>
+#include <utility>
+
+namespace dna {
+
+template<class TContainer, typename ... Args>
+void ensureHasSize(TContainer& target, std::size_t size, Args&& ... args) {
+    target.reserve(size);
+    while (target.size() < size) {
+        target.push_back(typename TContainer::value_type(std::forward<Args>(args)...));
+    }
+}
+
+#ifdef _MSC_VER
+    #pragma warning(push)
+    #pragma warning(disable : 4589)
+#endif
+template<class TWriterBase>
+WriterImpl<TWriterBase>::WriterImpl(MemoryResource* memRes_) : BaseImpl{memRes_} {
+}
+
+#ifdef _MSC_VER
+    #pragma warning(pop)
+#endif
+
+#ifdef _MSC_VER
+    #pragma warning(push)
+    #pragma warning(disable : 4505)
+#endif
+template<class TWriterBase>
+void WriterImpl<TWriterBase>::setName(const char* name) {
+    dna.descriptor.name = name;
+}
+
+template<class TWriterBase>
+void WriterImpl<TWriterBase>::setArchetype(Archetype archetype) {
+    dna.descriptor.archetype = static_cast<std::uint16_t>(archetype);
+}
+
+template<class TWriterBase>
+void WriterImpl<TWriterBase>::setGender(Gender gender) {
+    dna.descriptor.gender = static_cast<std::uint16_t>(gender);
+}
+
+template<class TWriterBase>
+void WriterImpl<TWriterBase>::setAge(std::uint16_t age) {
+    dna.descriptor.age = age;
+}
+
+template<class TWriterBase>
+void WriterImpl<TWriterBase>::setMetaData(const char* key, const char* value) {
+    using CharStringPair = std::tuple<String<char>, String<char> >;
+    auto it = std::find_if(dna.descriptor.metadata.begin(), dna.descriptor.metadata.end(), [&key](const CharStringPair& kv) {
+            auto& k = std::get<0>(kv);
+            return (std::strlen(key) == k.size() && std::strncmp(k.data(), key, k.size()) == 0);
+        });
+    if (it == dna.descriptor.metadata.end()) {
+        dna.descriptor.metadata.push_back({String<char>{key, memRes}, String<char>{value, memRes}});
+    } else {
+        std::get<1>(*it) = value;
+    }
+}
+
+template<class TWriterBase>
+void WriterImpl<TWriterBase>::setTranslationUnit(TranslationUnit unit) {
+    dna.descriptor.translationUnit = static_cast<std::uint16_t>(unit);
+}
+
+template<class TWriterBase>
+void WriterImpl<TWriterBase>::setRotationUnit(RotationUnit unit) {
+    dna.descriptor.rotationUnit = static_cast<std::uint16_t>(unit);
+}
+
+template<class TWriterBase>
+void WriterImpl<TWriterBase>::setCoordinateSystem(CoordinateSystem system) {
+    dna.descriptor.coordinateSystem.xAxis = static_cast<std::uint16_t>(system.xAxis);
+    dna.descriptor.coordinateSystem.yAxis = static_cast<std::uint16_t>(system.yAxis);
+    dna.descriptor.coordinateSystem.zAxis = static_cast<std::uint16_t>(system.zAxis);
+}
+
+template<class TWriterBase>
+void WriterImpl<TWriterBase>::setLODCount(std::uint16_t lodCount) {
+    dna.descriptor.lodCount = lodCount;
+}
+
+template<class TWriterBase>
+void WriterImpl<TWriterBase>::setDBMaxLOD(std::uint16_t lod) {
+    dna.descriptor.maxLOD = lod;
+}
+
+template<class TWriterBase>
+void WriterImpl<TWriterBase>::setDBComplexity(const char* name) {
+    dna.descriptor.complexity = name;
+}
+
+template<class TWriterBase>
+void WriterImpl<TWriterBase>::setDBName(const char* name) {
+    dna.descriptor.dbName = name;
+}
+
+template<class TWriterBase>
+void WriterImpl<TWriterBase>::setGUIControlName(std::uint16_t index, const char* name) {
+    ensureHasSize(dna.definition.guiControlNames, index + 1ul, memRes);
+    dna.definition.guiControlNames[index] = name;
+}
+
+template<class TWriterBase>
+void WriterImpl<TWriterBase>::setRawControlName(std::uint16_t index, const char* name) {
+    ensureHasSize(dna.definition.rawControlNames, index + 1ul, memRes);
+    dna.definition.rawControlNames[index] = name;
+}
+
+template<class TWriterBase>
+void WriterImpl<TWriterBase>::setJointName(std::uint16_t index, const char* name) {
+    ensureHasSize(dna.definition.jointNames, index + 1ul, memRes);
+    dna.definition.jointNames[index] = name;
+}
+
+template<class TWriterBase>
+void WriterImpl<TWriterBase>::setJointIndices(std::uint16_t index, const std::uint16_t* jointIndices, std::uint16_t count) {
+    dna.definition.lodJointMapping.clearIndices(index);
+    dna.definition.lodJointMapping.addIndices(index, jointIndices, count);
+}
+
+template<class TWriterBase>
+void WriterImpl<TWriterBase>::setLODJointMapping(std::uint16_t lod, std::uint16_t index) {
+    dna.definition.lodJointMapping.associateLODWithIndices(lod, index);
+}
+
+template<class TWriterBase>
+void WriterImpl<TWriterBase>::setJointHierarchy(const std::uint16_t* jointIndices, std::uint16_t count) {
+    dna.definition.jointHierarchy.assign(jointIndices, jointIndices + count);
+}
+
+template<class TWriterBase>
+void WriterImpl<TWriterBase>::setBlendShapeChannelName(std::uint16_t index, const char* name) {
+    ensureHasSize(dna.definition.blendShapeChannelNames, index + 1ul, memRes);
+    dna.definition.blendShapeChannelNames[index] = name;
+}
+
+template<class TWriterBase>
+void WriterImpl<TWriterBase>::setBlendShapeChannelIndices(std::uint16_t index,
+                                                          const std::uint16_t* blendShapeChannelIndices,
+                                                          std::uint16_t count) {
+    dna.definition.lodBlendShapeMapping.clearIndices(index);
+    dna.definition.lodBlendShapeMapping.addIndices(index, blendShapeChannelIndices, count);
+}
+
+template<class TWriterBase>
+void WriterImpl<TWriterBase>::setLODBlendShapeChannelMapping(std::uint16_t lod, std::uint16_t index) {
+    dna.definition.lodBlendShapeMapping.associateLODWithIndices(lod, index);
+}
+
+template<class TWriterBase>
+void WriterImpl<TWriterBase>::setAnimatedMapName(std::uint16_t index, const char* name) {
+    ensureHasSize(dna.definition.animatedMapNames, index + 1ul, memRes);
+    dna.definition.animatedMapNames[index] = name;
+}
+
+template<class TWriterBase>
+void WriterImpl<TWriterBase>::setAnimatedMapIndices(std::uint16_t index,
+                                                    const std::uint16_t* animatedMapIndices,
+                                                    std::uint16_t count) {
+    dna.definition.lodAnimatedMapMapping.clearIndices(index);
+    dna.definition.lodAnimatedMapMapping.addIndices(index, animatedMapIndices, count);
+}
+
+template<class TWriterBase>
+void WriterImpl<TWriterBase>::setLODAnimatedMapMapping(std::uint16_t lod, std::uint16_t index) {
+    dna.definition.lodAnimatedMapMapping.associateLODWithIndices(lod, index);
+}
+
+template<class TWriterBase>
+void WriterImpl<TWriterBase>::setMeshName(std::uint16_t index, const char* name) {
+    ensureHasSize(dna.definition.meshNames, index + 1ul, memRes);
+    dna.definition.meshNames[index] = name;
+}
+
+template<class TWriterBase>
+void WriterImpl<TWriterBase>::setMeshIndices(std::uint16_t index, const std::uint16_t* meshIndices, std::uint16_t count) {
+    dna.definition.lodMeshMapping.clearIndices(index);
+    dna.definition.lodMeshMapping.addIndices(index, meshIndices, count);
+}
+
+template<class TWriterBase>
+void WriterImpl<TWriterBase>::setLODMeshMapping(std::uint16_t lod, std::uint16_t index) {
+    dna.definition.lodMeshMapping.associateLODWithIndices(lod, index);
+}
+
+template<class TWriterBase>
+void WriterImpl<TWriterBase>::addMeshBlendShapeChannelMapping(std::uint16_t meshIndex, std::uint16_t blendShapeChannelIndex) {
+    dna.definition.meshBlendShapeChannelMapping.add(meshIndex, blendShapeChannelIndex);
+}
+
+template<class TWriterBase>
+void WriterImpl<TWriterBase>::setNeutralJointTranslations(const Vector3* translations, std::uint16_t count) {
+    dna.definition.neutralJointTranslations.assign(translations, translations + count);
+}
+
+template<class TWriterBase>
+void WriterImpl<TWriterBase>::setNeutralJointRotations(const Vector3* rotations, std::uint16_t count) {
+    dna.definition.neutralJointRotations.assign(rotations, rotations + count);
+}
+
+template<class TWriterBase>
+void WriterImpl<TWriterBase>::setGUIToRawInputIndices(const std::uint16_t* inputIndices, std::uint16_t count) {
+    dna.behavior.controls.conditionals.inputIndices.assign(inputIndices, inputIndices + count);
+}
+
+template<class TWriterBase>
+void WriterImpl<TWriterBase>::setGUIToRawOutputIndices(const std::uint16_t* outputIndices, std::uint16_t count) {
+    dna.behavior.controls.conditionals.outputIndices.assign(outputIndices, outputIndices + count);
+}
+
+template<class TWriterBase>
+void WriterImpl<TWriterBase>::setGUIToRawFromValues(const float* fromValues, std::uint16_t count) {
+    dna.behavior.controls.conditionals.fromValues.assign(fromValues, fromValues + count);
+}
+
+template<class TWriterBase>
+void WriterImpl<TWriterBase>::setGUIToRawToValues(const float* toValues, std::uint16_t count) {
+    dna.behavior.controls.conditionals.toValues.assign(toValues, toValues + count);
+}
+
+template<class TWriterBase>
+void WriterImpl<TWriterBase>::setGUIToRawSlopeValues(const float* slopeValues, std::uint16_t count) {
+    dna.behavior.controls.conditionals.slopeValues.assign(slopeValues, slopeValues + count);
+}
+
+template<class TWriterBase>
+void WriterImpl<TWriterBase>::setGUIToRawCutValues(const float* cutValues, std::uint16_t count) {
+    dna.behavior.controls.conditionals.cutValues.assign(cutValues, cutValues + count);
+}
+
+template<class TWriterBase>
+void WriterImpl<TWriterBase>::setPSDCount(std::uint16_t count) {
+    dna.behavior.controls.psdCount = count;
+}
+
+template<class TWriterBase>
+void WriterImpl<TWriterBase>::setPSDRowIndices(const std::uint16_t* rowIndices, std::uint16_t count) {
+    dna.behavior.controls.psds.rows.assign(rowIndices, rowIndices + count);
+}
+
+template<class TWriterBase>
+void WriterImpl<TWriterBase>::setPSDColumnIndices(const std::uint16_t* columnIndices, std::uint16_t count) {
+    dna.behavior.controls.psds.columns.assign(columnIndices, columnIndices + count);
+}
+
+template<class TWriterBase>
+void WriterImpl<TWriterBase>::setPSDValues(const float* weights, std::uint16_t count) {
+    dna.behavior.controls.psds.values.assign(weights, weights + count);
+}
+
+template<class TWriterBase>
+void WriterImpl<TWriterBase>::setJointRowCount(std::uint16_t rowCount) {
+    dna.behavior.joints.rowCount = rowCount;
+}
+
+template<class TWriterBase>
+void WriterImpl<TWriterBase>::setJointColumnCount(std::uint16_t columnCount) {
+    dna.behavior.joints.colCount = columnCount;
+}
+
+template<class TWriterBase>
+void WriterImpl<TWriterBase>::setJointGroupLODs(std::uint16_t jointGroupIndex, const std::uint16_t* lods, std::uint16_t count) {
+    auto& jointGroups = dna.behavior.joints.jointGroups;
+    ensureHasSize(jointGroups, jointGroupIndex + 1ul, memRes);
+    jointGroups[jointGroupIndex].lods.assign(lods, lods + count);
+}
+
+template<class TWriterBase>
+void WriterImpl<TWriterBase>::setJointGroupInputIndices(std::uint16_t jointGroupIndex,
+                                                        const std::uint16_t* inputIndices,
+                                                        std::uint16_t count) {
+    auto& jointGroups = dna.behavior.joints.jointGroups;
+    ensureHasSize(jointGroups, jointGroupIndex + 1ul, memRes);
+    jointGroups[jointGroupIndex].inputIndices.assign(inputIndices, inputIndices + count);
+}
+
+template<class TWriterBase>
+void WriterImpl<TWriterBase>::setJointGroupOutputIndices(std::uint16_t jointGroupIndex,
+                                                         const std::uint16_t* outputIndices,
+                                                         std::uint16_t count) {
+    auto& jointGroups = dna.behavior.joints.jointGroups;
+    ensureHasSize(jointGroups, jointGroupIndex + 1ul, memRes);
+    jointGroups[jointGroupIndex].outputIndices.assign(outputIndices, outputIndices + count);
+}
+
+template<class TWriterBase>
+void WriterImpl<TWriterBase>::setJointGroupValues(std::uint16_t jointGroupIndex, const float* values, std::uint32_t count) {
+    auto& jointGroups = dna.behavior.joints.jointGroups;
+    ensureHasSize(jointGroups, jointGroupIndex + 1ul, memRes);
+    jointGroups[jointGroupIndex].values.assign(values, values + count);
+}
+
+template<class TWriterBase>
+void WriterImpl<TWriterBase>::setJointGroupJointIndices(std::uint16_t jointGroupIndex,
+                                                        const std::uint16_t* jointIndices,
+                                                        std::uint16_t count) {
+    auto& jointGroups = dna.behavior.joints.jointGroups;
+    ensureHasSize(jointGroups, jointGroupIndex + 1ul, memRes);
+    jointGroups[jointGroupIndex].jointIndices.assign(jointIndices, jointIndices + count);
+}
+
+template<class TWriterBase>
+void WriterImpl<TWriterBase>::setBlendShapeChannelLODs(const std::uint16_t* lods, std::uint16_t count) {
+    dna.behavior.blendShapeChannels.lods.assign(lods, lods + count);
+}
+
+template<class TWriterBase>
+void WriterImpl<TWriterBase>::setBlendShapeChannelInputIndices(const std::uint16_t* inputIndices, std::uint16_t count) {
+    dna.behavior.blendShapeChannels.inputIndices.assign(inputIndices, inputIndices + count);
+}
+
+template<class TWriterBase>
+void WriterImpl<TWriterBase>::setBlendShapeChannelOutputIndices(const std::uint16_t* outputIndices, std::uint16_t count) {
+    dna.behavior.blendShapeChannels.outputIndices.assign(outputIndices, outputIndices + count);
+}
+
+template<class TWriterBase>
+void WriterImpl<TWriterBase>::setAnimatedMapLODs(const std::uint16_t* lods, std::uint16_t count) {
+    dna.behavior.animatedMaps.lods.assign(lods, lods + count);
+}
+
+template<class TWriterBase>
+void WriterImpl<TWriterBase>::setAnimatedMapInputIndices(const std::uint16_t* inputIndices, std::uint16_t count) {
+    dna.behavior.animatedMaps.conditionals.inputIndices.assign(inputIndices, inputIndices + count);
+}
+
+template<class TWriterBase>
+void WriterImpl<TWriterBase>::setAnimatedMapOutputIndices(const std::uint16_t* outputIndices, std::uint16_t count) {
+    dna.behavior.animatedMaps.conditionals.outputIndices.assign(outputIndices, outputIndices + count);
+}
+
+template<class TWriterBase>
+void WriterImpl<TWriterBase>::setAnimatedMapFromValues(const float* fromValues, std::uint16_t count) {
+    dna.behavior.animatedMaps.conditionals.fromValues.assign(fromValues, fromValues + count);
+}
+
+template<class TWriterBase>
+void WriterImpl<TWriterBase>::setAnimatedMapToValues(const float* toValues, std::uint16_t count) {
+    dna.behavior.animatedMaps.conditionals.toValues.assign(toValues, toValues + count);
+}
+
+template<class TWriterBase>
+void WriterImpl<TWriterBase>::setAnimatedMapSlopeValues(const float* slopeValues, std::uint16_t count) {
+    dna.behavior.animatedMaps.conditionals.slopeValues.assign(slopeValues, slopeValues + count);
+}
+
+template<class TWriterBase>
+void WriterImpl<TWriterBase>::setAnimatedMapCutValues(const float* cutValues, std::uint16_t count) {
+    dna.behavior.animatedMaps.conditionals.cutValues.assign(cutValues, cutValues + count);
+}
+
+template<class TWriterBase>
+void WriterImpl<TWriterBase>::setVertexPositions(std::uint16_t meshIndex, const Position* positions, std::uint32_t count) {
+    ensureHasSize(dna.geometry.meshes, meshIndex + 1ul, memRes);
+    dna.geometry.meshes[meshIndex].positions.assign(positions, positions + count);
+}
+
+template<class TWriterBase>
+void WriterImpl<TWriterBase>::setVertexTextureCoordinates(std::uint16_t meshIndex,
+                                                          const TextureCoordinate* textureCoordinates,
+                                                          std::uint32_t count) {
+    ensureHasSize(dna.geometry.meshes, meshIndex + 1ul, memRes);
+    auto& destination = dna.geometry.meshes[meshIndex].textureCoordinates;
+    destination.clear();
+    destination.us.reserve(count);
+    destination.vs.reserve(count);
+    for (std::size_t i = 0ul; i < count; ++i) {
+        destination.us.push_back(textureCoordinates[i].u);
+        destination.vs.push_back(textureCoordinates[i].v);
+    }
+}
+
+template<class TWriterBase>
+void WriterImpl<TWriterBase>::setVertexNormals(std::uint16_t meshIndex, const Normal* normals, std::uint32_t count) {
+    ensureHasSize(dna.geometry.meshes, meshIndex + 1ul, memRes);
+    dna.geometry.meshes[meshIndex].normals.assign(normals, normals + count);
+}
+
+template<class TWriterBase>
+void WriterImpl<TWriterBase>::setVertexLayouts(std::uint16_t meshIndex, const VertexLayout* layouts, std::uint32_t count) {
+    ensureHasSize(dna.geometry.meshes, meshIndex + 1ul, memRes);
+    auto& destination = dna.geometry.meshes[meshIndex].layouts;
+    destination.clear();
+    destination.positions.reserve(count);
+    destination.textureCoordinates.reserve(count);
+    destination.normals.reserve(count);
+    for (std::size_t i = 0ul; i < count; ++i) {
+        destination.positions.push_back(layouts[i].position);
+        destination.textureCoordinates.push_back(layouts[i].textureCoordinate);
+        destination.normals.push_back(layouts[i].normal);
+    }
+}
+
+template<class TWriterBase>
+void WriterImpl<TWriterBase>::setFaceVertexLayoutIndices(std::uint16_t meshIndex,
+                                                         std::uint32_t faceIndex,
+                                                         const std::uint32_t* layoutIndices,
+                                                         std::uint32_t count) {
+    ensureHasSize(dna.geometry.meshes, meshIndex + 1ul, memRes);
+    auto& faces = dna.geometry.meshes[meshIndex].faces;
+    ensureHasSize(faces, faceIndex + 1ul, memRes);
+    faces[faceIndex].layoutIndices.assign(layoutIndices, layoutIndices + count);
+}
+
+template<class TWriterBase>
+void WriterImpl<TWriterBase>::setMaximumInfluencePerVertex(std::uint16_t meshIndex, std::uint16_t maxInfluenceCount) {
+    ensureHasSize(dna.geometry.meshes, meshIndex + 1ul, memRes);
+    dna.geometry.meshes[meshIndex].maximumInfluencePerVertex = maxInfluenceCount;
+}
+
+template<class TWriterBase>
+void WriterImpl<TWriterBase>::setSkinWeightsValues(std::uint16_t meshIndex,
+                                                   std::uint32_t vertexIndex,
+                                                   const float* weights,
+                                                   std::uint16_t count) {
+    ensureHasSize(dna.geometry.meshes, meshIndex + 1ul, memRes);
+    auto& skinWeights = dna.geometry.meshes[meshIndex].skinWeights;
+    ensureHasSize(skinWeights, vertexIndex + 1ul, memRes);
+    skinWeights[vertexIndex].weights.assign(weights, weights + count);
+}
+
+template<class TWriterBase>
+void WriterImpl<TWriterBase>::setSkinWeightsJointIndices(std::uint16_t meshIndex,
+                                                         std::uint32_t vertexIndex,
+                                                         const std::uint16_t* jointIndices,
+                                                         std::uint16_t count) {
+    ensureHasSize(dna.geometry.meshes, meshIndex + 1ul, memRes);
+    auto& skinWeights = dna.geometry.meshes[meshIndex].skinWeights;
+    ensureHasSize(skinWeights, vertexIndex + 1ul, memRes);
+    skinWeights[vertexIndex].jointIndices.assign(jointIndices, jointIndices + count);
+}
+
+template<class TWriterBase>
+void WriterImpl<TWriterBase>::setBlendShapeChannelIndex(std::uint16_t meshIndex,
+                                                        std::uint16_t blendShapeTargetIndex,
+                                                        std::uint16_t blendShapeChannelIndex) {
+    ensureHasSize(dna.geometry.meshes, meshIndex + 1ul, memRes);
+    ensureHasSize(dna.geometry.meshes[meshIndex].blendShapeTargets, blendShapeTargetIndex + 1ul, memRes);
+    dna.geometry.meshes[meshIndex].blendShapeTargets[blendShapeTargetIndex].blendShapeChannelIndex = blendShapeChannelIndex;
+}
+
+template<class TWriterBase>
+void WriterImpl<TWriterBase>::setBlendShapeTargetDeltas(std::uint16_t meshIndex,
+                                                        std::uint16_t blendShapeTargetIndex,
+                                                        const Delta* deltas,
+                                                        std::uint32_t count) {
+    ensureHasSize(dna.geometry.meshes, meshIndex + 1ul, memRes);
+    ensureHasSize(dna.geometry.meshes[meshIndex].blendShapeTargets, blendShapeTargetIndex + 1ul, memRes);
+    dna.geometry.meshes[meshIndex].blendShapeTargets[blendShapeTargetIndex].deltas.assign(deltas, deltas + count);
+}
+
+template<class TWriterBase>
+void WriterImpl<TWriterBase>::setBlendShapeTargetVertexIndices(std::uint16_t meshIndex,
+                                                               std::uint16_t blendShapeTargetIndex,
+                                                               const std::uint32_t* vertexIndices,
+                                                               std::uint32_t count) {
+    ensureHasSize(dna.geometry.meshes, meshIndex + 1ul, memRes);
+    auto& blendShapeTargets = dna.geometry.meshes[meshIndex].blendShapeTargets;
+    ensureHasSize(blendShapeTargets, blendShapeTargetIndex + 1ul, memRes);
+    blendShapeTargets[blendShapeTargetIndex].vertexIndices.assign(vertexIndices, vertexIndices + count);
+}
+
+#ifdef _MSC_VER
+    #pragma warning(pop)
+#endif
+
+}  // namespace dna
diff --git a/Engine/Plugins/Animation/RigLogic/Source/RigLogicLib/Private/dna/stream/FilteredInputArchive.cpp b/Engine/Plugins/Animation/RigLogic/Source/RigLogicLib/Private/dna/stream/FilteredInputArchive.cpp
new file mode 100644
--- /dev/null
+++ b/Engine/Plugins/Animation/RigLogic/Source/RigLogicLib/Private/dna/stream/FilteredInputArchive.cpp
@@ -0,0 +1,277 @@
+#include "dna/stream/FilteredInputArchive.h"
+
+#include "dna/DNA.h"
+#include "dna/TypeDefs.h"
+#include "dna/utils/Extd.h"
+#include "dna/utils/ScopedEnumEx.h"
+
+#ifdef _MSC_VER
+    #pragma warning(push)
+    #pragma warning(disable : 4365 4987)
+#endif
+#include <cassert>
+#include <cstddef>
+#include <cstdint>
+#include <functional>
+#ifdef _MSC_VER
+    #pragma warning(pop)
+#endif
+
+namespace dna {
+
+static constexpr std::uint16_t jointAttributeCount = 9u;
+
+template<typename T, typename U>
+static UnorderedMap<U, U> remappedPositions(const Vector<T>& target, const UnorderedSet<U>& indices, MemoryResource* memRes) {
+    UnorderedMap<U, U> mapping{memRes};
+    for (U oldIndex{}, newIndex{}; oldIndex < static_cast<U>(target.size()); ++oldIndex) {
+        if (indices.find(oldIndex) != indices.end()) {
+            mapping.insert({oldIndex, newIndex});
+            ++newIndex;
+        }
+    }
+    return mapping;
+}
+
+FilteredInputArchive::FilteredInputArchive(BoundedIOStream* stream_,
+                                           DataLayer layer_,
+                                           std::uint16_t maxLOD_,
+                                           std::uint16_t minLOD_,
+                                           MemoryResource* memRes_) :
+    AnimatedMapFilter{memRes_},
+    BlendShapeFilter{memRes_},
+    JointFilter{memRes_},
+    MeshFilter{memRes_},
+    BaseArchive{this, stream_},
+    stream{stream_},
+    memRes{memRes_},
+    layerBitmask{computeDataLayerBitmask(layer_)},
+    lodConstraint{maxLOD_, minLOD_, memRes},
+    unconstrainedLODCount{} {
+}
+
+FilteredInputArchive::FilteredInputArchive(BoundedIOStream* stream_,
+                                           DataLayer layer_,
+                                           ConstArrayView<std::uint16_t> lods_,
+                                           MemoryResource* memRes_) :
+    AnimatedMapFilter{memRes_},
+    BlendShapeFilter{memRes_},
+    JointFilter{memRes_},
+    MeshFilter{memRes_},
+    BaseArchive{this, stream_},
+    stream{stream_},
+    memRes{memRes_},
+    layerBitmask{computeDataLayerBitmask(layer_)},
+    lodConstraint{lods_, memRes},
+    unconstrainedLODCount{} {
+}
+
+template<typename TContainer>
+void FilteredInputArchive::processSubset(TContainer& dest, std::size_t offset, std::size_t size) {
+    using ElementType = typename TContainer::value_type;
+    const auto availableSize = processSize();
+    assert(offset + size <= availableSize);
+    const auto startPosition = stream->tell();
+    // Skip over first N elements
+    stream->seek(startPosition + offset * sizeof(ElementType));
+    // Read requested number of elements
+    processElements(dest, size);
+    // Even if not all elements were read, seek to the end of the list
+    stream->seek(startPosition + availableSize * sizeof(ElementType));
+}
+
+void FilteredInputArchive::process(RawDescriptor& dest) {
+    BaseArchive::process(dest);
+    assert(dest.lodCount > 0u);
+    lodConstraint.clampTo(dest.lodCount);
+    unconstrainedLODCount = dest.lodCount;
+    dest.maxLOD = static_cast<std::uint16_t>(dest.maxLOD + lodConstraint.getMaxLOD());
+    dest.lodCount = static_cast<std::uint16_t>(lodConstraint.getMinLOD() - lodConstraint.getMaxLOD() + 1);
+}
+
+void FilteredInputArchive::process(RawDefinition& dest) {
+    if (!contains(layerBitmask, DataLayerBitmask::Definition)) {
+        return;
+    }
+    // Load all data
+    BaseArchive::process(dest);
+    // No filtering is done, unless LOD constraint may have some effect
+    if (!lodConstraint.hasImpactOn(unconstrainedLODCount)) {
+        return;
+    }
+
+    // Discard LOD data that is not relevant for the selected MaxLOD and MinLOD constraints
+    dest.lodMeshMapping.discardLODs(lodConstraint);
+    dest.lodJointMapping.discardLODs(lodConstraint);
+    dest.lodBlendShapeMapping.discardLODs(lodConstraint);
+    dest.lodAnimatedMapMapping.discardLODs(lodConstraint);
+    MeshFilter::apply(dest);
+    JointFilter::apply(dest);
+    BlendShapeFilter::apply(dest);
+    AnimatedMapFilter::apply(dest);
+}
+
+void FilteredInputArchive::process(RawBehavior& dest) {
+    if (contains(layerBitmask, DataLayerBitmask::Behavior)) {
+        process(dest.marker);
+        process(dest.controlsMarker);
+        process(dest.controls);
+        process(dest.jointsMarker);
+        process(dest.joints);
+        process(dest.blendShapeChannelsMarker);
+        process(dest.blendShapeChannels);
+        process(dest.animatedMapsMarker);
+        process(dest.animatedMaps);
+    }
+}
+
+void FilteredInputArchive::process(RawJoints& dest) {
+    process(dest.rowCount);
+    process(dest.colCount);
+    if (!lodConstraint.hasImpactOn(unconstrainedLODCount)) {
+        process(dest.jointGroups);
+        return;
+    }
+    // Perform filtered load only if LOD constraints have been set
+    const auto jointGroupCount = processSize();
+    dest.jointGroups.reserve(jointGroupCount);
+    for (std::size_t i = 0ul; i < jointGroupCount; ++i) {
+        RawJointGroup jointGroup{memRes};
+        process(jointGroup.lods);
+        // Discard everything that falls outside the region bounded by LOD constraints
+        lodConstraint.applyTo(jointGroup.lods);
+        // Input indices are all loaded always (unless the whole joint group is empty)
+        const auto jointGroupRowCount = (jointGroup.lods.empty() ? static_cast<std::uint16_t>(0) : jointGroup.lods.front());
+        if (jointGroupRowCount != 0u) {
+            process(jointGroup.inputIndices);
+        } else {
+            processSubset(jointGroup.inputIndices, 0ul, 0ul);
+        }
+        const auto jointGroupColumnCount = jointGroup.inputIndices.size();
+
+        processSubset(jointGroup.outputIndices, 0ul, jointGroupRowCount);
+        // Remap joint attribute indices
+        for (auto& attrIdx : jointGroup.outputIndices) {
+            const auto jntIdx = static_cast<std::uint16_t>(attrIdx / jointAttributeCount);
+            const auto relAttrIdx = attrIdx - (jntIdx * jointAttributeCount);
+            attrIdx = static_cast<std::uint16_t>(JointFilter::remapped(jntIdx) * jointAttributeCount + relAttrIdx);
+        }
+
+        processSubset(jointGroup.values, 0ul, jointGroupRowCount * jointGroupColumnCount);
+        // Load and remap joint indices (according to the remapping created while loading the Definition layer)
+        process(jointGroup.jointIndices);
+        extd::filter(jointGroup.jointIndices, [this](std::uint16_t jntIdx, std::size_t  /*unused*/) {
+                return JointFilter::passes(jntIdx);
+            });
+        for (auto& jntIdx : jointGroup.jointIndices) {
+            jntIdx = JointFilter::remapped(jntIdx);
+        }
+
+        dest.jointGroups.push_back(std::move(jointGroup));
+    }
+    const auto uncompressedJointCount = static_cast<std::uint16_t>(JointFilter::maxRemappedIndex() + 1u);
+    dest.rowCount = static_cast<std::uint16_t>(uncompressedJointCount * jointAttributeCount);
+}
+
+void FilteredInputArchive::process(RawBlendShapeChannels& dest) {
+    process(dest.lods);
+    if (!lodConstraint.hasImpactOn(unconstrainedLODCount)) {
+        process(dest.inputIndices);
+        process(dest.outputIndices);
+        return;
+    }
+    // Discard everything that falls outside the region bounded by LOD constraints
+    lodConstraint.applyTo(dest.lods);
+    const auto count = (dest.lods.empty() ? static_cast<std::uint16_t>(0) : dest.lods.front());
+    processSubset(dest.inputIndices, 0ul, count);
+    processSubset(dest.outputIndices, 0ul, count);
+}
+
+void FilteredInputArchive::process(RawAnimatedMaps& dest) {
+    process(dest.lods);
+    if (!lodConstraint.hasImpactOn(unconstrainedLODCount)) {
+        process(dest.conditionals);
+        return;
+    }
+    // Discard everything that falls outside the region bounded by LOD constraints
+    lodConstraint.applyTo(dest.lods);
+    const auto rowCount = (dest.lods.empty() ? static_cast<std::uint16_t>(0) : dest.lods.front());
+    processSubset(dest.conditionals.inputIndices, 0ul, rowCount);
+    processSubset(dest.conditionals.outputIndices, 0ul, rowCount);
+    processSubset(dest.conditionals.fromValues, 0ul, rowCount);
+    processSubset(dest.conditionals.toValues, 0ul, rowCount);
+    processSubset(dest.conditionals.slopeValues, 0ul, rowCount);
+    processSubset(dest.conditionals.cutValues, 0ul, rowCount);
+}
+
+void FilteredInputArchive::process(RawGeometry& dest) {
+    if (contains(layerBitmask, DataLayerBitmask::GeometryRest)) {
+        process(dest.marker);
+        if (!lodConstraint.hasImpactOn(unconstrainedLODCount)) {
+            process(dest.meshes);
+            return;
+        }
+        // Perform filtered load only if a different maxLOD is set
+        const auto meshCount = processSize();
+        dest.meshes.reserve(meshCount);
+        for (std::uint16_t i = {}; i < meshCount; ++i) {
+            RawMesh mesh{memRes};
+            // Check if the mesh indices filtered for the current maxLOD permit loading this mesh
+            if (MeshFilter::passes(i)) {
+                process(mesh);
+                dest.meshes.push_back(std::move(mesh));
+            } else {
+                // Jump over the whole section of data related to this mesh
+                process(mesh.offset);
+                process(mesh.marker);
+            }
+        }
+    }
+}
+
+void FilteredInputArchive::process(RawMesh& dest) {
+    process(dest.offset);
+    process(dest.positions);
+    process(dest.textureCoordinates);
+    process(dest.normals);
+    process(dest.layouts);
+    process(dest.faces);
+    process(dest.maximumInfluencePerVertex);
+    process(dest.skinWeights);
+    if (contains(layerBitmask, DataLayerBitmask::GeometryBlendShapesOnly)) {
+        process(dest.blendShapeTargets);
+        if (lodConstraint.hasImpactOn(unconstrainedLODCount)) {
+            extd::filter(dest.blendShapeTargets, [this](const RawBlendShapeTarget& bst, std::size_t  /*unused*/) {
+                    return BlendShapeFilter::passes(bst.blendShapeChannelIndex);
+                });
+        }
+    }
+    process(dest.marker);
+}
+
+void FilteredInputArchive::process(RawVertexSkinWeights& dest) {
+    process(dest.weights);
+    process(dest.jointIndices);
+
+    if (lodConstraint.hasImpactOn(unconstrainedLODCount)) {
+        assert(dest.weights.size() == dest.jointIndices.size());
+
+        auto itWeight = dest.weights.begin();
+        auto itJoint = dest.jointIndices.begin();
+        while (itJoint != dest.jointIndices.end()) {
+            if (JointFilter::passes(*itJoint)) {
+                ++itJoint;
+                ++itWeight;
+            } else {
+                itJoint = dest.jointIndices.erase(itJoint);
+                itWeight = dest.weights.erase(itWeight);
+            }
+        }
+
+        for (auto& jntIdx : dest.jointIndices) {
+            jntIdx = JointFilter::remapped(jntIdx);
+        }
+    }
+}
+
+}  // namespace dna
diff --git a/Engine/Plugins/Animation/RigLogic/Source/RigLogicLib/Private/dna/stream/FilteredInputArchive.h b/Engine/Plugins/Animation/RigLogic/Source/RigLogicLib/Private/dna/stream/FilteredInputArchive.h
new file mode 100644
--- /dev/null
+++ b/Engine/Plugins/Animation/RigLogic/Source/RigLogicLib/Private/dna/stream/FilteredInputArchive.h
@@ -0,0 +1,80 @@
+#pragma once
+
+#include "dna/DataLayer.h"
+#include "dna/DataLayerBitmask.h"
+#include "dna/LODConstraint.h"
+#include "dna/TypeDefs.h"
+#include "dna/filters/AnimatedMapFilter.h"
+#include "dna/filters/BlendShapeFilter.h"
+#include "dna/filters/JointFilter.h"
+#include "dna/filters/MeshFilter.h"
+
+#include <terse/archives/binary/InputArchive.h>
+
+#include <cstddef>
+#include <cstdint>
+
+namespace dna {
+
+struct RawAnimatedMaps;
+struct RawBehavior;
+struct RawBlendShapeChannels;
+struct RawDefinition;
+struct RawDescriptor;
+struct RawGeometry;
+struct RawJoints;
+struct RawMesh;
+struct RawVertexSkinWeights;
+
+class FilteredInputArchive final : public AnimatedMapFilter, public BlendShapeFilter, public JointFilter, public MeshFilter,
+    public terse::ExtendableBinaryInputArchive<FilteredInputArchive,
+                                               BoundedIOStream,
+                                               std::uint32_t,
+                                               std::uint32_t> {
+    private:
+        using BaseArchive = terse::ExtendableBinaryInputArchive<FilteredInputArchive,
+                                                                BoundedIOStream,
+                                                                std::uint32_t,
+                                                                std::uint32_t>;
+        friend Archive<FilteredInputArchive>;
+
+    public:
+        FilteredInputArchive(BoundedIOStream* stream_,
+                             DataLayer layer_,
+                             std::uint16_t maxLOD_,
+                             std::uint16_t minLOD_,
+                             MemoryResource* memRes_);
+        FilteredInputArchive(BoundedIOStream* stream_,
+                             DataLayer layer_,
+                             ConstArrayView<std::uint16_t> lods_,
+                             MemoryResource* memRes_);
+
+    private:
+        void process(RawDescriptor& dest);
+        void process(RawDefinition& dest);
+        void process(RawBehavior& dest);
+        void process(RawJoints& dest);
+        void process(RawBlendShapeChannels& dest);
+        void process(RawAnimatedMaps& dest);
+        void process(RawGeometry& dest);
+        void process(RawMesh& dest);
+        void process(RawVertexSkinWeights& dest);
+
+        template<typename ... Args>
+        void process(Args&& ... args) {
+            BaseArchive::process(std::forward<Args>(args)...);
+        }
+
+        template<typename TContainer>
+        void processSubset(TContainer& dest, std::size_t offset, std::size_t size);
+
+    private:
+        BoundedIOStream* stream;
+        MemoryResource* memRes;
+        DataLayerBitmask layerBitmask;
+        LODConstraint lodConstraint;
+        std::uint16_t unconstrainedLODCount;
+
+};
+
+}  // namespace dna
diff --git a/Engine/Plugins/Animation/RigLogic/Source/RigLogicLib/Private/dna/stream/StreamReaderImpl.cpp b/Engine/Plugins/Animation/RigLogic/Source/RigLogicLib/Private/dna/stream/StreamReaderImpl.cpp
new file mode 100644
--- /dev/null
+++ b/Engine/Plugins/Animation/RigLogic/Source/RigLogicLib/Private/dna/stream/StreamReaderImpl.cpp
@@ -0,0 +1,109 @@
+#include "dna/stream/StreamReaderImpl.h"
+
+#include "dna/TypeDefs.h"
+
+#include <status/Provider.h>
+#include <trio/utils/StreamScope.h>
+
+#include <cstddef>
+#include <limits>
+#include <tuple>
+
+namespace dna {
+
+const sc::StatusCode StreamReader::SignatureMismatchError{200, "DNA signature mismatched, expected %.3s, got %.3s"};
+const sc::StatusCode StreamReader::VersionMismatchError{201, "DNA version mismatched, expected %hu.%hu, got %hu.%hu"};
+const sc::StatusCode StreamReader::InvalidDataError{202, "Invalid data in DNA"};
+
+#ifdef __clang__
+    #pragma clang diagnostic push
+    #pragma clang diagnostic ignored "-Wglobal-constructors"
+#endif
+sc::StatusProvider StreamReaderImpl::status{SignatureMismatchError, VersionMismatchError, InvalidDataError};
+#ifdef __clang__
+    #pragma clang diagnostic pop
+#endif
+
+StreamReader::~StreamReader() = default;
+
+StreamReader* StreamReader::create(BoundedIOStream* stream, DataLayer layer, std::uint16_t maxLOD, MemoryResource* memRes) {
+    PolyAllocator<StreamReaderImpl> alloc{memRes};
+    return alloc.newObject(stream, layer, maxLOD, std::numeric_limits<std::uint16_t>::max(), memRes);
+}
+
+StreamReader* StreamReader::create(BoundedIOStream* stream,
+                                   DataLayer layer,
+                                   std::uint16_t maxLOD,
+                                   std::uint16_t minLOD,
+                                   MemoryResource* memRes) {
+    PolyAllocator<StreamReaderImpl> alloc{memRes};
+    return alloc.newObject(stream, layer, maxLOD, minLOD, memRes);
+}
+
+StreamReader* StreamReader::create(BoundedIOStream* stream,
+                                   DataLayer layer,
+                                   std::uint16_t* lods,
+                                   std::uint16_t lodCount,
+                                   MemoryResource* memRes) {
+    PolyAllocator<StreamReaderImpl> alloc{memRes};
+    return alloc.newObject(stream, layer, ConstArrayView<std::uint16_t>{lods, lodCount}, memRes);
+}
+
+void StreamReader::destroy(StreamReader* instance) {
+    // NOLINTNEXTLINE(cppcoreguidelines-pro-type-static-cast-downcast)
+    auto reader = static_cast<StreamReaderImpl*>(instance);
+    PolyAllocator<StreamReaderImpl> alloc{reader->getMemoryResource()};
+    alloc.deleteObject(reader);
+}
+
+StreamReaderImpl::StreamReaderImpl(BoundedIOStream* stream_,
+                                   DataLayer layer_,
+                                   std::uint16_t maxLOD_,
+                                   std::uint16_t minLOD_,
+                                   MemoryResource* memRes_) :
+    BaseImpl{memRes_},
+    ReaderImpl{memRes_},
+    stream{stream_},
+    dnaInputArchive{stream_, layer_, maxLOD_, minLOD_, memRes_} {
+}
+
+StreamReaderImpl::StreamReaderImpl(BoundedIOStream* stream_,
+                                   DataLayer layer_,
+                                   ConstArrayView<std::uint16_t> lods_,
+                                   MemoryResource* memRes_) :
+    BaseImpl{memRes_},
+    ReaderImpl{memRes_},
+    stream{stream_},
+    dnaInputArchive{stream_, layer_, lods_, memRes_} {
+}
+
+void StreamReaderImpl::read() {
+    // Due to possible usage of custom stream implementations, the status actually must be cleared at this point
+    // as external streams do not have access to the status reset API
+    status.reset();
+
+    trio::StreamScope scope{stream};
+    if (!sc::Status::isOk()) {
+        return;
+    }
+
+    dnaInputArchive >> dna;
+    if (!sc::Status::isOk()) {
+        return;
+    }
+
+    if (!dna.signature.matches()) {
+        status.set(SignatureMismatchError, dna.signature.value.expected.data(), dna.signature.value.got.data());
+        return;
+    }
+    if (!dna.version.matches()) {
+        status.set(VersionMismatchError,
+                   dna.version.generation.expected,
+                   dna.version.version.expected,
+                   dna.version.generation.got,
+                   dna.version.version.got);
+        return;
+    }
+}
+
+}  // namespace dna
diff --git a/Engine/Plugins/Animation/RigLogic/Source/RigLogicLib/Private/dna/stream/StreamReaderImpl.h b/Engine/Plugins/Animation/RigLogic/Source/RigLogicLib/Private/dna/stream/StreamReaderImpl.h
new file mode 100644
--- /dev/null
+++ b/Engine/Plugins/Animation/RigLogic/Source/RigLogicLib/Private/dna/stream/StreamReaderImpl.h
@@ -0,0 +1,31 @@
+#pragma once
+
+#include "dna/Reader.h"
+#include "dna/ReaderImpl.h"
+#include "dna/StreamReader.h"
+#include "dna/TypeDefs.h"
+#include "dna/stream/FilteredInputArchive.h"
+
+#include <status/Provider.h>
+
+namespace dna {
+
+class StreamReaderImpl : public ReaderImpl<StreamReader> {
+    public:
+        StreamReaderImpl(BoundedIOStream* stream_,
+                         DataLayer layer_,
+                         std::uint16_t maxLOD_,
+                         std::uint16_t minLOD_,
+                         MemoryResource* memRes_);
+        StreamReaderImpl(BoundedIOStream* stream_, DataLayer layer_, ConstArrayView<std::uint16_t> lods, MemoryResource* memRes_);
+
+        void read() override;
+
+    private:
+        static sc::StatusProvider status;
+
+        BoundedIOStream* stream;
+        FilteredInputArchive dnaInputArchive;
+};
+
+}  // namespace dna
diff --git a/Engine/Plugins/Animation/RigLogic/Source/RigLogicLib/Private/dna/stream/StreamWriterImpl.cpp b/Engine/Plugins/Animation/RigLogic/Source/RigLogicLib/Private/dna/stream/StreamWriterImpl.cpp
new file mode 100644
--- /dev/null
+++ b/Engine/Plugins/Animation/RigLogic/Source/RigLogicLib/Private/dna/stream/StreamWriterImpl.cpp
@@ -0,0 +1,40 @@
+#include "dna/stream/StreamWriterImpl.h"
+
+#include "dna/TypeDefs.h"
+
+#include <cassert>
+#include <cstddef>
+#include <cstring>
+#include <tuple>
+#include <utility>
+
+namespace dna {
+
+StreamWriter::~StreamWriter() = default;
+
+StreamWriter* StreamWriter::create(BoundedIOStream* stream, MemoryResource* memRes) {
+    PolyAllocator<StreamWriterImpl> alloc{memRes};
+    return alloc.newObject(stream, memRes);
+}
+
+void StreamWriter::destroy(StreamWriter* instance) {
+    // NOLINTNEXTLINE(cppcoreguidelines-pro-type-static-cast-downcast)
+    auto writer = static_cast<StreamWriterImpl*>(instance);
+    PolyAllocator<StreamWriterImpl> alloc{writer->getMemoryResource()};
+    alloc.deleteObject(writer);
+}
+
+StreamWriterImpl::StreamWriterImpl(BoundedIOStream* stream_, MemoryResource* memRes_) :
+    BaseImpl{memRes_},
+    WriterImpl{memRes_},
+    stream{stream_},
+    binaryOutputArchive{stream_} {
+}
+
+void StreamWriterImpl::write() {
+    stream->open();
+    binaryOutputArchive << dna;
+    stream->close();
+}
+
+}  // namespace dna
diff --git a/Engine/Plugins/Animation/RigLogic/Source/RigLogicLib/Private/dna/stream/StreamWriterImpl.h b/Engine/Plugins/Animation/RigLogic/Source/RigLogicLib/Private/dna/stream/StreamWriterImpl.h
new file mode 100644
--- /dev/null
+++ b/Engine/Plugins/Animation/RigLogic/Source/RigLogicLib/Private/dna/stream/StreamWriterImpl.h
@@ -0,0 +1,23 @@
+#pragma once
+
+#include "dna/DNA.h"
+#include "dna/StreamWriter.h"
+#include "dna/WriterImpl.h"
+
+#include <terse/archives/binary/OutputArchive.h>
+
+namespace dna {
+
+class StreamWriterImpl : public WriterImpl<StreamWriter> {
+    public:
+        StreamWriterImpl(BoundedIOStream* stream_, MemoryResource* memRes_);
+
+        void write() override;
+
+    private:
+        BoundedIOStream* stream;
+        terse::BinaryOutputArchive<BoundedIOStream> binaryOutputArchive;
+
+};
+
+}  // namespace dna
diff --git a/Engine/Plugins/Animation/RigLogic/Source/RigLogicLib/Private/dna/utils/ScopedEnumEx.h b/Engine/Plugins/Animation/RigLogic/Source/RigLogicLib/Private/dna/utils/ScopedEnumEx.h
new file mode 100644
--- /dev/null
+++ b/Engine/Plugins/Animation/RigLogic/Source/RigLogicLib/Private/dna/utils/ScopedEnumEx.h
@@ -0,0 +1,59 @@
+#pragma once
+
+#include <type_traits>
+
+namespace dna {
+
+template<typename TEnum>
+typename std::enable_if<std::is_enum<TEnum>::value, TEnum>::type
+operator&(TEnum lhs, TEnum rhs) {
+    using Underlying = typename std::underlying_type<TEnum>::type;
+    return static_cast<TEnum>(static_cast<Underlying>(lhs) & static_cast<Underlying>(rhs));
+}
+
+template<typename TEnum>
+typename std::enable_if<std::is_enum<TEnum>::value, TEnum>::type
+operator|(TEnum lhs, TEnum rhs) {
+    using Underlying = typename std::underlying_type<TEnum>::type;
+    return static_cast<TEnum>(static_cast<Underlying>(lhs) | static_cast<Underlying>(rhs));
+}
+
+template<typename TEnum>
+typename std::enable_if<std::is_enum<TEnum>::value, TEnum>::type
+operator^(TEnum lhs, TEnum rhs) {
+    using Underlying = typename std::underlying_type<TEnum>::type;
+    return static_cast<TEnum>(static_cast<Underlying>(lhs) ^ static_cast<Underlying>(rhs));
+}
+
+template<typename TEnum>
+typename std::enable_if<std::is_enum<TEnum>::value, TEnum>::type
+operator~(TEnum value) {
+    using Underlying = typename std::underlying_type<TEnum>::type;
+    return static_cast<TEnum>(~static_cast<Underlying>(value));
+}
+
+template<typename TEnum>
+typename std::enable_if<std::is_enum<TEnum>::value, TEnum>::type
+operator&=(TEnum& lhs, TEnum rhs) {
+    return lhs = (lhs & rhs);
+}
+
+template<typename TEnum>
+typename std::enable_if<std::is_enum<TEnum>::value, TEnum>::type
+operator|=(TEnum& lhs, TEnum rhs) {
+    return lhs = (lhs | rhs);
+}
+
+template<typename TEnum>
+typename std::enable_if<std::is_enum<TEnum>::value, TEnum>::type
+operator^=(TEnum& lhs, TEnum rhs) {
+    return lhs = (lhs ^ rhs);
+}
+
+template<typename TEnum>
+typename std::enable_if<std::is_enum<TEnum>::value, bool>::type
+contains(TEnum lhs, TEnum rhs) {
+    return (lhs & rhs) == rhs;
+}
+
+}  // namespace dna
diff --git a/Engine/Plugins/Animation/RigLogic/Source/RigLogicLib/Private/pma/resources/AlignedMemoryResource.cpp b/Engine/Plugins/Animation/RigLogic/Source/RigLogicLib/Private/pma/resources/AlignedMemoryResource.cpp
new file mode 100644
--- /dev/null
+++ b/Engine/Plugins/Animation/RigLogic/Source/RigLogicLib/Private/pma/resources/AlignedMemoryResource.cpp
@@ -0,0 +1,43 @@
+#include "pma/resources/AlignedMemoryResource.h"
+
+#include <cstddef>
+#include <cstdlib>
+
+#if defined(__linux__) || defined(__APPLE__)
+    #include <stdlib.h>
+#elif defined(_MSC_VER) || defined(__ANDROID__)
+    #include <malloc.h>
+#endif
+
+#if defined(__cplusplus) && (__cplusplus >= 201703L) && (defined(_GLIBCXX_HAVE_ALIGNED_ALLOC) || \
+    defined(_LIBCPP_HAS_ALIGNED_ALLOC) || defined(_LIBCPP_HAS_C11_FEATURES))
+    #define ALIGNED_ALLOC(ptr, alignment, size) ptr = std::aligned_alloc(alignment, size)
+    #define ALIGNED_FREE(ptr) std::free(ptr)
+#elif ((defined(_POSIX_VERSION) && (_POSIX_VERSION >= 200112L)) || defined(__linux__) || defined(__APPLE__))
+    #define ALIGNED_ALLOC(ptr, alignment, size) if (::posix_memalign(&(ptr), alignment, size)) \
+        (ptr) = nullptr
+    #define ALIGNED_FREE(ptr) ::free(ptr)
+#elif defined(_MSC_VER)
+    #define ALIGNED_ALLOC(ptr, alignment, size) ptr = _aligned_malloc(size, alignment)
+    #define ALIGNED_FREE(ptr) _aligned_free(ptr)
+#elif defined(__ANDROID__)
+    #define ALIGNED_ALLOC(ptr, alignment, size) ptr = ::memalign(alignment, size)
+    #define ALIGNED_FREE(ptr) ::free(ptr)
+#else
+    #error No aligned allocation is possible.
+#endif
+
+namespace pma {
+
+void* AlignedMemoryResource::allocate(std::size_t size, std::size_t alignment) {
+    void* ptr;
+    ALIGNED_ALLOC(ptr, alignment, size);
+    return ptr;
+}
+
+void AlignedMemoryResource::deallocate(void* ptr, std::size_t  /*unused*/, std::size_t  /*unused*/) {
+    // NOLINTNEXTLINE(cppcoreguidelines-owning-memory,cppcoreguidelines-no-malloc,hicpp-no-malloc)
+    ALIGNED_FREE(ptr);
+}
+
+}  // namespace pma
diff --git a/Engine/Plugins/Animation/RigLogic/Source/RigLogicLib/Private/pma/resources/DefaultMemoryResource.cpp b/Engine/Plugins/Animation/RigLogic/Source/RigLogicLib/Private/pma/resources/DefaultMemoryResource.cpp
new file mode 100644
--- /dev/null
+++ b/Engine/Plugins/Animation/RigLogic/Source/RigLogicLib/Private/pma/resources/DefaultMemoryResource.cpp
@@ -0,0 +1,18 @@
+#include "pma/resources/DefaultMemoryResource.h"
+
+#include <cstddef>
+#include <cstdlib>
+
+namespace pma {
+
+void* DefaultMemoryResource::allocate(std::size_t size, std::size_t  /*unused*/) {
+    // NOLINTNEXTLINE(cppcoreguidelines-owning-memory,cppcoreguidelines-no-malloc,hicpp-no-malloc)
+    return std::malloc(size);
+}
+
+void DefaultMemoryResource::deallocate(void* ptr, std::size_t  /*unused*/, std::size_t  /*unused*/) {
+    // NOLINTNEXTLINE(cppcoreguidelines-owning-memory,cppcoreguidelines-no-malloc,hicpp-no-malloc)
+    std::free(ptr);
+}
+
+}  // namespace pma
diff --git a/Engine/Plugins/Animation/RigLogic/Source/RigLogicLib/Private/riglogic/animatedmaps/AnimatedMaps.cpp b/Engine/Plugins/Animation/RigLogic/Source/RigLogicLib/Private/riglogic/animatedmaps/AnimatedMaps.cpp
new file mode 100644
--- /dev/null
+++ b/Engine/Plugins/Animation/RigLogic/Source/RigLogicLib/Private/riglogic/animatedmaps/AnimatedMaps.cpp
@@ -0,0 +1,21 @@
+#include "riglogic/animatedmaps/AnimatedMaps.h"
+
+#include "riglogic/TypeDefs.h"
+#include "riglogic/conditionaltable/ConditionalTable.h"
+
+#include <cassert>
+#include <utility>
+
+namespace rl4 {
+
+AnimatedMaps::AnimatedMaps(Vector<std::uint16_t>&& lods_, ConditionalTable&& conditionals_) :
+    lods{std::move(lods_)},
+    conditionals{std::move(conditionals_)} {
+}
+
+void AnimatedMaps::calculate(ConstArrayView<float> inputs, ArrayView<float> outputs, std::uint16_t lod) const {
+    assert(lod < lods.size());
+    conditionals.calculate(inputs.data(), outputs.data(), lods[lod]);
+}
+
+}  // namespace rl4
diff --git a/Engine/Plugins/Animation/RigLogic/Source/RigLogicLib/Private/riglogic/animatedmaps/AnimatedMaps.h b/Engine/Plugins/Animation/RigLogic/Source/RigLogicLib/Private/riglogic/animatedmaps/AnimatedMaps.h
new file mode 100644
--- /dev/null
+++ b/Engine/Plugins/Animation/RigLogic/Source/RigLogicLib/Private/riglogic/animatedmaps/AnimatedMaps.h
@@ -0,0 +1,28 @@
+#pragma once
+
+#include "riglogic/TypeDefs.h"
+#include "riglogic/conditionaltable/ConditionalTable.h"
+#include "riglogic/types/Aliases.h"
+
+#include <cstddef>
+
+namespace rl4 {
+
+class AnimatedMaps {
+    public:
+        AnimatedMaps(Vector<std::uint16_t>&& lods_, ConditionalTable&& conditionals_);
+
+        void calculate(ConstArrayView<float> inputs, ArrayView<float> outputs, std::uint16_t lod) const;
+
+        template<class Archive>
+        void serialize(Archive& archive) {
+            archive(lods, conditionals);
+        }
+
+    private:
+        Vector<std::uint16_t> lods;
+        ConditionalTable conditionals;
+
+};
+
+}  // namespace rl4
diff --git a/Engine/Plugins/Animation/RigLogic/Source/RigLogicLib/Private/riglogic/animatedmaps/AnimatedMapsFactory.h b/Engine/Plugins/Animation/RigLogic/Source/RigLogicLib/Private/riglogic/animatedmaps/AnimatedMapsFactory.h
new file mode 100644
--- /dev/null
+++ b/Engine/Plugins/Animation/RigLogic/Source/RigLogicLib/Private/riglogic/animatedmaps/AnimatedMapsFactory.h
@@ -0,0 +1,23 @@
+#pragma once
+
+#include "riglogic/animatedmaps/AnimatedMaps.h"
+
+#include <pma/utils/ManagedInstance.h>
+
+namespace dna {
+
+class BehaviorReader;
+
+}  // namespace dna
+
+namespace rl4 {
+
+struct AnimatedMapsFactory {
+    using ManagedAnimatedMaps = pma::UniqueInstance<AnimatedMaps>;
+    using AnimatedMapsPtr = typename ManagedAnimatedMaps::PointerType;
+
+    static AnimatedMapsPtr create(MemoryResource* memRes);
+    static AnimatedMapsPtr create(const dna::BehaviorReader* reader, MemoryResource* memRes);
+};
+
+}  // namespace rl4
diff --git a/Engine/Plugins/Animation/RigLogic/Source/RigLogicLib/Private/riglogic/blendshapes/BlendShapes.cpp b/Engine/Plugins/Animation/RigLogic/Source/RigLogicLib/Private/riglogic/blendshapes/BlendShapes.cpp
new file mode 100644
--- /dev/null
+++ b/Engine/Plugins/Animation/RigLogic/Source/RigLogicLib/Private/riglogic/blendshapes/BlendShapes.cpp
@@ -0,0 +1,26 @@
+#include "riglogic/blendshapes/BlendShapes.h"
+
+#include "riglogic/TypeDefs.h"
+
+#include <cassert>
+#include <cstddef>
+
+namespace rl4 {
+
+BlendShapes::BlendShapes(Vector<std::uint16_t>&& lods_,
+                         Vector<std::uint16_t>&& inputIndices_,
+                         Vector<std::uint16_t>&& outputIndices_) :
+    lods{std::move(lods_)},
+    inputIndices{std::move(inputIndices_)},
+    outputIndices{std::move(outputIndices_)} {
+}
+
+void BlendShapes::calculate(ConstArrayView<float> inputs, ArrayView<float> outputs, std::uint16_t lod) const {
+    assert(lod < lods.size());
+    std::fill(outputs.begin(), outputs.end(), 0.0f);
+    for (std::uint16_t i = 0u; i < lods[lod]; ++i) {
+        outputs[outputIndices[i]] = inputs[inputIndices[i]];
+    }
+}
+
+}  // namespace rl4
diff --git a/Engine/Plugins/Animation/RigLogic/Source/RigLogicLib/Private/riglogic/blendshapes/BlendShapes.h b/Engine/Plugins/Animation/RigLogic/Source/RigLogicLib/Private/riglogic/blendshapes/BlendShapes.h
new file mode 100644
--- /dev/null
+++ b/Engine/Plugins/Animation/RigLogic/Source/RigLogicLib/Private/riglogic/blendshapes/BlendShapes.h
@@ -0,0 +1,28 @@
+#pragma once
+
+#include "riglogic/TypeDefs.h"
+#include "riglogic/types/Aliases.h"
+
+#include <cstddef>
+
+namespace rl4 {
+
+class BlendShapes {
+    public:
+        BlendShapes(Vector<std::uint16_t>&& lods_, Vector<std::uint16_t>&& inputIndices_, Vector<std::uint16_t>&& outputIndices_);
+
+        void calculate(ConstArrayView<float> inputs, ArrayView<float> outputs, std::uint16_t lod) const;
+
+        template<class Archive>
+        void serialize(Archive& archive) {
+            archive(lods, inputIndices, outputIndices);
+        }
+
+    private:
+        Vector<std::uint16_t> lods;
+        Vector<std::uint16_t> inputIndices;
+        Vector<std::uint16_t> outputIndices;
+
+};
+
+}  // namespace rl4
diff --git a/Engine/Plugins/Animation/RigLogic/Source/RigLogicLib/Private/riglogic/blendshapes/BlendShapesFactory.cpp b/Engine/Plugins/Animation/RigLogic/Source/RigLogicLib/Private/riglogic/blendshapes/BlendShapesFactory.cpp
new file mode 100644
--- /dev/null
+++ b/Engine/Plugins/Animation/RigLogic/Source/RigLogicLib/Private/riglogic/blendshapes/BlendShapesFactory.cpp
@@ -0,0 +1,33 @@
+#include "riglogic/blendshapes/BlendShapesFactory.h"
+
+#include "riglogic/utils/Extd.h"
+
+#include <dna/layers/BehaviorReader.h>
+
+#include <cstdint>
+#include <utility>
+#include <cstring>
+
+namespace rl4 {
+
+BlendShapesFactory::BlendShapesPtr BlendShapesFactory::create(MemoryResource* memRes) {
+    return ManagedBlendShapes::with(memRes).create(Vector<std::uint16_t>{memRes},
+                                                   Vector<std::uint16_t>{memRes},
+                                                   Vector<std::uint16_t>{memRes});
+}
+
+BlendShapesFactory::BlendShapesPtr BlendShapesFactory::create(const dna::BehaviorReader* reader, MemoryResource* memRes) {
+    Vector<std::uint16_t> lods{memRes};
+    Vector<std::uint16_t> inputIndices{memRes};
+    Vector<std::uint16_t> outputIndices{memRes};
+
+    extd::copy(reader->getBlendShapeChannelLODs(), lods);
+    extd::copy(reader->getBlendShapeChannelInputIndices(), inputIndices);
+    extd::copy(reader->getBlendShapeChannelOutputIndices(), outputIndices);
+
+    return ManagedBlendShapes::with(memRes).create(std::move(lods),
+                                                   std::move(inputIndices),
+                                                   std::move(outputIndices));
+}
+
+}  // namespace rl4
diff --git a/Engine/Plugins/Animation/RigLogic/Source/RigLogicLib/Private/riglogic/blendshapes/BlendShapesFactory.h b/Engine/Plugins/Animation/RigLogic/Source/RigLogicLib/Private/riglogic/blendshapes/BlendShapesFactory.h
new file mode 100644
--- /dev/null
+++ b/Engine/Plugins/Animation/RigLogic/Source/RigLogicLib/Private/riglogic/blendshapes/BlendShapesFactory.h
@@ -0,0 +1,23 @@
+#pragma once
+
+#include "riglogic/blendshapes/BlendShapes.h"
+
+#include <pma/utils/ManagedInstance.h>
+
+namespace dna {
+
+class BehaviorReader;
+
+}  // namespace dna
+
+namespace rl4 {
+
+struct BlendShapesFactory {
+    using ManagedBlendShapes = pma::UniqueInstance<BlendShapes>;
+    using BlendShapesPtr = typename ManagedBlendShapes::PointerType;
+
+    static BlendShapesPtr create(MemoryResource* memRes);
+    static BlendShapesPtr create(const dna::BehaviorReader* reader, MemoryResource* memRes);
+};
+
+}  // namespace rl4
diff --git a/Engine/Plugins/Animation/RigLogic/Source/RigLogicLib/Private/riglogic/controls/Controls.cpp b/Engine/Plugins/Animation/RigLogic/Source/RigLogicLib/Private/riglogic/controls/Controls.cpp
new file mode 100644
--- /dev/null
+++ b/Engine/Plugins/Animation/RigLogic/Source/RigLogicLib/Private/riglogic/controls/Controls.cpp
@@ -0,0 +1,25 @@
+#include "riglogic/controls/Controls.h"
+
+#include "riglogic/conditionaltable/ConditionalTable.h"
+#include "riglogic/psdmatrix/PSDMatrix.h"
+
+#include <cassert>
+#include <cstdint>
+
+namespace rl4 {
+
+Controls::Controls(ConditionalTable&& guiToRawMapping_, PSDMatrix&& psds_) :
+    guiToRawMapping{std::move(guiToRawMapping_)},
+    psds{std::move(psds_)} {
+}
+
+void Controls::mapGUIToRaw(ConstArrayView<float> guiValues, ArrayView<float> inputs) const {
+    assert(guiValues.size() == guiToRawMapping.getInputCount());
+    guiToRawMapping.calculate(guiValues.data(), inputs.data());
+}
+
+void Controls::calculate(ArrayView<float> inputs) const {
+    psds.calculate(inputs.data(), guiToRawMapping.getOutputCount());
+}
+
+}  // namespace rl4
diff --git a/Engine/Plugins/Animation/RigLogic/Source/RigLogicLib/Private/riglogic/controls/Controls.h b/Engine/Plugins/Animation/RigLogic/Source/RigLogicLib/Private/riglogic/controls/Controls.h
new file mode 100644
--- /dev/null
+++ b/Engine/Plugins/Animation/RigLogic/Source/RigLogicLib/Private/riglogic/controls/Controls.h
@@ -0,0 +1,29 @@
+#pragma once
+
+#include "riglogic/conditionaltable/ConditionalTable.h"
+#include "riglogic/psdmatrix/PSDMatrix.h"
+#include "riglogic/types/Aliases.h"
+
+#include <cstdint>
+
+namespace rl4 {
+
+class Controls {
+    public:
+        Controls(ConditionalTable&& guiToRawMapping_, PSDMatrix&& psds_);
+
+        void mapGUIToRaw(ConstArrayView<float> guiValues, ArrayView<float> inputs) const;
+        void calculate(ArrayView<float> inputs) const;
+
+        template<class Archive>
+        void serialize(Archive& archive) {
+            archive(guiToRawMapping, psds);
+        }
+
+    private:
+        ConditionalTable guiToRawMapping;
+        PSDMatrix psds;
+
+};
+
+}  // namespace rl4
diff --git a/Engine/Plugins/Animation/RigLogic/Source/RigLogicLib/Private/riglogic/controls/ControlsFactory.cpp b/Engine/Plugins/Animation/RigLogic/Source/RigLogicLib/Private/riglogic/controls/ControlsFactory.cpp
new file mode 100644
--- /dev/null
+++ b/Engine/Plugins/Animation/RigLogic/Source/RigLogicLib/Private/riglogic/controls/ControlsFactory.cpp
@@ -0,0 +1,63 @@
+#include "riglogic/controls/ControlsFactory.h"
+
+#include "riglogic/conditionaltable/ConditionalTable.h"
+#include "riglogic/controls/Controls.h"
+#include "riglogic/psdmatrix/PSDMatrix.h"
+#include "riglogic/utils/Extd.h"
+
+#include <dna/layers/BehaviorReader.h>
+
+#include <cstddef>
+#include <cstring>
+
+namespace rl4 {
+
+ControlsFactory::ControlsPtr ControlsFactory::create(MemoryResource* memRes) {
+    return ManagedControls::with(memRes).create(ConditionalTable{memRes}, PSDMatrix{memRes});
+}
+
+ControlsFactory::ControlsPtr ControlsFactory::create(const dna::BehaviorReader* reader, MemoryResource* memRes) {
+    Vector<std::uint16_t> inputIndices{memRes};
+    Vector<std::uint16_t> outputIndices{memRes};
+    Vector<float> fromValues{memRes};
+    Vector<float> toValues{memRes};
+    Vector<float> slopeValues{memRes};
+    Vector<float> cutValues{memRes};
+    extd::copy(reader->getGUIToRawInputIndices(), inputIndices);
+    extd::copy(reader->getGUIToRawOutputIndices(), outputIndices);
+    extd::copy(reader->getGUIToRawFromValues(), fromValues);
+    extd::copy(reader->getGUIToRawToValues(), toValues);
+    extd::copy(reader->getGUIToRawSlopeValues(), slopeValues);
+    extd::copy(reader->getGUIToRawCutValues(), cutValues);
+    // DNAs may contain these parameters in reverse order
+    // i.e. the `from` value is actually larger than the `to` value
+    assert(fromValues.size() == toValues.size());
+    for (std::size_t i = 0ul; i < fromValues.size(); ++i) {
+        if (fromValues[i] > toValues[i]) {
+            std::swap(fromValues[i], toValues[i]);
+        }
+    }
+    ConditionalTable conditionals{std::move(inputIndices),
+                                  std::move(outputIndices),
+                                  std::move(fromValues),
+                                  std::move(toValues),
+                                  std::move(slopeValues),
+                                  std::move(cutValues),
+                                  reader->getGUIControlCount(),
+                                  reader->getRawControlCount()};
+
+    Vector<std::uint16_t> psdRows{memRes};
+    Vector<std::uint16_t> psdColumns{memRes};
+    Vector<float> psdValues{memRes};
+    extd::copy(reader->getPSDRowIndices(), psdRows);
+    extd::copy(reader->getPSDColumnIndices(), psdColumns);
+    extd::copy(reader->getPSDValues(), psdValues);
+    PSDMatrix psds{reader->getPSDCount(),
+                   std::move(psdRows),
+                   std::move(psdColumns),
+                   std::move(psdValues)};
+
+    return ManagedControls::with(memRes).create(std::move(conditionals), std::move(psds));
+}
+
+}  // namespace rl4
diff --git a/Engine/Plugins/Animation/RigLogic/Source/RigLogicLib/Private/riglogic/controls/ControlsFactory.h b/Engine/Plugins/Animation/RigLogic/Source/RigLogicLib/Private/riglogic/controls/ControlsFactory.h
new file mode 100644
--- /dev/null
+++ b/Engine/Plugins/Animation/RigLogic/Source/RigLogicLib/Private/riglogic/controls/ControlsFactory.h
@@ -0,0 +1,23 @@
+#pragma once
+
+#include "riglogic/controls/Controls.h"
+
+#include <pma/utils/ManagedInstance.h>
+
+namespace dna {
+
+class BehaviorReader;
+
+}  // namespace dna
+
+namespace rl4 {
+
+struct ControlsFactory {
+    using ManagedControls = pma::UniqueInstance<Controls>;
+    using ControlsPtr = typename ManagedControls::PointerType;
+
+    static ControlsPtr create(MemoryResource* memRes);
+    static ControlsPtr create(const dna::BehaviorReader* reader, MemoryResource* memRes);
+};
+
+}  // namespace rl4
diff --git a/Engine/Plugins/Animation/RigLogic/Source/RigLogicLib/Private/riglogic/joints/Joints.cpp b/Engine/Plugins/Animation/RigLogic/Source/RigLogicLib/Private/riglogic/joints/Joints.cpp
new file mode 100644
--- /dev/null
+++ b/Engine/Plugins/Animation/RigLogic/Source/RigLogicLib/Private/riglogic/joints/Joints.cpp
@@ -0,0 +1,55 @@
+#include "riglogic/joints/Joints.h"
+
+#include "riglogic/riglogic/RigInstance.h"
+#include "riglogic/transformation/Transformation.h"
+#include "riglogic/types/Aliases.h"
+
+#include <cstdint>
+
+namespace rl4 {
+
+Joints::Joints(JointsEvaluatorPtr evaluator_, MemoryResource* memRes) :
+    evaluator{std::move(evaluator_)},
+    neutralValues{memRes},
+    variableAttributeIndices{memRes},
+    jointGroupCount{} {
+}
+
+Joints::Joints(JointsEvaluatorPtr evaluator_,
+               Vector<float>&& neutralValues_,
+               Matrix<std::uint16_t>&& variableAttributeIndices_,
+               std::uint16_t jointGroupCount_) :
+    evaluator{std::move(evaluator_)},
+    neutralValues{std::move(neutralValues_)},
+    variableAttributeIndices{std::move(variableAttributeIndices_)},
+    jointGroupCount{jointGroupCount_} {
+}
+
+void Joints::calculate(ConstArrayView<float> inputs, ArrayView<float> outputs, std::uint16_t lod) const {
+    evaluator->calculate(inputs, outputs, lod);
+}
+
+void Joints::calculate(ConstArrayView<float> inputs, ArrayView<float> outputs, std::uint16_t lod,
+                       std::uint16_t jointGroupIndex) const {
+    evaluator->calculate(inputs, outputs, lod, jointGroupIndex);
+}
+
+ConstArrayView<float> Joints::getRawNeutralValues() const {
+    return ConstArrayView<float>{neutralValues};
+}
+
+TransformationArrayView Joints::getNeutralValues() const {
+    return TransformationArrayView{neutralValues.data(), neutralValues.size()};
+}
+
+std::uint16_t Joints::getJointGroupCount() const {
+    return jointGroupCount;
+}
+
+ConstArrayView<std::uint16_t> Joints::getVariableAttributeIndices(std::uint16_t lod) const {
+    return (lod < variableAttributeIndices.size()
+            ? ConstArrayView<std::uint16_t>{variableAttributeIndices[lod]}
+            : ConstArrayView<std::uint16_t>{});
+}
+
+}  // namespace rl4
diff --git a/Engine/Plugins/Animation/RigLogic/Source/RigLogicLib/Private/riglogic/joints/Joints.h b/Engine/Plugins/Animation/RigLogic/Source/RigLogicLib/Private/riglogic/joints/Joints.h
new file mode 100644
--- /dev/null
+++ b/Engine/Plugins/Animation/RigLogic/Source/RigLogicLib/Private/riglogic/joints/Joints.h
@@ -0,0 +1,60 @@
+#pragma once
+
+#include "riglogic/TypeDefs.h"
+#include "riglogic/joints/JointsEvaluator.h"
+#include "riglogic/transformation/Transformation.h"
+#include "riglogic/types/Aliases.h"
+
+#ifdef _MSC_VER
+    #pragma warning(push)
+    #pragma warning(disable : 4365 4987)
+#endif
+#include <cstdint>
+#include <functional>
+#include <memory>
+#ifdef _MSC_VER
+    #pragma warning(pop)
+#endif
+
+namespace rl4 {
+
+class RigInstance;
+
+class Joints {
+    public:
+        Joints(JointsEvaluatorPtr evaluator_, MemoryResource* memRes);
+        Joints(JointsEvaluatorPtr evaluator_,
+               Vector<float>&& neutralValues_,
+               Matrix<std::uint16_t>&& variableAttributeIndices_,
+               std::uint16_t jointGroupCount_);
+
+        void calculate(ConstArrayView<float> inputs, ArrayView<float> outputs, std::uint16_t lod) const;
+        void calculate(ConstArrayView<float> inputs, ArrayView<float> outputs, std::uint16_t lod,
+                       std::uint16_t jointGroupIndex) const;
+
+        template<class Archive>
+        void load(Archive& archive) {
+            evaluator->load(archive);
+            archive >> neutralValues >> variableAttributeIndices >> jointGroupCount;
+        }
+
+        template<class Archive>
+        void save(Archive& archive) {
+            evaluator->save(archive);
+            archive << neutralValues << variableAttributeIndices << jointGroupCount;
+        }
+
+        std::uint16_t getJointGroupCount() const;
+        ConstArrayView<float> getRawNeutralValues() const;
+        TransformationArrayView getNeutralValues() const;
+        ConstArrayView<std::uint16_t> getVariableAttributeIndices(std::uint16_t lod) const;
+
+    private:
+        JointsEvaluatorPtr evaluator;
+        Vector<float> neutralValues;
+        Matrix<std::uint16_t> variableAttributeIndices;
+        std::uint16_t jointGroupCount;
+
+};
+
+}  // namespace rl4
diff --git a/Engine/Plugins/Animation/RigLogic/Source/RigLogicLib/Private/riglogic/joints/JointsBuilder.cpp b/Engine/Plugins/Animation/RigLogic/Source/RigLogicLib/Private/riglogic/joints/JointsBuilder.cpp
new file mode 100644
--- /dev/null
+++ b/Engine/Plugins/Animation/RigLogic/Source/RigLogicLib/Private/riglogic/joints/JointsBuilder.cpp
@@ -0,0 +1,39 @@
+#include "riglogic/joints/JointsBuilder.h"
+
+#include "riglogic/joints/bpcm/builders/Scalar.h"
+#ifdef RL_BUILD_WITH_AVX
+    #ifdef RL_USE_HALF_FLOATS
+        #include "riglogic/joints/bpcm/builders/AVXHalfFloat.h"
+    #else
+        #include "riglogic/joints/bpcm/builders/AVX.h"
+    #endif  // RL_USE_HALF_FLOATS
+#endif  // RL_BUILD_WITH_AVX
+#ifdef RL_BUILD_WITH_SSE
+    #ifdef RL_USE_HALF_FLOATS
+        #include "riglogic/joints/bpcm/builders/SSEHalfFloat.h"
+    #else
+        #include "riglogic/joints/bpcm/builders/SSE.h"
+    #endif  // RL_USE_HALF_FLOATS
+#endif  // RL_BUILD_WITH_SSE
+
+#include <pma/utils/ManagedInstance.h>
+
+namespace rl4 {
+
+JointsBuilder::~JointsBuilder() = default;
+
+JointsBuilder::JointsBuilderPtr JointsBuilder::create(Configuration config, MemoryResource* memRes) {
+    #ifdef RL_BUILD_WITH_SSE
+        if (config.calculationType == CalculationType::SSE) {
+            return pma::UniqueInstance<bpcm::SSEJointsBuilder, JointsBuilder>::with(memRes).create(memRes);
+        }
+    #endif  // RL_BUILD_WITH_SSE
+    #ifdef RL_BUILD_WITH_AVX
+        if (config.calculationType == CalculationType::AVX) {
+            return pma::UniqueInstance<bpcm::AVXJointsBuilder, JointsBuilder>::with(memRes).create(memRes);
+        }
+    #endif  // RL_BUILD_WITH_AVX
+    return pma::UniqueInstance<bpcm::ScalarJointsBuilder, JointsBuilder>::with(memRes).create(memRes);
+}
+
+}  // namespace rl4
diff --git a/Engine/Plugins/Animation/RigLogic/Source/RigLogicLib/Private/riglogic/joints/JointsBuilder.h b/Engine/Plugins/Animation/RigLogic/Source/RigLogicLib/Private/riglogic/joints/JointsBuilder.h
new file mode 100644
--- /dev/null
+++ b/Engine/Plugins/Animation/RigLogic/Source/RigLogicLib/Private/riglogic/joints/JointsBuilder.h
@@ -0,0 +1,34 @@
+#pragma once
+
+#include "riglogic/joints/JointsEvaluator.h"
+#include "riglogic/riglogic/Configuration.h"
+#include "riglogic/types/Aliases.h"
+
+#ifdef _MSC_VER
+    #pragma warning(push)
+    #pragma warning(disable : 4365 4987)
+#endif
+#include <functional>
+#include <memory>
+#ifdef _MSC_VER
+    #pragma warning(pop)
+#endif
+
+namespace rl4 {
+
+class JointsBuilder {
+    public:
+        using JointsEvaluatorPtr = std::unique_ptr<JointsEvaluator, std::function<void (JointsEvaluator*)> >;
+        using JointsBuilderPtr = std::unique_ptr<JointsBuilder, std::function<void (JointsBuilder*)> >;
+
+    public:
+        virtual ~JointsBuilder();
+
+        static JointsBuilderPtr create(Configuration config, MemoryResource* memRes);
+
+        virtual void allocateStorage(const dna::BehaviorReader* source) = 0;
+        virtual void fillStorage(const dna::BehaviorReader* source) = 0;
+        virtual JointsEvaluatorPtr build() = 0;
+};
+
+}  // namespace rl4
diff --git a/Engine/Plugins/Animation/RigLogic/Source/RigLogicLib/Private/riglogic/joints/JointsEvaluator.h b/Engine/Plugins/Animation/RigLogic/Source/RigLogicLib/Private/riglogic/joints/JointsEvaluator.h
new file mode 100644
--- /dev/null
+++ b/Engine/Plugins/Animation/RigLogic/Source/RigLogicLib/Private/riglogic/joints/JointsEvaluator.h
@@ -0,0 +1,40 @@
+#pragma once
+
+#include "riglogic/types/Aliases.h"
+
+#include <terse/archives/binary/InputArchive.h>
+#include <terse/archives/binary/OutputArchive.h>
+#include <trio/Stream.h>
+
+#ifdef _MSC_VER
+    #pragma warning(push)
+    #pragma warning(disable : 4365 4987)
+#endif
+#include <cstdint>
+#include <functional>
+#include <memory>
+#ifdef _MSC_VER
+    #pragma warning(pop)
+#endif
+
+namespace rl4 {
+
+class RigInstance;
+
+class JointsEvaluator {
+    protected:
+        virtual ~JointsEvaluator();
+
+    public:
+        virtual void calculate(ConstArrayView<float> inputs, ArrayView<float> outputs, std::uint16_t lod) const = 0;
+        virtual void calculate(ConstArrayView<float> inputs,
+                               ArrayView<float> outputs,
+                               std::uint16_t lod,
+                               std::uint16_t jointGroupIndex) const = 0;
+        virtual void load(terse::BinaryInputArchive<BoundedIOStream>& archive) = 0;
+        virtual void save(terse::BinaryOutputArchive<BoundedIOStream>& archive) = 0;
+};
+
+using JointsEvaluatorPtr = std::unique_ptr<JointsEvaluator, std::function<void (JointsEvaluator*)> >;
+
+}  // namespace rl4
diff --git a/Engine/Plugins/Animation/RigLogic/Source/RigLogicLib/Private/riglogic/joints/JointsFactory.cpp b/Engine/Plugins/Animation/RigLogic/Source/RigLogicLib/Private/riglogic/joints/JointsFactory.cpp
new file mode 100644
--- /dev/null
+++ b/Engine/Plugins/Animation/RigLogic/Source/RigLogicLib/Private/riglogic/joints/JointsFactory.cpp
@@ -0,0 +1,91 @@
+#include "riglogic/joints/JointsFactory.h"
+
+#include "riglogic/joints/JointsBuilder.h"
+#include "riglogic/joints/JointsEvaluator.h"
+#include "riglogic/riglogic/Configuration.h"
+#include "riglogic/transformation/Transformation.h"
+#include "riglogic/types/Aliases.h"
+
+#include <dna/layers/BehaviorReader.h>
+#include <pma/utils/ManagedInstance.h>
+
+#ifdef _MSC_VER
+    #pragma warning(push)
+    #pragma warning(disable : 4365 4987)
+#endif
+#include <algorithm>
+#include <cstddef>
+#include <cstdint>
+#include <iterator>
+#include <utility>
+#ifdef _MSC_VER
+    #pragma warning(pop)
+#endif
+
+namespace rl4 {
+
+namespace {
+
+constexpr std::size_t txOffset = 0ul;
+constexpr std::size_t tyOffset = 1ul;
+constexpr std::size_t tzOffset = 2ul;
+constexpr std::size_t rxOffset = 3ul;
+constexpr std::size_t ryOffset = 4ul;
+constexpr std::size_t rzOffset = 5ul;
+
+}  // namespace
+
+template<class TSource, class TDestination>
+static void scatter(const TSource& source, TDestination& destination, std::size_t stride, std::size_t offset) {
+    const auto size = static_cast<std::size_t>(std::distance(std::begin(source), std::end(source)));
+    for (std::size_t i = 0ul; i < size; ++i) {
+        destination[i * stride + offset] = source[i];
+    }
+}
+
+static Vector<float> copyNeutralValues(const dna::BehaviorReader* reader, MemoryResource* memRes) {
+    Vector<float> neutralValues{memRes};
+    neutralValues.resize(reader->getJointCount() * Transformation::size());
+    std::fill(neutralValues.begin(), neutralValues.end(), 1.0f);
+    scatter(reader->getNeutralJointTranslationXs(), neutralValues, Transformation::size(), txOffset);
+    scatter(reader->getNeutralJointTranslationYs(), neutralValues, Transformation::size(), tyOffset);
+    scatter(reader->getNeutralJointTranslationZs(), neutralValues, Transformation::size(), tzOffset);
+    scatter(reader->getNeutralJointRotationXs(), neutralValues, Transformation::size(), rxOffset);
+    scatter(reader->getNeutralJointRotationYs(), neutralValues, Transformation::size(), ryOffset);
+    scatter(reader->getNeutralJointRotationZs(), neutralValues, Transformation::size(), rzOffset);
+    return neutralValues;
+}
+
+static Matrix<std::uint16_t> copyVariableAttributeIndices(const dna::BehaviorReader* reader, MemoryResource* memRes) {
+    Matrix<std::uint16_t> variableAttributeIndices{memRes};
+    variableAttributeIndices.resize(reader->getLODCount());
+    for (std::uint16_t lod = 0u; lod < reader->getLODCount(); ++lod) {
+        auto indices = reader->getJointVariableAttributeIndices(lod);
+        variableAttributeIndices[lod].assign(indices.begin(), indices.end());
+    }
+    return variableAttributeIndices;
+}
+
+JointsFactory::JointsPtr JointsFactory::create(const Configuration& config,
+                                               const dna::BehaviorReader* reader,
+                                               MemoryResource* memRes) {
+    auto neutralValues = copyNeutralValues(reader, memRes);
+    auto variableAttributeIndices = copyVariableAttributeIndices(reader, memRes);
+    auto builder = JointsBuilder::create(config, memRes);
+    builder->allocateStorage(reader);
+    builder->fillStorage(reader);
+    auto evaluator = builder->build();
+    return pma::UniqueInstance<Joints>::with(memRes).create(std::move(evaluator),
+                                                            std::move(neutralValues),
+                                                            std::move(variableAttributeIndices),
+                                                            reader->getJointGroupCount());
+}
+
+JointsFactory::JointsPtr JointsFactory::create(const Configuration& config, const RigMetrics&  /*unused*/,
+                                               MemoryResource* memRes) {
+    auto builder = JointsBuilder::create(config, memRes);
+    auto evaluator = builder->build();
+    return pma::UniqueInstance<Joints>::with(memRes).create(std::move(evaluator), memRes);
+}
+
+}  // namespace rl4
diff --git a/Engine/Plugins/Animation/RigLogic/Source/RigLogicLib/Private/riglogic/joints/JointsFactory.h b/Engine/Plugins/Animation/RigLogic/Source/RigLogicLib/Private/riglogic/joints/JointsFactory.h
new file mode 100644
--- /dev/null
+++ b/Engine/Plugins/Animation/RigLogic/Source/RigLogicLib/Private/riglogic/joints/JointsFactory.h
@@ -0,0 +1,28 @@
+#pragma once
+
+#include "riglogic/joints/Joints.h"
+#include "riglogic/riglogic/Configuration.h"
+#include "riglogic/riglogic/RigMetrics.h"
+#include "riglogic/types/Aliases.h"
+
+#ifdef _MSC_VER
+    #pragma warning(push)
+    #pragma warning(disable : 4365 4987)
+#endif
+#include <functional>
+#include <memory>
+#ifdef _MSC_VER
+    #pragma warning(pop)
+#endif
+
+namespace rl4 {
+
+struct JointsFactory {
+    using JointsPtr = std::unique_ptr<Joints, std::function<void (Joints*)> >;
+
+    static JointsPtr create(const Configuration& config, const dna::BehaviorReader* reader, MemoryResource* memRes);
+    static JointsPtr create(const Configuration& config, const RigMetrics& metrics, MemoryResource* memRes);
+
+};
+
+}  // namespace rl4
diff --git a/Engine/Plugins/Animation/RigLogic/Source/RigLogicLib/Private/riglogic/joints/bpcm/Consts.h b/Engine/Plugins/Animation/RigLogic/Source/RigLogicLib/Private/riglogic/joints/bpcm/Consts.h
new file mode 100644
--- /dev/null
+++ b/Engine/Plugins/Animation/RigLogic/Source/RigLogicLib/Private/riglogic/joints/bpcm/Consts.h
@@ -0,0 +1,24 @@
+#pragma once
+
+#include <cstddef>
+#include <cstdint>
+
+namespace rl4 {
+
+namespace bpcm {
+
+constexpr std::size_t cacheLineSize = 64ul;  // bytes
+constexpr std::size_t lookAheadOffset = cacheLineSize / sizeof(float);  // in number of elements
+
+constexpr std::uint32_t block4Width = 8u;
+constexpr std::uint32_t block4Height = 4u;
+
+constexpr std::uint32_t block8Width = 4u;
+constexpr std::uint32_t block8Height = 8u;
+
+constexpr std::uint32_t block16Width = 4u;
+constexpr std::uint32_t block16Height = 16u;
+
+}  // namespace bpcm
+
+}  // namespace rl4
diff --git a/Engine/Plugins/Animation/RigLogic/Source/RigLogicLib/Private/riglogic/joints/bpcm/Extent.h b/Engine/Plugins/Animation/RigLogic/Source/RigLogicLib/Private/riglogic/joints/bpcm/Extent.h
new file mode 100644
--- /dev/null
+++ b/Engine/Plugins/Animation/RigLogic/Source/RigLogicLib/Private/riglogic/joints/bpcm/Extent.h
@@ -0,0 +1,29 @@
+#pragma once
+
+#include <cstdint>
+
+namespace rl4 {
+
+namespace bpcm {
+
+struct Extent {
+    std::uint16_t rows;
+    std::uint16_t cols;
+
+    std::uint16_t size() const {
+        return static_cast<std::uint16_t>(rows * cols);
+    }
+
+};
+
+inline bool operator==(const Extent& lhs, const Extent& rhs) {
+    return (lhs.rows == rhs.rows && lhs.cols == rhs.cols);
+}
+
+inline bool operator!=(const Extent& lhs, const Extent& rhs) {
+    return !(lhs == rhs);
+}
+
+}  // namespace bpcm
+
+}  // namespace rl4
diff --git a/Engine/Plugins/Animation/RigLogic/Source/RigLogicLib/Private/riglogic/joints/bpcm/JointCalculationStrategy.h b/Engine/Plugins/Animation/RigLogic/Source/RigLogicLib/Private/riglogic/joints/bpcm/JointCalculationStrategy.h
new file mode 100644
--- /dev/null
+++ b/Engine/Plugins/Animation/RigLogic/Source/RigLogicLib/Private/riglogic/joints/bpcm/JointCalculationStrategy.h
@@ -0,0 +1,32 @@
+#pragma once
+
+#include "riglogic/joints/bpcm/JointStorageView.h"
+#include "riglogic/types/Aliases.h"
+
+#include <cstdint>
+
+namespace rl4 {
+
+namespace bpcm {
+
+template<typename TValue>
+struct JointCalculationStrategy {
+    using JointGroupArrayView = ConstArrayView<JointGroupView<TValue> >;
+
+    virtual ~JointCalculationStrategy() = default;
+
+    virtual void calculate(const JointGroupArrayView& jointGroups,
+                           ConstArrayView<float> inputs,
+                           ArrayView<float> outputs,
+                           std::uint16_t lod) const = 0;
+
+    virtual void calculate(const JointGroupArrayView& jointGroups,
+                           ConstArrayView<float> inputs,
+                           ArrayView<float> outputs,
+                           std::uint16_t lod,
+                           std::uint16_t jointGroupIndex) const = 0;
+};
+
+}  // namespace bpcm
+
+}  // namespace rl4
diff --git a/Engine/Plugins/Animation/RigLogic/Source/RigLogicLib/Private/riglogic/joints/bpcm/JointStorage.h b/Engine/Plugins/Animation/RigLogic/Source/RigLogicLib/Private/riglogic/joints/bpcm/JointStorage.h
new file mode 100644
--- /dev/null
+++ b/Engine/Plugins/Animation/RigLogic/Source/RigLogicLib/Private/riglogic/joints/bpcm/JointStorage.h
@@ -0,0 +1,43 @@
+#pragma once
+
+#include "riglogic/TypeDefs.h"
+#include "riglogic/joints/bpcm/JointGroup.h"
+#include "riglogic/joints/bpcm/LODRegion.h"
+
+#include <cstdint>
+
+namespace rl4 {
+
+namespace bpcm {
+
+template<typename TValue>
+struct JointStorage {
+    // All non-zero values
+    AlignedVector<TValue> values;
+    // Sub-matrix col -> input vector
+    AlignedVector<std::uint16_t> inputIndices;
+    // Sub-matrix row -> output vector
+    AlignedVector<std::uint16_t> outputIndices;
+    // Output index boundaries for each LOD
+    Vector<LODRegion> lodRegions;
+    // Delineate storage into joint-groups
+    Vector<JointGroup> jointGroups;
+
+    explicit JointStorage(MemoryResource* memRes) :
+        values{memRes},
+        inputIndices{memRes},
+        outputIndices{memRes},
+        lodRegions{memRes},
+        jointGroups{memRes} {
+    }
+
+    template<class Archive>
+    void serialize(Archive& archive) {
+        archive(values, inputIndices, outputIndices, lodRegions, jointGroups);
+    }
+
+};
+
+}  // namespace bpcm
+
+}  // namespace rl4
diff --git a/Engine/Plugins/Animation/RigLogic/Source/RigLogicLib/Private/riglogic/joints/bpcm/JointStorageView.h b/Engine/Plugins/Animation/RigLogic/Source/RigLogicLib/Private/riglogic/joints/bpcm/JointStorageView.h
new file mode 100644
--- /dev/null
+++ b/Engine/Plugins/Animation/RigLogic/Source/RigLogicLib/Private/riglogic/joints/bpcm/JointStorageView.h
@@ -0,0 +1,41 @@
+#pragma once
+
+#include "riglogic/TypeDefs.h"
+#include "riglogic/joints/bpcm/JointStorage.h"
+
+#include <cstdint>
+
+namespace rl4 {
+
+namespace bpcm {
+
+template<typename TValue>
+struct JointGroupView {
+    TValue* values;
+    std::uint16_t* inputIndices;
+    std::uint16_t* inputIndicesEnd;
+    std::uint16_t* inputIndicesEndAlignedTo4;
+    std::uint16_t* inputIndicesEndAlignedTo8;
+    std::uint16_t* outputIndices;
+    LODRegion* lods;
+};
+
+template<typename TValue>
+Vector<JointGroupView<TValue> > takeStorageSnapshot(JointStorage<TValue>& storage, MemoryResource* memRes) {
+    Vector<JointGroupView<TValue> > snapshot{storage.jointGroups.size(), {}, memRes};
+    for (std::size_t i = 0ul; i < storage.jointGroups.size(); ++i) {
+        const auto& jointGroup = storage.jointGroups[i];
+        snapshot[i].values = storage.values.data() + jointGroup.valuesOffset;
+        snapshot[i].inputIndices = storage.inputIndices.data() + jointGroup.inputIndicesOffset;
+        snapshot[i].inputIndicesEnd = snapshot[i].inputIndices + jointGroup.inputIndicesSize;
+        snapshot[i].inputIndicesEndAlignedTo4 = snapshot[i].inputIndices + jointGroup.inputIndicesSizeAlignedTo4;
+        snapshot[i].inputIndicesEndAlignedTo8 = snapshot[i].inputIndices + jointGroup.inputIndicesSizeAlignedTo8;
+        snapshot[i].outputIndices = storage.outputIndices.data() + jointGroup.outputIndicesOffset;
+        snapshot[i].lods = storage.lodRegions.data() + jointGroup.lodsOffset;
+    }
+    return snapshot;
+}
+
+}  // namespace bpcm
+
+}  // namespace rl4
diff --git a/Engine/Plugins/Animation/RigLogic/Source/RigLogicLib/Private/riglogic/joints/bpcm/JointsBuilderCommon.cpp b/Engine/Plugins/Animation/RigLogic/Source/RigLogicLib/Private/riglogic/joints/bpcm/JointsBuilderCommon.cpp
new file mode 100644
--- /dev/null
+++ b/Engine/Plugins/Animation/RigLogic/Source/RigLogicLib/Private/riglogic/joints/bpcm/JointsBuilderCommon.cpp
@@ -0,0 +1,107 @@
+#include "riglogic/joints/bpcm/JointsBuilderCommon.h"
+
+#include "riglogic/joints/bpcm/JointsEvaluator.h"
+#include "riglogic/joints/bpcm/StorageSize.h"
+#include "riglogic/utils/Extd.h"
+
+#include <dna/layers/BehaviorReader.h>
+#include <pma/utils/ManagedInstance.h>
+
+#ifdef _MSC_VER
+    #pragma warning(push)
+    #pragma warning(disable : 4365 4987)
+#endif
+#include <algorithm>
+#include <cstddef>
+#include <cstdint>
+#include <utility>
+#ifdef _MSC_VER
+    #pragma warning(pop)
+#endif
+
+namespace rl4 {
+
+namespace bpcm {
+
+template<typename TValue>
+JointsBuilderCommon<TValue>::JointsBuilderCommon(std::uint32_t blockHeight_, std::uint32_t padTo_, MemoryResource* memRes_) :
+    blockHeight{blockHeight_},
+    padTo{padTo_},
+    memRes{memRes_},
+    sizeReqs{memRes},
+    storage{memRes} {
+}
+
+template<typename TValue>
+void JointsBuilderCommon<TValue>::allocateStorage(const dna::BehaviorReader* reader) {
+    sizeReqs.computeFrom(reader, padTo);
+    storage.values.resize(sizeReqs.valueCount);
+    storage.inputIndices.resize(sizeReqs.inputIndexCount);
+    storage.outputIndices.resize(sizeReqs.outputIndexCount);
+    storage.lodRegions.reserve(sizeReqs.lodRegionCount);
+    storage.jointGroups.resize(sizeReqs.jointGroups.size());
+}
+
+template<typename TValue>
+void JointsBuilderCommon<TValue>::fillStorage(const dna::BehaviorReader* reader) {
+    setValues(reader);
+    setInputIndices(reader);
+    setOutputIndices(reader);
+    setLODs(reader);
+}
+
+template<typename TValue>
+void JointsBuilderCommon<TValue>::setInputIndices(const dna::BehaviorReader* reader) {
+    std::uint32_t offset = 0ul;
+    for (std::uint16_t i = 0u; i < reader->getJointGroupCount(); ++i) {
+        const auto jointGroupSize = sizeReqs.getJointGroupSize(i);
+        const auto indices = reader->getJointGroupInputIndices(i);
+        std::copy(indices.begin(), indices.end(), extd::advanced(storage.inputIndices.begin(), offset));
+        storage.jointGroups[i].inputIndicesOffset = offset;
+        storage.jointGroups[i].inputIndicesSize = jointGroupSize.padded.cols;
+        storage.jointGroups[i].inputIndicesSizeAlignedTo4 = jointGroupSize.padded.cols - (jointGroupSize.padded.cols % 4u);
+        storage.jointGroups[i].inputIndicesSizeAlignedTo8 = jointGroupSize.padded.cols - (jointGroupSize.padded.cols % 8u);
+        offset += jointGroupSize.padded.cols;
+    }
+}
+
+template<typename TValue>
+void JointsBuilderCommon<TValue>::setOutputIndices(const dna::BehaviorReader* reader) {
+    std::uint32_t offset = 0ul;
+    for (std::uint16_t i = 0u; i < reader->getJointGroupCount(); ++i) {
+        const auto jointGroupSize = sizeReqs.getJointGroupSize(i);
+        const auto indices = reader->getJointGroupOutputIndices(i);
+        std::copy(indices.begin(), indices.end(), extd::advanced(storage.outputIndices.begin(), offset));
+        storage.jointGroups[i].outputIndicesOffset = offset;
+        offset += jointGroupSize.padded.rows;
+    }
+}
+
+template<typename TValue>
+void JointsBuilderCommon<TValue>::setLODs(const dna::BehaviorReader* reader) {
+    std::uint32_t offset = 0ul;
+    for (std::uint16_t i = 0u; i < reader->getJointGroupCount(); ++i) {
+        const auto jointGroupSize = sizeReqs.getJointGroupSize(i);
+        const auto dstRowCount = jointGroupSize.padded.rows;
+        auto makeLODRegion = [dstRowCount, this](std::uint16_t lodRowCount) {
+                return LODRegion{lodRowCount, dstRowCount, blockHeight, padTo};
+            };
+        const auto lods = reader->getJointGroupLODs(i);
+        std::transform(lods.begin(), lods.end(), std::back_inserter(storage.lodRegions), makeLODRegion);
+        storage.jointGroups[i].lodsOffset = offset;
+        offset += reader->getLODCount();
+    }
+}
+
+template<typename TValue>
+typename JointsBuilderCommon<TValue>::JointsEvaluatorPtr JointsBuilderCommon<TValue>::build() {
+    using ManagedJoints = pma::UniqueInstance<Evaluator<TValue>, JointsEvaluator>;
+    return ManagedJoints::with(memRes).create(std::move(storage), std::move(strategy), memRes);
+}
+
+template class JointsBuilderCommon<float>;
+template class JointsBuilderCommon<std::uint16_t>;
+
+}  // namespace bpcm
+
+}  // namespace rl4
diff --git a/Engine/Plugins/Animation/RigLogic/Source/RigLogicLib/Private/riglogic/joints/bpcm/JointsBuilderCommon.h b/Engine/Plugins/Animation/RigLogic/Source/RigLogicLib/Private/riglogic/joints/bpcm/JointsBuilderCommon.h
new file mode 100644
--- /dev/null
+++ b/Engine/Plugins/Animation/RigLogic/Source/RigLogicLib/Private/riglogic/joints/bpcm/JointsBuilderCommon.h
@@ -0,0 +1,59 @@
+#pragma once
+
+#include "riglogic/TypeDefs.h"
+#include "riglogic/joints/JointsBuilder.h"
+#include "riglogic/joints/bpcm/JointStorage.h"
+#include "riglogic/joints/bpcm/StorageSize.h"
+#include "riglogic/joints/bpcm/JointCalculationStrategy.h"
+#include "riglogic/types/Aliases.h"
+
+#ifdef _MSC_VER
+    #pragma warning(push)
+    #pragma warning(disable : 4365 4987)
+#endif
+#include <cstddef>
+#include <cstdint>
+#include <functional>
+#include <memory>
+#ifdef _MSC_VER
+    #pragma warning(pop)
+#endif
+
+namespace rl4 {
+
+namespace bpcm {
+
+template<typename TValue>
+class JointsBuilderCommon : public JointsBuilder {
+    protected:
+        using CalculationStrategy = JointCalculationStrategy<TValue>;
+        using CalculationStrategyPtr = std::unique_ptr<CalculationStrategy, std::function<void (CalculationStrategy*)> >;
+
+    protected:
+        JointsBuilderCommon(std::uint32_t blockHeight_, std::uint32_t padTo_, MemoryResource* memRes_);
+
+    public:
+        void allocateStorage(const dna::BehaviorReader* reader) override;
+        void fillStorage(const dna::BehaviorReader* reader) override;
+        JointsEvaluatorPtr build() override;
+
+    protected:
+        virtual void setValues(const dna::BehaviorReader* source) = 0;
+        virtual void setInputIndices(const dna::BehaviorReader* reader);
+        virtual void setOutputIndices(const dna::BehaviorReader* reader);
+        virtual void setLODs(const dna::BehaviorReader* reader);
+
+    protected:
+        std::uint32_t blockHeight;
+        std::uint32_t padTo;
+        MemoryResource* memRes;
+
+        StorageSize sizeReqs;
+
+        JointStorage<TValue> storage;
+        CalculationStrategyPtr strategy;
+};
+
+}  // namespace bpcm
+
+}  // namespace rl4
diff --git a/Engine/Plugins/Animation/RigLogic/Source/RigLogicLib/Private/riglogic/joints/bpcm/JointsEvaluator.h b/Engine/Plugins/Animation/RigLogic/Source/RigLogicLib/Private/riglogic/joints/bpcm/JointsEvaluator.h
new file mode 100644
--- /dev/null
+++ b/Engine/Plugins/Animation/RigLogic/Source/RigLogicLib/Private/riglogic/joints/bpcm/JointsEvaluator.h
@@ -0,0 +1,81 @@
+#pragma once
+
+#include "riglogic/TypeDefs.h"
+#include "riglogic/joints/JointsEvaluator.h"
+#include "riglogic/joints/bpcm/JointCalculationStrategy.h"
+#include "riglogic/joints/bpcm/JointStorage.h"
+#include "riglogic/joints/bpcm/JointStorageView.h"
+#include "riglogic/riglogic/RigInstanceImpl.h"
+
+#include <terse/archives/binary/InputArchive.h>
+#include <terse/archives/binary/OutputArchive.h>
+#include <trio/Stream.h>
+
+#ifdef _MSC_VER
+    #pragma warning(push)
+    #pragma warning(disable : 4365 4987)
+#endif
+#include <cstdint>
+#include <functional>
+#include <memory>
+#ifdef _MSC_VER
+    #pragma warning(pop)
+#endif
+
+namespace rl4 {
+
+namespace bpcm {
+
+template<typename TValue>
+class Evaluator : public JointsEvaluator {
+    public:
+        using CalculationStrategy = JointCalculationStrategy<TValue>;
+        using CalculationStrategyPtr = std::unique_ptr<CalculationStrategy, std::function<void (CalculationStrategy*)> >;
+
+        struct Accessor;
+        friend Accessor;
+
+    private:
+        using JointGroupArrayView = typename CalculationStrategy::JointGroupArrayView;
+
+    public:
+        Evaluator(JointStorage<TValue>&& storage_, CalculationStrategyPtr strategy_, MemoryResource* memRes_) :
+            memRes{memRes_},
+            storage{std::move(storage_)},
+            storageView{takeStorageSnapshot(storage, memRes)},
+            jointGroupsView{storageView},
+            strategy{std::move(strategy_)} {
+        }
+
+        void calculate(ConstArrayView<float> inputs, ArrayView<float> outputs, std::uint16_t lod) const override {
+            assert(strategy != nullptr);
+            strategy->calculate(jointGroupsView, inputs, outputs, lod);
+        }
+
+        void calculate(ConstArrayView<float> inputs, ArrayView<float> outputs, std::uint16_t lod,
+                       std::uint16_t jointGroupIndex) const override {
+            assert(strategy != nullptr);
+            strategy->calculate(jointGroupsView, inputs, outputs, lod, jointGroupIndex);
+        }
+
+        void load(terse::BinaryInputArchive<BoundedIOStream>& archive) override {
+            archive(storage);
+            storageView = takeStorageSnapshot(storage, memRes);
+            jointGroupsView = JointGroupArrayView{storageView};
+        }
+
+        void save(terse::BinaryOutputArchive<BoundedIOStream>& archive) override {
+            archive(storage);
+        }
+
+    private:
+        MemoryResource* memRes;
+        JointStorage<TValue> storage;
+        Vector<JointGroupView<TValue> > storageView;
+        JointGroupArrayView jointGroupsView;
+        CalculationStrategyPtr strategy;
+};
+
+}  // namespace bpcm
+
+}  // namespace rl4
diff --git a/Engine/Plugins/Animation/RigLogic/Source/RigLogicLib/Private/riglogic/joints/bpcm/LODRegion.cpp b/Engine/Plugins/Animation/RigLogic/Source/RigLogicLib/Private/riglogic/joints/bpcm/LODRegion.cpp
new file mode 100644
--- /dev/null
+++ b/Engine/Plugins/Animation/RigLogic/Source/RigLogicLib/Private/riglogic/joints/bpcm/LODRegion.cpp
@@ -0,0 +1,51 @@
+#include "riglogic/joints/bpcm/LODRegion.h"
+
+#include "riglogic/utils/Extd.h"
+
+#include <cstddef>
+
+namespace rl4 {
+
+namespace bpcm {
+
+LODRegion::LODRegion(std::uint32_t lodRowCount, std::uint32_t rowCount, std::uint32_t blockHeight,
+                     std::uint32_t padTo) : size{lodRowCount} {
+    // Check whether the sub-matrix was padded to blockHeight size or padTo size
+    const bool endsWithPadToBlock = (rowCount % blockHeight) == padTo;
+    // Check if LOD boundary is within the region of the last padTo rows
+    const bool isLODInLastPadToRows = (rowCount - lodRowCount) < padTo;
+    // Check if LOD row is within the region of the last padTo-sized block
+    const bool isLODInLastPadToBlock = (endsWithPadToBlock && isLODInLastPadToRows);
+    // Round the LOD row to the nearest index that is modulo padTo or blockHeight
+    // If LOD is in the last padTo-sized block, it rounds to padTo, otherwise to blockHeight
+    const std::uint32_t alignTo = (isLODInLastPadToBlock ? padTo : blockHeight);
+    const std::uint32_t upAlignedLodRow = extd::roundUp(lodRowCount, alignTo);
+    // Check if the unaligned LOD boundary coincides with a block boundary
+    // (e.g. this is the case for LOD-0)
+    const bool isLODOnBlockBoundary = (upAlignedLodRow == lodRowCount);
+    // Should the last blockHeight-sized block be masked-off or not
+    // If LOD coincides with a block boundary, no mask-off is needed at all
+    // If LOD is within the last padTo-sized block, no mask-off is needed for the last blockHeight-sized block
+    const bool maskOffLastBlock = (!isLODOnBlockBoundary && !isLODInLastPadToBlock);
+    // Find the last blockHeight-sized block boundary before the aligned LOD boundary
+    // Sometimes this equals the aligned LOD boundary, but not if the matrix is padded to padTo block-size
+    // and the LOD boundary is within that last padTo-sized block region
+    sizeAlignedToLastFullBlock = upAlignedLodRow - (upAlignedLodRow % blockHeight);
+    // This is one blockHeight before `sizeAlignedToLastFullBlock` because the last blockHeight-sized block
+    // might need special care to mask off results that came from rows after the LOD boundary
+    // If no mask-off is needed, it equals to `sizeAlignedToLastFullBlock`
+    sizeAlignedToSecondLastFullBlock = (maskOffLastBlock ? sizeAlignedToLastFullBlock - blockHeight
+                                        : sizeAlignedToLastFullBlock);
+}
+
+LODRegion::LODRegion(std::uint32_t lodRowCount,
+                     std::uint32_t lodRowCountAlignedToLastFullBlock,
+                     std::uint32_t lodRowCountAlignedToSecondLastFullBlock) :
+    size{lodRowCount},
+    sizeAlignedToLastFullBlock{lodRowCountAlignedToLastFullBlock},
+    sizeAlignedToSecondLastFullBlock{lodRowCountAlignedToSecondLastFullBlock} {
+}
+
+}  // namespace bpcm
+
+}  // namespace rl4
diff --git a/Engine/Plugins/Animation/RigLogic/Source/RigLogicLib/Private/riglogic/joints/bpcm/LODRegion.h b/Engine/Plugins/Animation/RigLogic/Source/RigLogicLib/Private/riglogic/joints/bpcm/LODRegion.h
new file mode 100644
--- /dev/null
+++ b/Engine/Plugins/Animation/RigLogic/Source/RigLogicLib/Private/riglogic/joints/bpcm/LODRegion.h
@@ -0,0 +1,34 @@
+#pragma once
+
+#include <cstdint>
+
+namespace rl4 {
+
+namespace bpcm {
+
+struct LODRegion {
+    // Unaligned number of output indices
+    std::uint32_t size;
+    // If the last full block needs to be masked-off because of LODs,
+    // this view will include that last full block as well
+    std::uint32_t sizeAlignedToLastFullBlock;
+    // If the last full block needs to be masked-off because of LODs,
+    // that full block won't be included in this boundary as it needs special care
+    std::uint32_t sizeAlignedToSecondLastFullBlock;
+
+    LODRegion() = default;
+    LODRegion(std::uint32_t lodRowCount, std::uint32_t rowCount, std::uint32_t blockHeight, std::uint32_t padTo);
+    LODRegion(std::uint32_t lodRowCount,
+              std::uint32_t lodRowCountAlignedToLastFullBlock,
+              std::uint32_t lodRowCountAlignedToSecondLastFullBlock);
+
+    template<class Archive>
+    void serialize(Archive& archive) {
+        archive(size, sizeAlignedToLastFullBlock, sizeAlignedToSecondLastFullBlock);
+    }
+
+};
+
+}  // namespace bpcm
+
+}  // namespace rl4
