Editgrating RigLogic plugin from Dev-CharacterTech to UE5/Main @ CL 13891876.

diff --git a/Engine/Plugins/Animation/RigLogic/Source/RigLogicLib/Private/riglogic/joints/bpcm/StorageSize.cpp b/Engine/Plugins/Animation/RigLogic/Source/RigLogicLib/Private/riglogic/joints/bpcm/StorageSize.cpp
new file mode 100644
--- /dev/null
+++ b/Engine/Plugins/Animation/RigLogic/Source/RigLogicLib/Private/riglogic/joints/bpcm/StorageSize.cpp
@@ -0,0 +1,52 @@
+#include "riglogic/joints/bpcm/StorageSize.h"
+
+#include "riglogic/joints/bpcm/Extent.h"
+#include "riglogic/utils/Extd.h"
+
+#include <dna/layers/BehaviorReader.h>
+
+#include <cassert>
+#include <cstddef>
+
+namespace rl4 {
+
+namespace bpcm {
+
+StorageSize::StorageSize(MemoryResource* memRes) :
+    valueCount{},
+    inputIndexCount{},
+    outputIndexCount{},
+    lodRegionCount{},
+    lodCount{},
+    jointGroups{memRes} {
+}
+
+void StorageSize::computeFrom(const dna::BehaviorReader* src, std::size_t padTo) {
+    const auto jointGroupCount = src->getJointGroupCount();
+    jointGroups.clear();
+    jointGroups.reserve(jointGroupCount);
+    lodCount = src->getLODCount();
+    for (std::uint16_t i = 0u; i < jointGroupCount; ++i) {
+        assert(src->getJointGroupLODs(i).size() == lodCount);
+        const auto columnCount = src->getJointGroupInputIndices(i).size();
+        const auto rowCount = src->getJointGroupOutputIndices(i).size();
+        const std::size_t padding = extd::roundUp(rowCount, padTo) - rowCount;
+        valueCount += (rowCount * columnCount) + (padding * columnCount);
+        inputIndexCount += columnCount;
+        outputIndexCount += (rowCount + padding);
+        lodRegionCount += lodCount;
+        jointGroups.push_back({
+                    {static_cast<std::uint16_t>(rowCount), static_cast<std::uint16_t>(columnCount)},
+                    {static_cast<std::uint16_t>(rowCount + padding), static_cast<std::uint16_t>(columnCount)}
+                });
+    }
+}
+
+JointGroupSize StorageSize::getJointGroupSize(std::size_t jointGroupIndex) const {
+    assert(jointGroupIndex < jointGroups.size());
+    return jointGroups[jointGroupIndex];
+}
+
+}  // namespace bpcm
+
+}  // namespace rl4
diff --git a/Engine/Plugins/Animation/RigLogic/Source/RigLogicLib/Private/riglogic/joints/bpcm/StorageSize.h b/Engine/Plugins/Animation/RigLogic/Source/RigLogicLib/Private/riglogic/joints/bpcm/StorageSize.h
new file mode 100644
--- /dev/null
+++ b/Engine/Plugins/Animation/RigLogic/Source/RigLogicLib/Private/riglogic/joints/bpcm/StorageSize.h
@@ -0,0 +1,33 @@
+#pragma once
+
+#include "riglogic/TypeDefs.h"
+#include "riglogic/joints/bpcm/Extent.h"
+
+#include <cstddef>
+
+namespace rl4 {
+
+namespace bpcm {
+
+struct JointGroupSize {
+    Extent original;
+    Extent padded;
+};
+
+struct StorageSize {
+    std::size_t valueCount;
+    std::size_t inputIndexCount;
+    std::size_t outputIndexCount;
+    std::size_t lodRegionCount;
+    std::size_t lodCount;
+    Vector<JointGroupSize> jointGroups;
+
+    explicit StorageSize(MemoryResource* memRes);
+
+    void computeFrom(const dna::BehaviorReader* src, std::size_t padTo);
+    JointGroupSize getJointGroupSize(std::size_t jointGroupIndex) const;
+};
+
+}  // namespace bpcm
+
+}  // namespace rl4
diff --git a/Engine/Plugins/Animation/RigLogic/Source/RigLogicLib/Private/riglogic/joints/bpcm/builders/AVX.cpp b/Engine/Plugins/Animation/RigLogic/Source/RigLogicLib/Private/riglogic/joints/bpcm/builders/AVX.cpp
new file mode 100644
--- /dev/null
+++ b/Engine/Plugins/Animation/RigLogic/Source/RigLogicLib/Private/riglogic/joints/bpcm/builders/AVX.cpp
@@ -0,0 +1,27 @@
+#if !defined(RL_USE_HALF_FLOATS) && defined(RL_BUILD_WITH_AVX)
+
+#include "riglogic/joints/bpcm/builders/AVX.h"
+
+#include "riglogic/joints/bpcm/Consts.h"
+#include "riglogic/joints/bpcm/strategies/AVX.h"
+
+#include <pma/utils/ManagedInstance.h>
+
+#include <cstddef>
+#include <cstdint>
+
+namespace rl4 {
+
+namespace bpcm {
+
+AVXJointsBuilder::AVXJointsBuilder(MemoryResource* memRes_) : FloatStorageBuilder{block16Height, block8Height, memRes_} {
+    strategy = pma::UniqueInstance<AVXJointCalculationStrategy<float>, CalculationStrategy>::with(memRes_).create();
+}
+
+AVXJointsBuilder::~AVXJointsBuilder() = default;
+
+}  // namespace bpcm
+
+}  // namespace rl4
+
+#endif  // !defined(RL_USE_HALF_FLOATS) && defined(RL_BUILD_WITH_AVX)
diff --git a/Engine/Plugins/Animation/RigLogic/Source/RigLogicLib/Private/riglogic/joints/bpcm/builders/AVX.h b/Engine/Plugins/Animation/RigLogic/Source/RigLogicLib/Private/riglogic/joints/bpcm/builders/AVX.h
new file mode 100644
--- /dev/null
+++ b/Engine/Plugins/Animation/RigLogic/Source/RigLogicLib/Private/riglogic/joints/bpcm/builders/AVX.h
@@ -0,0 +1,25 @@
+#pragma once
+
+// *INDENT-OFF*
+#if !defined(RL_USE_HALF_FLOATS) && defined(RL_BUILD_WITH_AVX)
+
+#include "riglogic/joints/bpcm/builders/Float.h"
+#include "riglogic/types/Aliases.h"
+
+namespace rl4 {
+
+namespace bpcm {
+
+class AVXJointsBuilder : public FloatStorageBuilder {
+    public:
+        explicit AVXJointsBuilder(MemoryResource* memRes_);
+        ~AVXJointsBuilder();
+
+};
+
+}  // namespace bpcm
+
+}  // namespace rl4
+
+#endif  // !defined(RL_USE_HALF_FLOATS) && defined(RL_BUILD_WITH_AVX)
+// *INDENT-ON*
diff --git a/Engine/Plugins/Animation/RigLogic/Source/RigLogicLib/Private/riglogic/joints/bpcm/builders/AVXHalfFloat.cpp b/Engine/Plugins/Animation/RigLogic/Source/RigLogicLib/Private/riglogic/joints/bpcm/builders/AVXHalfFloat.cpp
new file mode 100644
--- /dev/null
+++ b/Engine/Plugins/Animation/RigLogic/Source/RigLogicLib/Private/riglogic/joints/bpcm/builders/AVXHalfFloat.cpp
@@ -0,0 +1,98 @@
+#if defined(RL_USE_HALF_FLOATS) && defined(RL_BUILD_WITH_AVX)
+
+#include "riglogic/joints/bpcm/builders/AVXHalfFloat.h"
+
+#include "riglogic/joints/bpcm/Consts.h"
+#include "riglogic/joints/bpcm/strategies/AVX.h"
+#include "riglogic/utils/Extd.h"
+
+#include <dna/layers/BehaviorReader.h>
+#include <pma/utils/ManagedInstance.h>
+
+#include <immintrin.h>
+
+#include <cstddef>
+#include <cstdint>
+
+namespace rl4 {
+
+namespace bpcm {
+
+inline void toHalfFloat256(const float* source, std::uint16_t* destination) {
+    #if defined(__clang__) || defined(__GNUC__)
+        #pragma GCC diagnostic push
+        #pragma GCC diagnostic ignored "-Wold-style-cast"
+        #if !defined(__clang__)
+            #pragma GCC diagnostic ignored "-Wuseless-cast"
+        #endif
+    #endif
+    __m256 block = _mm256_load_ps(source);
+    __m128i half = _mm256_cvtps_ph(block, _MM_FROUND_CUR_DIRECTION);
+    _mm_store_si128(reinterpret_cast<__m128i*>(destination), half);
+    #if defined(__clang__) || defined(__GNUC__)
+        #pragma GCC diagnostic pop
+    #endif
+}
+
+AVXJointsBuilder::AVXJointsBuilder(MemoryResource* memRes_) : JointsBuilderCommon{block16Height, block8Height, memRes_} {
+    strategy = pma::UniqueInstance<AVXJointCalculationStrategy<std::uint16_t>, CalculationStrategy>::with(memRes_).create();
+}
+
+AVXJointsBuilder::~AVXJointsBuilder() = default;
+
+void AVXJointsBuilder::setValues(const dna::BehaviorReader* reader) {
+    std::uint32_t offset = 0ul;
+    for (std::uint16_t i = 0u; i < reader->getJointGroupCount(); ++i) {
+        const auto source = reader->getJointGroupValues(i);
+        const auto jointGroupSize = sizeReqs.getJointGroupSize(i);
+        storage.jointGroups[i].valuesOffset = offset;
+        storage.jointGroups[i].valuesSize = jointGroupSize.padded.size();
+
+        std::uint32_t remainder = jointGroupSize.original.rows % block16Height;
+        std::uint32_t target = jointGroupSize.original.rows - remainder;
+        // Copy the portion of the source matrix that is partitionable into block16Height chunks
+        for (std::size_t row = 0ul; row < target; row += block16Height) {
+            for (std::size_t col = 0ul; col < jointGroupSize.original.cols; ++col) {
+                alignas(32) float fBlock[block16Height];
+                alignas(16) std::uint16_t hfBlock[block16Height];
+
+                for (std::size_t blkIdx = 0ul; blkIdx < block16Height; ++blkIdx) {
+                    fBlock[blkIdx] = source[(row + blkIdx) * jointGroupSize.original.cols + col];
+                }
+
+                toHalfFloat256(fBlock, hfBlock);
+                toHalfFloat256(fBlock + block8Height, hfBlock + block8Height);
+
+                for (std::size_t blkIdx = 0ul; blkIdx < block16Height; ++blkIdx, ++offset) {
+                    storage.values[offset] = hfBlock[blkIdx];
+                }
+            }
+        }
+        // Copy the non-blockHeight sized remainder portion of the source matrix
+        for (std::size_t row = target; row < jointGroupSize.original.rows; row += remainder) {
+            for (std::size_t col = 0ul; col < jointGroupSize.original.cols; ++col) {
+                alignas(32) float fBlock[block16Height] = {};
+                alignas(16) std::uint16_t hfBlock[block16Height] = {};
+
+                for (std::size_t blkIdx = 0ul; blkIdx < remainder; ++blkIdx) {
+                    fBlock[blkIdx] = source[(row + blkIdx) * jointGroupSize.original.cols + col];
+                }
+
+                toHalfFloat256(fBlock, hfBlock);
+                toHalfFloat256(fBlock + block8Height, hfBlock + block8Height);
+
+                // Copy even the padded zero values as the representation of half-floats might
+                // differ from the storage's zero-initialized state
+                for (std::size_t blkIdx = 0ul; blkIdx < extd::roundUp(remainder, block8Height); ++blkIdx, ++offset) {
+                    storage.values[offset] = hfBlock[blkIdx];
+                }
+            }
+        }
+    }
+}
+
+}  // namespace bpcm
+
+}  // namespace rl4
+
+#endif  // defined(RL_USE_HALF_FLOATS) && defined(RL_BUILD_WITH_AVX)
diff --git a/Engine/Plugins/Animation/RigLogic/Source/RigLogicLib/Private/riglogic/joints/bpcm/builders/AVXHalfFloat.h b/Engine/Plugins/Animation/RigLogic/Source/RigLogicLib/Private/riglogic/joints/bpcm/builders/AVXHalfFloat.h
new file mode 100644
--- /dev/null
+++ b/Engine/Plugins/Animation/RigLogic/Source/RigLogicLib/Private/riglogic/joints/bpcm/builders/AVXHalfFloat.h
@@ -0,0 +1,30 @@
+#pragma once
+
+// *INDENT-OFF*
+#if defined(RL_USE_HALF_FLOATS) && defined(RL_BUILD_WITH_AVX)
+
+#include "riglogic/joints/bpcm/JointsBuilderCommon.h"
+#include "riglogic/types/Aliases.h"
+
+#include <cstdint>
+
+namespace rl4 {
+
+namespace bpcm {
+
+class AVXJointsBuilder : public JointsBuilderCommon<std::uint16_t> {
+    public:
+        explicit AVXJointsBuilder(MemoryResource* memRes_);
+        ~AVXJointsBuilder();
+
+    protected:
+        void setValues(const dna::BehaviorReader* reader) override;
+
+};
+
+}  // namespace bpcm
+
+}  // namespace rl4
+
+#endif  // defined(RL_USE_HALF_FLOATS) && defined(RL_BUILD_WITH_AVX)
+// *INDENT-ON*
diff --git a/Engine/Plugins/Animation/RigLogic/Source/RigLogicLib/Private/riglogic/joints/bpcm/builders/Float.cpp b/Engine/Plugins/Animation/RigLogic/Source/RigLogicLib/Private/riglogic/joints/bpcm/builders/Float.cpp
new file mode 100644
--- /dev/null
+++ b/Engine/Plugins/Animation/RigLogic/Source/RigLogicLib/Private/riglogic/joints/bpcm/builders/Float.cpp
@@ -0,0 +1,52 @@
+#include "riglogic/joints/bpcm/builders/Float.h"
+
+#include "riglogic/joints/bpcm/Consts.h"
+
+#include <dna/layers/BehaviorReader.h>
+
+#include <cstddef>
+#include <cstdint>
+
+namespace rl4 {
+
+namespace bpcm {
+
+FloatStorageBuilder::FloatStorageBuilder(std::uint32_t blockHeight_, std::uint32_t padTo_,
+                                         MemoryResource* memRes_) : JointsBuilderCommon{blockHeight_, padTo_, memRes_} {
+}
+
+FloatStorageBuilder::~FloatStorageBuilder() = default;
+
+void FloatStorageBuilder::setValues(const dna::BehaviorReader* reader) {
+    std::uint32_t offset = 0ul;
+    for (std::uint16_t i = 0u; i < reader->getJointGroupCount(); ++i) {
+        const auto source = reader->getJointGroupValues(i);
+        const auto jointGroupSize = sizeReqs.getJointGroupSize(i);
+        storage.jointGroups[i].valuesOffset = offset;
+        storage.jointGroups[i].valuesSize = jointGroupSize.padded.size();
+        const std::uint32_t remainder = jointGroupSize.original.rows % blockHeight;
+        const std::uint32_t target = jointGroupSize.original.rows - remainder;
+        // Copy the portion of the source matrix that is partitionable into blockHeight chunks
+        for (std::size_t row = 0ul; row < target; row += blockHeight) {
+            for (std::size_t col = 0ul; col < jointGroupSize.original.cols; ++col) {
+                for (std::size_t blkIdx = 0ul; blkIdx < blockHeight; ++blkIdx, ++offset) {
+                    storage.values[offset] = source[(row + blkIdx) * jointGroupSize.original.cols + col];
+                }
+            }
+        }
+        // Copy the non-blockHeight sized remainder portion of the source matrix
+        for (std::size_t row = target; row < jointGroupSize.original.rows; row += remainder) {
+            for (std::size_t col = 0ul; col < jointGroupSize.original.cols; ++col) {
+                for (std::size_t blkIdx = 0ul; blkIdx < remainder; ++blkIdx, ++offset) {
+                    storage.values[offset] = source[(row + blkIdx) * jointGroupSize.original.cols + col];
+                }
+                // Skip over padding (already filled with zeros) in destination storage
+                offset += static_cast<std::uint32_t>(jointGroupSize.padded.rows - jointGroupSize.original.rows);
+            }
+        }
+    }
+}
+
+}  // namespace bpcm
+
+}  // namespace rl4
diff --git a/Engine/Plugins/Animation/RigLogic/Source/RigLogicLib/Private/riglogic/joints/bpcm/builders/Float.h b/Engine/Plugins/Animation/RigLogic/Source/RigLogicLib/Private/riglogic/joints/bpcm/builders/Float.h
new file mode 100644
--- /dev/null
+++ b/Engine/Plugins/Animation/RigLogic/Source/RigLogicLib/Private/riglogic/joints/bpcm/builders/Float.h
@@ -0,0 +1,24 @@
+#pragma once
+
+#include "riglogic/joints/bpcm/JointsBuilderCommon.h"
+#include "riglogic/types/Aliases.h"
+
+#include <cstddef>
+
+namespace rl4 {
+
+namespace bpcm {
+
+class FloatStorageBuilder : public JointsBuilderCommon<float> {
+    public:
+        FloatStorageBuilder(std::uint32_t blockHeight_, std::uint32_t padTo_, MemoryResource* memRes_);
+        ~FloatStorageBuilder();
+
+    protected:
+        void setValues(const dna::BehaviorReader* reader) override;
+
+};
+
+}  // namespace bpcm
+
+}  // namespace rl4
diff --git a/Engine/Plugins/Animation/RigLogic/Source/RigLogicLib/Private/riglogic/joints/bpcm/builders/SSE.cpp b/Engine/Plugins/Animation/RigLogic/Source/RigLogicLib/Private/riglogic/joints/bpcm/builders/SSE.cpp
new file mode 100644
--- /dev/null
+++ b/Engine/Plugins/Animation/RigLogic/Source/RigLogicLib/Private/riglogic/joints/bpcm/builders/SSE.cpp
@@ -0,0 +1,27 @@
+#if !defined(RL_USE_HALF_FLOATS) && defined(RL_BUILD_WITH_SSE)
+
+#include "riglogic/joints/bpcm/builders/SSE.h"
+
+#include "riglogic/joints/bpcm/Consts.h"
+#include "riglogic/joints/bpcm/strategies/SSE.h"
+
+#include <pma/utils/ManagedInstance.h>
+
+#include <cstddef>
+#include <cstdint>
+
+namespace rl4 {
+
+namespace bpcm {
+
+SSEJointsBuilder::SSEJointsBuilder(MemoryResource* memRes_) : FloatStorageBuilder{block8Height, block4Height, memRes_} {
+    strategy = pma::UniqueInstance<SSEJointCalculationStrategy<float>, CalculationStrategy>::with(memRes).create();
+}
+
+SSEJointsBuilder::~SSEJointsBuilder() = default;
+
+}  // namespace bpcm
+
+}  // namespace rl4
+
+#endif  // !defined(RL_USE_HALF_FLOATS) && defined(RL_BUILD_WITH_SSE)
diff --git a/Engine/Plugins/Animation/RigLogic/Source/RigLogicLib/Private/riglogic/joints/bpcm/builders/SSE.h b/Engine/Plugins/Animation/RigLogic/Source/RigLogicLib/Private/riglogic/joints/bpcm/builders/SSE.h
new file mode 100644
--- /dev/null
+++ b/Engine/Plugins/Animation/RigLogic/Source/RigLogicLib/Private/riglogic/joints/bpcm/builders/SSE.h
@@ -0,0 +1,24 @@
+#pragma once
+
+// *INDENT-OFF*
+#if !defined(RL_USE_HALF_FLOATS) && defined(RL_BUILD_WITH_SSE)
+
+#include "riglogic/joints/bpcm/builders/Float.h"
+#include "riglogic/types/Aliases.h"
+
+namespace rl4 {
+
+namespace bpcm {
+
+class SSEJointsBuilder : public FloatStorageBuilder {
+    public:
+        explicit SSEJointsBuilder(MemoryResource* memRes_);
+        ~SSEJointsBuilder();
+};
+
+}  // namespace bpcm
+
+}  // namespace rl4
+
+#endif  // !defined(RL_USE_HALF_FLOATS) && defined(RL_BUILD_WITH_SSE)
+// *INDENT-ON*
diff --git a/Engine/Plugins/Animation/RigLogic/Source/RigLogicLib/Private/riglogic/joints/bpcm/builders/SSEHalfFloat.cpp b/Engine/Plugins/Animation/RigLogic/Source/RigLogicLib/Private/riglogic/joints/bpcm/builders/SSEHalfFloat.cpp
new file mode 100644
--- /dev/null
+++ b/Engine/Plugins/Animation/RigLogic/Source/RigLogicLib/Private/riglogic/joints/bpcm/builders/SSEHalfFloat.cpp
@@ -0,0 +1,98 @@
+#if defined(RL_USE_HALF_FLOATS) && defined(RL_BUILD_WITH_SSE)
+
+#include "riglogic/joints/bpcm/builders/SSEHalfFloat.h"
+
+#include "riglogic/joints/bpcm/Consts.h"
+#include "riglogic/joints/bpcm/strategies/SSE.h"
+#include "riglogic/utils/Extd.h"
+
+#include <dna/layers/BehaviorReader.h>
+#include <pma/utils/ManagedInstance.h>
+
+#include <immintrin.h>
+
+#include <cstddef>
+#include <cstdint>
+
+namespace rl4 {
+
+namespace bpcm {
+
+inline void toHalfFloat128(const float* source, std::uint16_t* destination) {
+    #if defined(__clang__) || defined(__GNUC__)
+        #pragma GCC diagnostic push
+        #pragma GCC diagnostic ignored "-Wold-style-cast"
+        #if !defined(__clang__)
+            #pragma GCC diagnostic ignored "-Wuseless-cast"
+        #endif
+    #endif
+    __m128 block = _mm_load_ps(source);
+    __m128i half = _mm_cvtps_ph(block, _MM_FROUND_CUR_DIRECTION);
+    _mm_storel_epi64(reinterpret_cast<__m128i*>(destination), half);
+    #if defined(__clang__) || defined(__GNUC__)
+        #pragma GCC diagnostic pop
+    #endif
+}
+
+SSEJointsBuilder::~SSEJointsBuilder() = default;
+
+SSEJointsBuilder::SSEJointsBuilder(MemoryResource* memRes_) : JointsBuilderCommon{block8Height, block4Height, memRes_} {
+    strategy = pma::UniqueInstance<SSEJointCalculationStrategy<std::uint16_t>, CalculationStrategy>::with(memRes_).create();
+}
+
+void SSEJointsBuilder::setValues(const dna::BehaviorReader* reader) {
+    std::uint32_t offset = 0ul;
+    for (std::uint16_t i = 0u; i < reader->getJointGroupCount(); ++i) {
+        const auto source = reader->getJointGroupValues(i);
+        const auto jointGroupSize = sizeReqs.jointGroups[i];
+        storage.jointGroups[i].valuesOffset = offset;
+        storage.jointGroups[i].valuesSize = jointGroupSize.padded.size();
+
+        std::uint32_t remainder = jointGroupSize.original.rows % block8Height;
+        std::uint32_t target = jointGroupSize.original.rows - remainder;
+        // Copy the portion of the source matrix that is partitionable into block8Height chunks
+        for (std::size_t row = 0ul; row < target; row += block8Height) {
+            for (std::size_t col = 0ul; col < jointGroupSize.original.cols; ++col) {
+                alignas(16) float fBlock[block8Height];
+                alignas(16) std::uint16_t hfBlock[block8Height];
+
+                for (std::size_t blkIdx = 0ul; blkIdx < block8Height; ++blkIdx) {
+                    fBlock[blkIdx] = source[(row + blkIdx) * jointGroupSize.original.cols + col];
+                }
+
+                toHalfFloat128(fBlock, hfBlock);
+                toHalfFloat128(fBlock + block4Height, hfBlock + block4Height);
+
+                for (std::size_t blkIdx = 0ul; blkIdx < block8Height; ++blkIdx, ++offset) {
+                    storage.values[offset] = hfBlock[blkIdx];
+                }
+            }
+        }
+        // Copy the non-block8Height sized remainder portion of the source matrix
+        for (std::size_t row = target; row < jointGroupSize.original.rows; row += remainder) {
+            for (std::size_t col = 0ul; col < jointGroupSize.original.cols; ++col) {
+                alignas(16) float fBlock[block8Height] = {};
+                alignas(16) std::uint16_t hfBlock[block8Height] = {};
+
+                for (std::size_t blkIdx = 0ul; blkIdx < remainder; ++blkIdx) {
+                    fBlock[blkIdx] = source[(row + blkIdx) * jointGroupSize.original.cols + col];
+                }
+
+                toHalfFloat128(fBlock, hfBlock);
+                toHalfFloat128(fBlock + block4Height, hfBlock + block4Height);
+
+                // Copy even the padded zero values as the representation of half-floats might
+                // differ from the storage's zero-initialized state
+                for (std::size_t blkIdx = 0ul; blkIdx < extd::roundUp(remainder, block4Height); ++blkIdx, ++offset) {
+                    storage.values[offset] = hfBlock[blkIdx];
+                }
+            }
+        }
+    }
+}
+
+}  // namespace bpcm
+
+}  // namespace rl4
+
+#endif  // defined(RL_USE_HALF_FLOATS) && defined(RL_BUILD_WITH_SSE)
diff --git a/Engine/Plugins/Animation/RigLogic/Source/RigLogicLib/Private/riglogic/joints/bpcm/builders/SSEHalfFloat.h b/Engine/Plugins/Animation/RigLogic/Source/RigLogicLib/Private/riglogic/joints/bpcm/builders/SSEHalfFloat.h
new file mode 100644
--- /dev/null
+++ b/Engine/Plugins/Animation/RigLogic/Source/RigLogicLib/Private/riglogic/joints/bpcm/builders/SSEHalfFloat.h
@@ -0,0 +1,30 @@
+#pragma once
+
+// *INDENT-OFF*
+#if defined(RL_USE_HALF_FLOATS) && defined(RL_BUILD_WITH_SSE)
+
+#include "riglogic/joints/bpcm/JointsBuilderCommon.h"
+#include "riglogic/types/Aliases.h"
+
+#include <cstdint>
+
+namespace rl4 {
+
+namespace bpcm {
+
+class SSEJointsBuilder : public JointsBuilderCommon<std::uint16_t> {
+    public:
+        explicit SSEJointsBuilder(MemoryResource* memRes_);
+        ~SSEJointsBuilder();
+
+    protected:
+        void setValues(const dna::BehaviorReader* reader) override;
+
+};
+
+}  // namespace bpcm
+
+}  // namespace rl4
+
+#endif  // defined(RL_USE_HALF_FLOATS) && defined(RL_BUILD_WITH_SSE)
+// *INDENT-ON*
diff --git a/Engine/Plugins/Animation/RigLogic/Source/RigLogicLib/Private/riglogic/joints/bpcm/builders/Scalar.cpp b/Engine/Plugins/Animation/RigLogic/Source/RigLogicLib/Private/riglogic/joints/bpcm/builders/Scalar.cpp
new file mode 100644
--- /dev/null
+++ b/Engine/Plugins/Animation/RigLogic/Source/RigLogicLib/Private/riglogic/joints/bpcm/builders/Scalar.cpp
@@ -0,0 +1,24 @@
+#include "riglogic/joints/bpcm/builders/Scalar.h"
+
+#include "riglogic/joints/bpcm/Consts.h"
+#include "riglogic/joints/bpcm/strategies/Scalar.h"
+
+#include <pma/utils/ManagedInstance.h>
+
+#include <cstddef>
+#include <cstdint>
+
+namespace rl4 {
+
+namespace bpcm {
+
+ScalarJointsBuilder::ScalarJointsBuilder(MemoryResource* memRes_) : FloatStorageBuilder{block8Height, block4Height,
+                                                                                        memRes_} {
+    strategy = pma::UniqueInstance<ScalarJointCalculationStrategy, CalculationStrategy>::with(memRes).create();
+}
+
+ScalarJointsBuilder::~ScalarJointsBuilder() = default;
+
+}  // namespace bpcm
+
+}  // namespace rl4
diff --git a/Engine/Plugins/Animation/RigLogic/Source/RigLogicLib/Private/riglogic/joints/bpcm/builders/Scalar.h b/Engine/Plugins/Animation/RigLogic/Source/RigLogicLib/Private/riglogic/joints/bpcm/builders/Scalar.h
new file mode 100644
--- /dev/null
+++ b/Engine/Plugins/Animation/RigLogic/Source/RigLogicLib/Private/riglogic/joints/bpcm/builders/Scalar.h
@@ -0,0 +1,18 @@
+#pragma once
+
+#include "riglogic/joints/bpcm/builders/Float.h"
+#include "riglogic/types/Aliases.h"
+
+namespace rl4 {
+
+namespace bpcm {
+
+class ScalarJointsBuilder : public FloatStorageBuilder {
+    public:
+        explicit ScalarJointsBuilder(MemoryResource* memRes_);
+        ~ScalarJointsBuilder();
+};
+
+}  // namespace bpcm
+
+}  // namespace rl4
diff --git a/Engine/Plugins/Animation/RigLogic/Source/RigLogicLib/Private/riglogic/joints/bpcm/strategies/AVX.h b/Engine/Plugins/Animation/RigLogic/Source/RigLogicLib/Private/riglogic/joints/bpcm/strategies/AVX.h
new file mode 100644
--- /dev/null
+++ b/Engine/Plugins/Animation/RigLogic/Source/RigLogicLib/Private/riglogic/joints/bpcm/strategies/AVX.h
@@ -0,0 +1,15 @@
+#pragma once
+
+#include "riglogic/joints/bpcm/strategies/Block8.h"
+#include "riglogic/system/simd/SIMD.h"
+
+namespace rl4 {
+
+namespace bpcm {
+
+template<typename T>
+using AVXJointCalculationStrategy = Block8JointCalculationStrategy<simd::avx::F256, T>;
+
+}  // namespace bpcm
+
+}  // namespace rl4
diff --git a/Engine/Plugins/Animation/RigLogic/Source/RigLogicLib/Private/riglogic/joints/bpcm/strategies/Block4.h b/Engine/Plugins/Animation/RigLogic/Source/RigLogicLib/Private/riglogic/joints/bpcm/strategies/Block4.h
new file mode 100644
--- /dev/null
+++ b/Engine/Plugins/Animation/RigLogic/Source/RigLogicLib/Private/riglogic/joints/bpcm/strategies/Block4.h
@@ -0,0 +1,362 @@
+#pragma once
+
+#include "riglogic/joints/bpcm/Consts.h"
+#include "riglogic/joints/bpcm/JointCalculationStrategy.h"
+#include "riglogic/joints/bpcm/JointGroup.h"
+#include "riglogic/joints/bpcm/JointStorageView.h"
+#include "riglogic/joints/bpcm/LODRegion.h"
+#include "riglogic/utils/Macros.h"
+
+#include <cstddef>
+#include <cstdint>
+
+namespace rl4 {
+
+namespace bpcm {
+
+/*
+ * Process the remainder portion after 8x4 blocks
+ *
+ * Consumes 4 x 2 values per input index, but input indexes are processed one
+ * by one. The maximum number of iterations performed is 3, as anything larger
+ * would be already consumed by the 8x4 block processor. It's executed only
+ * after the 8x4 block processor in cases when the matrix column count is not
+ * divisible by 8.
+ *
+ *  [ x x x x x x x x X Y Z ] Input Vector
+ *
+ *    0 1 2 3 4 5 6 7 8 9 0
+ *  0 x x x x x x x x A C E  |  X * A + Y * C + Z * E
+ *  1 x x x x x x x x A C E  |  X * A + Y * C + Z * E
+ *  2 x x x x x x x x A C E  |  X * A + Y * C + Z * E
+ *  3 x x x x x x x x A C E  |  X * A + Y * C + Z * E
+ *  4 x x x x x x x x B D F  |  X * B + Y * D + Z * F
+ *  5 x x x x x x x x B D F  |  X * B + Y * D + Z * F
+ *  6 x x x x x x x x B D F  |  X * B + Y * D + Z * F
+ *  7 x x x x x x x x B D F  |  X * B + Y * D + Z * F
+ *  8 x x x x x x x x x x x  |
+ *  9 x x x x x x x x x x x  |
+ * 10 x x x x x x x x x x x  |
+ * 11 x x x x x x x x x x x  |
+ *
+ */
+template<typename TF128, typename T>
+static FORCE_INLINE void processBlocks8x1(const std::uint16_t* inputIndicesEndAlignedTo4,
+                                          const std::uint16_t* inputIndicesEnd,
+                                          ConstArrayView<float> inputs,
+                                          const T*& values,
+                                          TF128& sum1,
+                                          TF128& sum2) {
+    for (const std::uint16_t* inputIndices = inputIndicesEndAlignedTo4;
+         inputIndices < inputIndicesEnd;
+         ++inputIndices, values += block8Height) {
+        const TF128 inputVec{inputs[*inputIndices]};
+        const TF128 blk1 = TF128::fromAlignedSource(values);
+        const TF128 blk2 = TF128::fromAlignedSource(values + TF128::size());
+        sum1 += (blk1 * inputVec);
+        sum2 += (blk2 * inputVec);
+    }
+}
+
+/*
+ * Process the main portion of the matrix that's partitionable into 8x4 blocks
+ *
+ * Consumes 4 x 2 values per input index, and loads 4 input indexes per iteration,
+ * thus a single iteration processes 32 values. This is the most performant of all
+ * the block processors as for every 4 x 2 values only a single input vector needs
+ * to be filled (which is expensive as it involves a lot of shuffle instructions).
+ *
+ *  [ S T U V W X Y Z x x x ] Input Vector
+ *
+ *    0 1 2 3 4 5 6 7 8 9 0
+ *  0 A C E G I K M O x x x  |  S * A + T * C + U * E + V * G + W * I + X * K + Y * M + Z * O
+ *  1 A C E G I K M O x x x  |  S * A + T * C + U * E + V * G + W * I + X * K + Y * M + Z * O
+ *  2 A C E G I K M O x x x  |  S * A + T * C + U * E + V * G + W * I + X * K + Y * M + Z * O
+ *  3 A C E G I K M O x x x  |  S * A + T * C + U * E + V * G + W * I + X * K + Y * M + Z * O
+ *  4 B D F H J L N P x x x  |  S * B + T * D + U * F + V * H + W * J + X * L + Y * N + Z * P
+ *  5 B D F H J L N P x x x  |  S * B + T * D + U * F + V * H + W * J + X * L + Y * N + Z * P
+ *  6 B D F H J L N P x x x  |  S * B + T * D + U * F + V * H + W * J + X * L + Y * N + Z * P
+ *  7 B D F H J L N P x x x  |  S * B + T * D + U * F + V * H + W * J + X * L + Y * N + Z * P
+ *  8 x x x x x x x x x x x  |
+ *  9 x x x x x x x x x x x  |
+ * 10 x x x x x x x x x x x  |
+ * 11 x x x x x x x x x x x  |
+ *
+ */
+template<typename TF128, typename T>
+static FORCE_INLINE void processBlocks8x4(const std::uint16_t* inputIndicesStart,
+                                          const std::uint16_t* inputIndicesEndAlignedTo4,
+                                          const std::uint16_t* inputIndicesEnd,
+                                          ConstArrayView<float> inputs,
+                                          const T*& values,
+                                          float* outbuf) {
+    TF128 sum1{};
+    TF128 sum2{};
+    TF128 sum3{};
+    TF128 sum4{};
+    TF128 sum5{};
+    TF128 sum6{};
+    TF128 sum7{};
+    TF128 sum8{};
+    for (const std::uint16_t* inputIndices = inputIndicesStart;
+         inputIndices < inputIndicesEndAlignedTo4;
+         inputIndices += block8Width, values += block8Height * block8Width) {
+        const TF128 inputVec1{inputs[inputIndices[0]]};
+        const TF128 inputVec2{inputs[inputIndices[1]]};
+        const TF128 inputVec3{inputs[inputIndices[2]]};
+        const TF128 inputVec4{inputs[inputIndices[3]]};
+        const TF128 blk1 = TF128::fromAlignedSource(values);
+        const TF128 blk2 = TF128::fromAlignedSource(values + TF128::size());
+        const TF128 blk3 = TF128::fromAlignedSource(values + TF128::size() * 2);
+        const TF128 blk4 = TF128::fromAlignedSource(values + TF128::size() * 3);
+        const TF128 blk5 = TF128::fromAlignedSource(values + TF128::size() * 4);
+        const TF128 blk6 = TF128::fromAlignedSource(values + TF128::size() * 5);
+        const TF128 blk7 = TF128::fromAlignedSource(values + TF128::size() * 6);
+        const TF128 blk8 = TF128::fromAlignedSource(values + TF128::size() * 7);
+        #ifndef RL_NO_MANUAL_PREFETCH
+            #ifdef RL_USE_NON_TEMPORAL_PREFETCH
+                TF128::prefetchNTA(values + (lookAheadOffset * 6ul));
+                TF128::prefetchNTA(values + (lookAheadOffset * 7ul));
+                TF128::prefetchNTA(values + (lookAheadOffset * 8ul));
+                TF128::prefetchNTA(values + (lookAheadOffset * 9ul));
+                TF128::prefetchNTA(values + (lookAheadOffset * 10ul));
+                TF128::prefetchNTA(values + (lookAheadOffset * 11ul));
+            #else
+                TF128::prefetchT0(values + (lookAheadOffset * 6ul));
+                TF128::prefetchT0(values + (lookAheadOffset * 7ul));
+                TF128::prefetchT0(values + (lookAheadOffset * 8ul));
+                TF128::prefetchT0(values + (lookAheadOffset * 9ul));
+                TF128::prefetchT0(values + (lookAheadOffset * 10ul));
+                TF128::prefetchT0(values + (lookAheadOffset * 11ul));
+            #endif
+        #endif
+        sum1 += (blk1 * inputVec1);
+        sum2 += (blk2 * inputVec1);
+        sum3 += (blk3 * inputVec2);
+        sum4 += (blk4 * inputVec2);
+        sum5 += (blk5 * inputVec3);
+        sum6 += (blk6 * inputVec3);
+        sum7 += (blk7 * inputVec4);
+        sum8 += (blk8 * inputVec4);
+    }
+    // Process 8x1 horizontal remainder portion after 8x4 blocks are consumed
+    processBlocks8x1(inputIndicesEndAlignedTo4, inputIndicesEnd, inputs, values, sum1, sum2);
+
+    sum1 += sum3;
+    sum2 += sum4;
+    sum5 += sum7;
+    sum6 += sum8;
+    sum1 += sum5;
+    sum2 += sum6;
+
+    sum1.alignedStore(outbuf);
+    sum2.alignedStore(outbuf + TF128::size());
+}
+
+/*
+ * Process the remainder portion after 4x8 blocks
+ *
+ * Consumes 4 x 1 values per input index, but input indexes are processed one
+ * by one. The maximum number of iterations performed is 7, as anything larger
+ * would be already consumed by the 4x8 block processor. It's executed only
+ * after the 4x8 block processor, in cases when the matrix column count is not
+ * divisible by 8.
+ *
+ *  [ x x x x x x x x T U V W X Y Z ] Input Vector
+ *
+ *    0 1 2 3 4 5 6 7 8 9 0 1 2 3 4
+ *  0 x x x x x x x x x x x x x x x  |
+ *  1 x x x x x x x x x x x x x x x  |
+ *  2 x x x x x x x x x x x x x x x  |
+ *  3 x x x x x x x x x x x x x x x  |
+ *  4 x x x x x x x x x x x x x x x  |
+ *  5 x x x x x x x x x x x x x x x  |
+ *  6 x x x x x x x x x x x x x x x  |
+ *  7 x x x x x x x x x x x x x x x  |
+ *  8 x x x x x x x x A B C D E F G  |  T * A + U * B + V * C + W * D + X * E + Y * F + Z * G
+ *  9 x x x x x x x x A B C D E F G  |  T * A + U * B + V * C + W * D + X * E + Y * F + Z * G
+ * 10 x x x x x x x x A B C D E F G  |  T * A + U * B + V * C + W * D + X * E + Y * F + Z * G
+ * 11 x x x x x x x x A B C D E F G  |  T * A + U * B + V * C + W * D + X * E + Y * F + Z * G
+ *
+ */
+template<typename TF128, typename T>
+static FORCE_INLINE void processBlocks4x1(const std::uint16_t* inputIndicesEndAlignedTo8,
+                                          const std::uint16_t* inputIndicesEnd,
+                                          ConstArrayView<float> inputs,
+                                          const T*& values,
+                                          TF128& sum1) {
+    for (const std::uint16_t* inputIndices = inputIndicesEndAlignedTo8;
+         inputIndices < inputIndicesEnd;
+         ++inputIndices, values += block4Height) {
+        const TF128 inputVec{inputs[*inputIndices]};
+        const TF128 blk = TF128::fromAlignedSource(values);
+        sum1 += (blk * inputVec);
+    }
+}
+
+/*
+ * Process the vertical remainder portion of the matrix that's partitionable into 4x8 blocks
+ *
+ * Consumes 4 x 1 values per input index, and loads 8 input indexes per
+ * iteration, thus a single iteration processes 32 values. Although it
+ * processes the same number of values per iteration as the 8x4 block
+ * processor, it does so by having to fill 8 input vectors (one for each
+ * 4 x 1 value-block), thus requiring twice the number of shuffle instructions
+ * which makes it somewhat less performant than the 8x4 variant. It gets
+ * executed only in cases when the matrix row count is not divisible by 8.
+ *
+ *  [ S T U V W X Y Z x x x ] Input Vector
+ *
+ *    0 1 2 3 4 5 6 7 8 9 0
+ *  0 x x x x x x x x x x x  |
+ *  1 x x x x x x x x x x x  |
+ *  2 x x x x x x x x x x x  |
+ *  3 x x x x x x x x x x x  |
+ *  4 x x x x x x x x x x x  |
+ *  5 x x x x x x x x x x x  |
+ *  6 x x x x x x x x x x x  |
+ *  7 x x x x x x x x x x x  |
+ *  8 A B C D E F G H x x x  |  S * A + T * B + U * C + V * D + W * E + X * F + Y * G + Z * H
+ *  9 A B C D E F G H x x x  |  S * A + T * B + U * C + V * D + W * E + X * F + Y * G + Z * H
+ * 10 A B C D E F G H x x x  |  S * A + T * B + U * C + V * D + W * E + X * F + Y * G + Z * H
+ * 11 A B C D E F G H x x x  |  S * A + T * B + U * C + V * D + W * E + X * F + Y * G + Z * H
+ *
+ */
+template<typename TF128, typename T>
+static FORCE_INLINE void processBlocks4x8(const std::uint16_t* inputIndicesStart,
+                                          const std::uint16_t* inputIndicesEndAlignedTo8,
+                                          const std::uint16_t* inputIndicesEnd,
+                                          ConstArrayView<float> inputs,
+                                          const T*& values,
+                                          float* outbuf) {
+    TF128 sum1{};
+    TF128 sum2{};
+    TF128 sum3{};
+    TF128 sum4{};
+    TF128 sum5{};
+    TF128 sum6{};
+    TF128 sum7{};
+    TF128 sum8{};
+    for (const std::uint16_t* inputIndices = inputIndicesStart;
+         inputIndices < inputIndicesEndAlignedTo8;
+         inputIndices += block4Width, values += block4Height * block4Width) {
+        const TF128 inputVec1{inputs[inputIndices[0]]};
+        const TF128 inputVec2{inputs[inputIndices[1]]};
+        const TF128 inputVec3{inputs[inputIndices[2]]};
+        const TF128 inputVec4{inputs[inputIndices[3]]};
+        const TF128 inputVec5{inputs[inputIndices[4]]};
+        const TF128 inputVec6{inputs[inputIndices[5]]};
+        const TF128 inputVec7{inputs[inputIndices[6]]};
+        const TF128 inputVec8{inputs[inputIndices[7]]};
+        const TF128 blk1 = TF128::fromAlignedSource(values);
+        const TF128 blk2 = TF128::fromAlignedSource(values + TF128::size());
+        const TF128 blk3 = TF128::fromAlignedSource(values + TF128::size() * 2);
+        const TF128 blk4 = TF128::fromAlignedSource(values + TF128::size() * 3);
+        const TF128 blk5 = TF128::fromAlignedSource(values + TF128::size() * 4);
+        const TF128 blk6 = TF128::fromAlignedSource(values + TF128::size() * 5);
+        const TF128 blk7 = TF128::fromAlignedSource(values + TF128::size() * 6);
+        const TF128 blk8 = TF128::fromAlignedSource(values + TF128::size() * 7);
+        sum1 += (blk1 * inputVec1);
+        sum2 += (blk2 * inputVec2);
+        sum3 += (blk3 * inputVec3);
+        sum4 += (blk4 * inputVec4);
+        sum5 += (blk5 * inputVec5);
+        sum6 += (blk6 * inputVec6);
+        sum7 += (blk7 * inputVec7);
+        sum8 += (blk8 * inputVec8);
+    }
+    // Process 4x1 horizontal remainder portion after 4x8 blocks are consumed
+    processBlocks4x1(inputIndicesEndAlignedTo8, inputIndicesEnd, inputs, values, sum1);
+
+    sum1 += sum2;
+    sum3 += sum4;
+    sum5 += sum6;
+    sum7 += sum8;
+    sum1 += sum3;
+    sum5 += sum7;
+    sum1 += sum5;
+
+    sum1.alignedStore(outbuf);
+}
+
+/*
+ * Orchestrate the execution of the needed block processors for a given joint group
+ */
+template<typename TF128, typename T>
+static FORCE_INLINE void processJointGroupBlock4(const JointGroupView<T>& jointGroup, ConstArrayView<float> inputs,
+                                                 ArrayView<float> outputs, std::uint16_t lod) {
+    const T* values = jointGroup.values;
+    const std::uint16_t* const inputIndices = jointGroup.inputIndices;
+    const std::uint16_t* const inputIndicesEnd = jointGroup.inputIndicesEnd;
+    const std::uint16_t* const inputIndicesEndAlignedTo4 = jointGroup.inputIndicesEndAlignedTo4;
+    const std::uint16_t* const inputIndicesEndAlignedTo8 = jointGroup.inputIndicesEndAlignedTo8;
+    const LODRegion& lodRegion = jointGroup.lods[lod];
+    const std::uint16_t* outputIndices = jointGroup.outputIndices;
+    const std::uint16_t* const outputIndicesEnd = outputIndices + lodRegion.size;
+    const std::uint16_t* const outputIndicesEndAlignedToLastBlock8 = outputIndices + lodRegion.sizeAlignedToLastFullBlock;
+    const std::uint16_t* const outputIndicesEndAlignedToSecondLastBlock8 = outputIndices +
+        lodRegion.sizeAlignedToSecondLastFullBlock;
+    // Process portion of matrix that's partitionable into 8x4 blocks
+    for (; outputIndices < outputIndicesEndAlignedToSecondLastBlock8; outputIndices += block8Height) {
+        alignas(TF128::alignment()) float outbuf[block8Height];
+        processBlocks8x4<TF128>(inputIndices, inputIndicesEndAlignedTo4, inputIndicesEnd, inputs, values,
+                                static_cast<float*>(outbuf));
+        for (std::size_t i = 0ul; i < block8Height; ++i) {
+            outputs[outputIndices[i]] = outbuf[i];
+        }
+    }
+    // Process the last 8x4 block which needs special handling when some of
+    // the output values need to be masked-off because of LODs
+    for (; outputIndices < outputIndicesEndAlignedToLastBlock8; outputIndices += block8Height) {
+        alignas(TF128::alignment()) float outbuf[block8Height];
+        processBlocks8x4<TF128>(inputIndices, inputIndicesEndAlignedTo4, inputIndicesEnd, inputs, values,
+                                static_cast<float*>(outbuf));
+        // Ignore results that came from rows after the last LOD row
+        for (std::size_t i = 0ul; i < (lodRegion.size % block8Height); ++i) {
+            // NOLINTNEXTLINE(cppcoreguidelines-pro-bounds-constant-array-index)
+            outputs[outputIndices[i]] = outbuf[i];
+        }
+    }
+    // Process vertical remainder portion of matrix that's partitionable into 4x8 blocks
+    for (; outputIndices < outputIndicesEnd; outputIndices += block4Height) {
+        alignas(TF128::alignment()) float outbuf[block4Height];
+        processBlocks4x8<TF128>(inputIndices, inputIndicesEndAlignedTo8, inputIndicesEnd, inputs, values,
+                                static_cast<float*>(outbuf));
+        // Ignore results that came from rows after the last LOD row
+        auto maskOffStart = static_cast<std::size_t>(outputIndicesEnd - outputIndices);
+        for (std::size_t i = 0; i < maskOffStart; ++i) {
+            // NOLINTNEXTLINE(cppcoreguidelines-pro-bounds-constant-array-index)
+            outputs[outputIndices[i]] = outbuf[i];
+        }
+    }
+}
+
+template<typename TF128, typename T>
+struct Block4JointCalculationStrategy : public JointCalculationStrategy<T> {
+
+    using JointGroupArrayView = typename JointCalculationStrategy<T>::JointGroupArrayView;
+
+    void calculate(const JointGroupArrayView& jointGroups,
+                   ConstArrayView<float> inputs,
+                   ArrayView<float> outputs,
+                   std::uint16_t lod) const override {
+
+        for (std::size_t i = 0ul; i < jointGroups.size(); ++i) {
+            processJointGroupBlock4<TF128>(jointGroups[i], inputs, outputs, lod);
+        }
+    }
+
+    void calculate(const JointGroupArrayView& jointGroups,
+                   ConstArrayView<float> inputs,
+                   ArrayView<float> outputs,
+                   std::uint16_t lod,
+                   std::uint16_t jointGroupIndex) const override {
+
+        processJointGroupBlock4<TF128>(jointGroups[jointGroupIndex], inputs, outputs, lod);
+    }
+
+};
+
+}  // namespace bpcm
+
+}  // namespace rl4
diff --git a/Engine/Plugins/Animation/RigLogic/Source/RigLogicLib/Private/riglogic/joints/bpcm/strategies/Block8.h b/Engine/Plugins/Animation/RigLogic/Source/RigLogicLib/Private/riglogic/joints/bpcm/strategies/Block8.h
new file mode 100644
--- /dev/null
+++ b/Engine/Plugins/Animation/RigLogic/Source/RigLogicLib/Private/riglogic/joints/bpcm/strategies/Block8.h
@@ -0,0 +1,260 @@
+#pragma once
+
+#include "riglogic/joints/bpcm/Consts.h"
+#include "riglogic/joints/bpcm/JointCalculationStrategy.h"
+#include "riglogic/joints/bpcm/JointGroup.h"
+#include "riglogic/joints/bpcm/JointStorageView.h"
+#include "riglogic/joints/bpcm/LODRegion.h"
+#include "riglogic/utils/Macros.h"
+
+#include <cstddef>
+#include <cstdint>
+
+// The Block-8 algorithm follows the same logic and data layout as the Block-4 algorithm,
+// and the only difference between the two is the size of blocks that are processed.
+
+namespace rl4 {
+
+namespace bpcm {
+
+template<typename TF256, typename T>
+static FORCE_INLINE void processBlocks16x1(const std::uint16_t* inputIndicesEndAlignedTo4,
+                                           const std::uint16_t* inputIndicesEnd,
+                                           ConstArrayView<float> inputs,
+                                           const T*& values,
+                                           TF256& sum1,
+                                           TF256& sum2) {
+    for (const std::uint16_t* inputIndices = inputIndicesEndAlignedTo4;
+         inputIndices < inputIndicesEnd;
+         ++inputIndices, values += (2ul * TF256::size())) {
+        const TF256 inputVec{inputs[*inputIndices]};
+        const TF256 blk1 = TF256::fromAlignedSource(values);
+        const TF256 blk2 = TF256::fromAlignedSource(values + TF256::size());
+        sum1 += (blk1 * inputVec);
+        sum2 += (blk2 * inputVec);
+    }
+}
+
+template<typename TF256, typename T>
+static FORCE_INLINE void processBlocks16x4(const std::uint16_t* inputIndicesStart,
+                                           const std::uint16_t* inputIndicesEndAlignedTo4,
+                                           const std::uint16_t* inputIndicesEnd,
+                                           ConstArrayView<float> inputs,
+                                           const T*& values,
+                                           float* outbuf) {
+    TF256 sum1{};
+    TF256 sum2{};
+    TF256 sum3{};
+    TF256 sum4{};
+    TF256 sum5{};
+    TF256 sum6{};
+    TF256 sum7{};
+    TF256 sum8{};
+    for (const std::uint16_t* inputIndices = inputIndicesStart;
+         inputIndices < inputIndicesEndAlignedTo4;
+         inputIndices += block16Width, values += (2ul * TF256::size() * block16Width)) {
+        const TF256 inputVec1{inputs[inputIndices[0]]};
+        const TF256 inputVec2{inputs[inputIndices[1]]};
+        const TF256 inputVec3{inputs[inputIndices[2]]};
+        const TF256 inputVec4{inputs[inputIndices[3]]};
+        const TF256 blk1 = TF256::fromAlignedSource(values);
+        const TF256 blk2 = TF256::fromAlignedSource(values + TF256::size());
+        const TF256 blk3 = TF256::fromAlignedSource(values + TF256::size() * 2);
+        const TF256 blk4 = TF256::fromAlignedSource(values + TF256::size() * 3);
+        const TF256 blk5 = TF256::fromAlignedSource(values + TF256::size() * 4);
+        const TF256 blk6 = TF256::fromAlignedSource(values + TF256::size() * 5);
+        const TF256 blk7 = TF256::fromAlignedSource(values + TF256::size() * 6);
+        const TF256 blk8 = TF256::fromAlignedSource(values + TF256::size() * 7);
+        #ifndef RL_NO_MANUAL_PREFETCH
+            #ifdef RL_USE_NON_TEMPORAL_PREFETCH
+                TF256::prefetchNTA(values + (lookAheadOffset * 6ul));
+                TF256::prefetchNTA(values + (lookAheadOffset * 7ul));
+                TF256::prefetchNTA(values + (lookAheadOffset * 8ul));
+                TF256::prefetchNTA(values + (lookAheadOffset * 9ul));
+                TF256::prefetchNTA(values + (lookAheadOffset * 10ul));
+                TF256::prefetchNTA(values + (lookAheadOffset * 11ul));
+            #else
+                TF256::prefetchT0(values + (lookAheadOffset * 6ul));
+                TF256::prefetchT0(values + (lookAheadOffset * 7ul));
+                TF256::prefetchT0(values + (lookAheadOffset * 8ul));
+                TF256::prefetchT0(values + (lookAheadOffset * 9ul));
+                TF256::prefetchT0(values + (lookAheadOffset * 10ul));
+                TF256::prefetchT0(values + (lookAheadOffset * 11ul));
+            #endif
+        #endif
+        sum1 += (blk1 * inputVec1);
+        sum2 += (blk2 * inputVec1);
+        sum3 += (blk3 * inputVec2);
+        sum4 += (blk4 * inputVec2);
+        sum5 += (blk5 * inputVec3);
+        sum6 += (blk6 * inputVec3);
+        sum7 += (blk7 * inputVec4);
+        sum8 += (blk8 * inputVec4);
+    }
+    // Process 16x1 horizontal remainder portion after 16x4 blocks are consumed
+    processBlocks16x1(inputIndicesEndAlignedTo4, inputIndicesEnd, inputs, values, sum1, sum2);
+
+    sum1 += sum3;
+    sum2 += sum4;
+    sum5 += sum7;
+    sum6 += sum8;
+    sum1 += sum5;
+    sum2 += sum6;
+
+    sum1.alignedStore(outbuf);
+    sum2.alignedStore(outbuf + TF256::size());
+}
+
+template<typename TF256, typename T>
+static FORCE_INLINE void processBlocks8x1(const std::uint16_t* inputIndicesEndAlignedTo8,
+                                          const std::uint16_t* inputIndicesEnd,
+                                          ConstArrayView<float> inputs,
+                                          const T*& values,
+                                          TF256& sum1) {
+    for (const std::uint16_t* inputIndices = inputIndicesEndAlignedTo8;
+         inputIndices < inputIndicesEnd;
+         ++inputIndices, values += TF256::size()) {
+        const TF256 inputVec{inputs[*inputIndices]};
+        const TF256 blk = TF256::fromAlignedSource(values);
+        sum1 += (blk * inputVec);
+    }
+}
+
+template<typename TF256, typename T>
+static FORCE_INLINE void processBlocks8x8(const std::uint16_t* inputIndicesStart,
+                                          const std::uint16_t* inputIndicesEndAlignedTo8,
+                                          const std::uint16_t* inputIndicesEnd,
+                                          ConstArrayView<float> inputs,
+                                          const T*& values,
+                                          float* outbuf) {
+    TF256 sum1{};
+    TF256 sum2{};
+    TF256 sum3{};
+    TF256 sum4{};
+    TF256 sum5{};
+    TF256 sum6{};
+    TF256 sum7{};
+    TF256 sum8{};
+    for (const std::uint16_t* inputIndices = inputIndicesStart;
+         inputIndices < inputIndicesEndAlignedTo8;
+         inputIndices += block4Width, values += (TF256::size() * block4Width)) {
+        const TF256 inputVec1{inputs[inputIndices[0]]};
+        const TF256 inputVec2{inputs[inputIndices[1]]};
+        const TF256 inputVec3{inputs[inputIndices[2]]};
+        const TF256 inputVec4{inputs[inputIndices[3]]};
+        const TF256 inputVec5{inputs[inputIndices[4]]};
+        const TF256 inputVec6{inputs[inputIndices[5]]};
+        const TF256 inputVec7{inputs[inputIndices[6]]};
+        const TF256 inputVec8{inputs[inputIndices[7]]};
+        const TF256 blk1 = TF256::fromAlignedSource(values);
+        const TF256 blk2 = TF256::fromAlignedSource(values + TF256::size());
+        const TF256 blk3 = TF256::fromAlignedSource(values + TF256::size() * 2);
+        const TF256 blk4 = TF256::fromAlignedSource(values + TF256::size() * 3);
+        const TF256 blk5 = TF256::fromAlignedSource(values + TF256::size() * 4);
+        const TF256 blk6 = TF256::fromAlignedSource(values + TF256::size() * 5);
+        const TF256 blk7 = TF256::fromAlignedSource(values + TF256::size() * 6);
+        const TF256 blk8 = TF256::fromAlignedSource(values + TF256::size() * 7);
+        sum1 += (blk1 * inputVec1);
+        sum2 += (blk2 * inputVec2);
+        sum3 += (blk3 * inputVec3);
+        sum4 += (blk4 * inputVec4);
+        sum5 += (blk5 * inputVec5);
+        sum6 += (blk6 * inputVec6);
+        sum7 += (blk7 * inputVec7);
+        sum8 += (blk8 * inputVec8);
+    }
+    // Process 8x1 horizontal remainder portion after 8x8 blocks are consumed
+    processBlocks8x1(inputIndicesEndAlignedTo8, inputIndicesEnd, inputs, values, sum1);
+
+    sum1 += sum2;
+    sum3 += sum4;
+    sum5 += sum6;
+    sum7 += sum8;
+    sum1 += sum3;
+    sum5 += sum7;
+    sum1 += sum5;
+
+    sum1.alignedStore(outbuf);
+}
+
+/*
+ * Orchestrate the execution of the needed block processors for a given joint group
+ */
+template<typename TF256, typename T>
+static FORCE_INLINE void processJointGroupBlock8(const JointGroupView<T>& jointGroup, ConstArrayView<float> inputs,
+                                                 ArrayView<float> outputs, std::uint16_t lod) {
+    const T* values = jointGroup.values;
+    const std::uint16_t* const inputIndices = jointGroup.inputIndices;
+    const std::uint16_t* const inputIndicesEnd = jointGroup.inputIndicesEnd;
+    const std::uint16_t* const inputIndicesEndAlignedTo4 = jointGroup.inputIndicesEndAlignedTo4;
+    const std::uint16_t* const inputIndicesEndAlignedTo8 = jointGroup.inputIndicesEndAlignedTo8;
+    const LODRegion& lodRegion = jointGroup.lods[lod];
+    const std::uint16_t* outputIndices = jointGroup.outputIndices;
+    const std::uint16_t* const outputIndicesEnd = outputIndices + lodRegion.size;
+    const std::uint16_t* const outputIndicesEndAlignedToLastFullBlock = outputIndices + lodRegion.sizeAlignedToLastFullBlock;
+    const std::uint16_t* const outputIndicesEndAlignedToSecondLastFullBlock = outputIndices +
+        lodRegion.sizeAlignedToSecondLastFullBlock;
+    // Process portion of matrix that's partitionable into 16x4 blocks
+    for (; outputIndices < outputIndicesEndAlignedToSecondLastFullBlock; outputIndices += block16Height) {
+        alignas(TF256::alignment()) float outbuf[block16Height];
+        processBlocks16x4<TF256>(inputIndices, inputIndicesEndAlignedTo4, inputIndicesEnd, inputs, values,
+                                 static_cast<float*>(outbuf));
+        for (std::size_t i = 0ul; i < block16Height; ++i) {
+            outputs[outputIndices[i]] = outbuf[i];
+        }
+    }
+    // Process the last 16x4 block which needs special handling when some of
+    // the output values need to be masked-off because of LODs
+    for (; outputIndices < outputIndicesEndAlignedToLastFullBlock; outputIndices += block16Height) {
+        alignas(TF256::alignment()) float outbuf[block16Height];
+        processBlocks16x4<TF256>(inputIndices, inputIndicesEndAlignedTo4, inputIndicesEnd, inputs, values,
+                                 static_cast<float*>(outbuf));
+        // Ignore results that came from rows after the last LOD row
+        for (std::size_t i = 0ul; i < (lodRegion.size % block16Height); ++i) {
+            // NOLINTNEXTLINE(cppcoreguidelines-pro-bounds-constant-array-index)
+            outputs[outputIndices[i]] = outbuf[i];
+        }
+    }
+    // Process vertical remainder portion of matrix that's partitionable into 4x8 blocks
+    for (; outputIndices < outputIndicesEnd; outputIndices += block8Height) {
+        alignas(TF256::alignment()) float outbuf[block8Height];
+        processBlocks8x8<TF256>(inputIndices, inputIndicesEndAlignedTo8, inputIndicesEnd, inputs, values,
+                                static_cast<float*>(outbuf));
+        // Ignore results that came from rows after the last LOD row
+        auto maskOffStart = static_cast<std::size_t>(outputIndicesEnd - outputIndices);
+        for (std::size_t i = 0; i < maskOffStart; ++i) {
+            // NOLINTNEXTLINE(cppcoreguidelines-pro-bounds-constant-array-index)
+            outputs[outputIndices[i]] = outbuf[i];
+        }
+    }
+}
+
+template<typename TF256, typename T>
+struct Block8JointCalculationStrategy : public JointCalculationStrategy<T> {
+
+    using JointGroupArrayView = typename JointCalculationStrategy<T>::JointGroupArrayView;
+
+    void calculate(const JointGroupArrayView& jointGroups,
+                   ConstArrayView<float> inputs,
+                   ArrayView<float> outputs,
+                   std::uint16_t lod) const override {
+
+        for (std::size_t i = 0ul; i < jointGroups.size(); ++i) {
+            processJointGroupBlock8<TF256>(jointGroups[i], inputs, outputs, lod);
+        }
+    }
+
+    void calculate(const JointGroupArrayView& jointGroups,
+                   ConstArrayView<float> inputs,
+                   ArrayView<float> outputs,
+                   std::uint16_t lod,
+                   std::uint16_t jointGroupIndex) const override {
+
+        processJointGroupBlock8<TF256>(jointGroups[jointGroupIndex], inputs, outputs, lod);
+    }
+
+};
+
+}  // namespace bpcm
+
+}  // namespace rl4
diff --git a/Engine/Plugins/Animation/RigLogic/Source/RigLogicLib/Private/riglogic/joints/bpcm/strategies/SSE.h b/Engine/Plugins/Animation/RigLogic/Source/RigLogicLib/Private/riglogic/joints/bpcm/strategies/SSE.h
new file mode 100644
--- /dev/null
+++ b/Engine/Plugins/Animation/RigLogic/Source/RigLogicLib/Private/riglogic/joints/bpcm/strategies/SSE.h
@@ -0,0 +1,15 @@
+#pragma once
+
+#include "riglogic/joints/bpcm/strategies/Block4.h"
+#include "riglogic/system/simd/SIMD.h"
+
+namespace rl4 {
+
+namespace bpcm {
+
+template<typename T>
+using SSEJointCalculationStrategy = Block4JointCalculationStrategy<simd::sse::F128, T>;
+
+}  // namespace bpcm
+
+}  // namespace rl4
diff --git a/Engine/Plugins/Animation/RigLogic/Source/RigLogicLib/Private/riglogic/joints/bpcm/strategies/Scalar.h b/Engine/Plugins/Animation/RigLogic/Source/RigLogicLib/Private/riglogic/joints/bpcm/strategies/Scalar.h
new file mode 100644
--- /dev/null
+++ b/Engine/Plugins/Animation/RigLogic/Source/RigLogicLib/Private/riglogic/joints/bpcm/strategies/Scalar.h
@@ -0,0 +1,14 @@
+#pragma once
+
+#include "riglogic/joints/bpcm/strategies/Block4.h"
+#include "riglogic/system/simd/SIMD.h"
+
+namespace rl4 {
+
+namespace bpcm {
+
+using ScalarJointCalculationStrategy = Block4JointCalculationStrategy<simd::scalar::F128, float>;
+
+}  // namespace bpcm
+
+}  // namespace rl4
diff --git a/Engine/Plugins/Animation/RigLogic/Source/RigLogicLib/Private/riglogic/psdmatrix/PSDMatrix.cpp b/Engine/Plugins/Animation/RigLogic/Source/RigLogicLib/Private/riglogic/psdmatrix/PSDMatrix.cpp
new file mode 100644
--- /dev/null
+++ b/Engine/Plugins/Animation/RigLogic/Source/RigLogicLib/Private/riglogic/psdmatrix/PSDMatrix.cpp
@@ -0,0 +1,86 @@
+#include "riglogic/psdmatrix/PSDMatrix.h"
+
+#include "riglogic/TypeDefs.h"
+
+#ifdef _MSC_VER
+    #pragma warning(push)
+    #pragma warning(disable : 4365 4987)
+#endif
+#include <algorithm>
+#include <cstddef>
+#ifdef _MSC_VER
+    #pragma warning(pop)
+#endif
+
+namespace rl4 {
+
+namespace {
+
+constexpr float maxControlValue = 1.0f;
+constexpr std::size_t blockSize = 4ul;
+
+}  // namespace
+
+PSDMatrix::PSDMatrix(MemoryResource* memRes) :
+    distinctPSDs{},
+    rowIndices{memRes},
+    columnIndices{memRes},
+    values{memRes} {
+}
+
+PSDMatrix::PSDMatrix(std::uint16_t distinctPSDs_,
+                     Vector<std::uint16_t>&& rowIndices_,
+                     Vector<std::uint16_t>&& columnIndices_,
+                     Vector<float>&& values_) :
+    distinctPSDs{distinctPSDs_},
+    rowIndices{std::move(rowIndices_)},
+    columnIndices{std::move(columnIndices_)},
+    values{std::move(values_)} {
+}
+
+std::uint16_t PSDMatrix::getDistinctPSDCount() const {
+    return distinctPSDs;
+}
+
+void PSDMatrix::calculate(float* inputs, std::uint16_t rawControlCount) const {
+    std::fill_n(inputs + rawControlCount, distinctPSDs, 1.0f);
+
+    std::size_t alignedSize = values.size() - (values.size() % blockSize);
+    for (std::size_t i = 0ul; i < alignedSize; i += blockSize) {
+        const std::uint16_t row0 = rowIndices[i];
+        const std::uint16_t row1 = rowIndices[i + 1ul];
+        const std::uint16_t row2 = rowIndices[i + 2ul];
+        const std::uint16_t row3 = rowIndices[i + 3ul];
+        const std::uint16_t col0 = columnIndices[i];
+        const std::uint16_t col1 = columnIndices[i + 1ul];
+        const std::uint16_t col2 = columnIndices[i + 2ul];
+        const std::uint16_t col3 = columnIndices[i + 3ul];
+        const float weight0 = values[i];
+        const float weight1 = values[i + 1ul];
+        const float weight2 = values[i + 2ul];
+        const float weight3 = values[i + 3ul];
+        const float input0 = inputs[col0];
+        const float input1 = inputs[col1];
+        const float input2 = inputs[col2];
+        const float input3 = inputs[col3];
+        inputs[row0] *= weight0 * input0;
+        inputs[row1] *= weight1 * input1;
+        inputs[row2] *= weight2 * input2;
+        inputs[row3] *= weight3 * input3;
+    }
+
+    for (std::size_t i = alignedSize; i < values.size(); ++i) {
+        const std::uint16_t row = rowIndices[i];
+        const std::uint16_t col = columnIndices[i];
+        const float weight = values[i];
+        const float input = inputs[col];
+        inputs[row] *= weight * input;
+    }
+
+    auto psds = inputs + rawControlCount;
+    for (std::size_t i = 0ul; i < distinctPSDs; ++i) {
+        psds[i] = std::min(maxControlValue, psds[i]);
+    }
+}
+
+}  // namespace rl4
diff --git a/Engine/Plugins/Animation/RigLogic/Source/RigLogicLib/Private/riglogic/psdmatrix/PSDMatrix.h b/Engine/Plugins/Animation/RigLogic/Source/RigLogicLib/Private/riglogic/psdmatrix/PSDMatrix.h
new file mode 100644
--- /dev/null
+++ b/Engine/Plugins/Animation/RigLogic/Source/RigLogicLib/Private/riglogic/psdmatrix/PSDMatrix.h
@@ -0,0 +1,32 @@
+#pragma once
+
+#include "riglogic/TypeDefs.h"
+
+#include <cstddef>
+
+namespace rl4 {
+
+class PSDMatrix {
+    public:
+        explicit PSDMatrix(MemoryResource* memRes);
+        PSDMatrix(std::uint16_t distinctPSDs_,
+                  Vector<std::uint16_t>&& rowIndices_,
+                  Vector<std::uint16_t>&& columnIndices_,
+                  Vector<float>&& values_);
+
+        std::uint16_t getDistinctPSDCount() const;
+        void calculate(float* inputs, std::uint16_t rawControlCount) const;
+
+        template<class Archive>
+        void serialize(Archive& archive) {
+            archive(distinctPSDs, rowIndices, columnIndices, values);
+        }
+
+    private:
+        std::uint16_t distinctPSDs;
+        Vector<std::uint16_t> rowIndices;
+        Vector<std::uint16_t> columnIndices;
+        Vector<float> values;
+};
+
+}  // namespace rl4
diff --git a/Engine/Plugins/Animation/RigLogic/Source/RigLogicLib/Private/riglogic/riglogic/Configuration.h b/Engine/Plugins/Animation/RigLogic/Source/RigLogicLib/Private/riglogic/riglogic/Configuration.h
new file mode 100644
--- /dev/null
+++ b/Engine/Plugins/Animation/RigLogic/Source/RigLogicLib/Private/riglogic/riglogic/Configuration.h
@@ -0,0 +1,17 @@
+#pragma once
+
+#include "riglogic/riglogic/CalculationType.h"
+
+namespace rl4 {
+
+struct Configuration {
+    CalculationType calculationType;
+
+    template<class Archive>
+    void serialize(Archive& archive) {
+        archive(calculationType);
+    }
+
+};
+
+}  // namespace rl4
diff --git a/Engine/Plugins/Animation/RigLogic/Source/RigLogicLib/Private/riglogic/riglogic/RigInstanceImpl.cpp b/Engine/Plugins/Animation/RigLogic/Source/RigLogicLib/Private/riglogic/riglogic/RigInstanceImpl.cpp
new file mode 100644
--- /dev/null
+++ b/Engine/Plugins/Animation/RigLogic/Source/RigLogicLib/Private/riglogic/riglogic/RigInstanceImpl.cpp
@@ -0,0 +1,146 @@
+#include "riglogic/riglogic/RigInstanceImpl.h"
+
+#include "riglogic/TypeDefs.h"
+#include "riglogic/riglogic/RigLogicImpl.h"
+#include "riglogic/riglogic/RigMetrics.h"
+#include "riglogic/types/Aliases.h"
+#include "riglogic/utils/Extd.h"
+
+#include <pma/PolyAllocator.h>
+#include <pma/resources/AlignedMemoryResource.h>
+
+#ifdef _MSC_VER
+    #pragma warning(push)
+    #pragma warning(disable : 4365 4987)
+#endif
+#include <algorithm>
+#include <cassert>
+#include <cstdint>
+#ifdef _MSC_VER
+    #pragma warning(pop)
+#endif
+
+namespace rl4 {
+
+namespace {
+
+const float controlsMin = 0.0f;
+const float controlsMax = 1.0f;
+
+}  // namespace
+
+RigInstance* RigInstance::create(RigLogic* rigLogic, MemoryResource* memRes) {
+    // NOLINTNEXTLINE(cppcoreguidelines-pro-type-static-cast-downcast)
+    auto ptr = static_cast<RigLogicImpl*>(rigLogic);
+    PolyAllocator<RigInstanceImpl> alloc{memRes};
+    return alloc.newObject(ptr->getRigMetrics(), memRes);
+}
+
+void RigInstance::destroy(RigInstance* instance) {
+    // NOLINTNEXTLINE(cppcoreguidelines-pro-type-static-cast-downcast)
+    auto ptr = static_cast<RigInstanceImpl*>(instance);
+    PolyAllocator<RigInstanceImpl> alloc{ptr->getMemoryResource()};
+    alloc.deleteObject(ptr);
+}
+
+RigInstance::~RigInstance() = default;
+
+RigInstanceImpl::RigInstanceImpl(const RigMetrics& metrics, MemoryResource* memRes_) :
+    memRes{memRes_},
+    lodCount{metrics.lodCount},
+    lodLevel{},
+    guiControlCount{metrics.guiControlCount},
+    rawControlCount{metrics.rawControlCount},
+    guiControlBuffer{guiControlCount, {}, memRes},
+    inputBuffer{static_cast<std::size_t>(rawControlCount + metrics.psdCount), {}, memRes},
+    jointOutputs{metrics.jointAttributeCount, {}, memRes},
+    blendShapeOutputs{metrics.blendShapeCount, {}, memRes},
+    animatedMapOutputs{metrics.animatedMapCount, {}, memRes} {
+}
+
+std::uint16_t RigInstanceImpl::getGUIControlCount() const {
+    return guiControlCount;
+}
+
+void RigInstanceImpl::setGUIControl(std::uint16_t index, float value) {
+    assert(index < guiControlBuffer.size());
+    guiControlBuffer[index] = value;
+}
+
+void RigInstanceImpl::setGUIControlValues(const float* values) {
+    std::copy(values, values + guiControlCount, guiControlBuffer.begin());
+}
+
+std::uint16_t RigInstanceImpl::getRawControlCount() const {
+    return rawControlCount;
+}
+
+void RigInstanceImpl::setRawControl(std::uint16_t index, float value) {
+    assert(index < inputBuffer.size());
+    inputBuffer[index] = extd::clamp(value, controlsMin, controlsMax);
+}
+
+void RigInstanceImpl::setRawControlValues(const float* values) {
+    auto clamper = [](float v) {
+            return extd::clamp(v, controlsMin, controlsMax);
+        };
+    std::transform(values, values + rawControlCount, inputBuffer.begin(), clamper);
+}
+
+std::uint16_t RigInstanceImpl::getLOD() const {
+    return lodLevel;
+}
+
+void RigInstanceImpl::setLOD(std::uint16_t level) {
+    lodLevel = extd::clamp(level, static_cast<std::uint16_t>(0), lodCount);
+}
+
+ArrayView<float> RigInstanceImpl::getGUIControlValues() {
+    return ArrayView<float>{guiControlBuffer};
+}
+
+ConstArrayView<float> RigInstanceImpl::getGUIControlValues() const {
+    return ConstArrayView<float>{guiControlBuffer};
+}
+
+ArrayView<float> RigInstanceImpl::getInputValues() {
+    return ArrayView<float>{inputBuffer};
+}
+
+ConstArrayView<float> RigInstanceImpl::getInputValues() const {
+    return ConstArrayView<float>{inputBuffer};
+}
+
+ArrayView<float> RigInstanceImpl::getRawJointOutputs() {
+    return ArrayView<float>{jointOutputs};
+}
+
+ConstArrayView<float> RigInstanceImpl::getRawJointOutputs() const {
+    return ConstArrayView<float>{jointOutputs};
+}
+
+TransformationArrayView RigInstanceImpl::getJointOutputs() const {
+    return TransformationArrayView{jointOutputs.data(), jointOutputs.size()};
+}
+
+ArrayView<float> RigInstanceImpl::getBlendShapeOutputs() {
+    return ArrayView<float>{blendShapeOutputs};
+}
+
+ConstArrayView<float> RigInstanceImpl::getBlendShapeOutputs() const {
+    return ConstArrayView<float>{blendShapeOutputs};
+}
+
+ArrayView<float> RigInstanceImpl::getAnimatedMapOutputs() {
+    return ArrayView<float>{animatedMapOutputs};
+}
+
+ConstArrayView<float> RigInstanceImpl::getAnimatedMapOutputs() const {
+    return ConstArrayView<float>{animatedMapOutputs};
+}
+
+MemoryResource* RigInstanceImpl::getMemoryResource() {
+    return memRes;
+}
+
+}  // namespace rl4
diff --git a/Engine/Plugins/Animation/RigLogic/Source/RigLogicLib/Private/riglogic/riglogic/RigInstanceImpl.h b/Engine/Plugins/Animation/RigLogic/Source/RigLogicLib/Private/riglogic/riglogic/RigInstanceImpl.h
new file mode 100644
--- /dev/null
+++ b/Engine/Plugins/Animation/RigLogic/Source/RigLogicLib/Private/riglogic/riglogic/RigInstanceImpl.h
@@ -0,0 +1,56 @@
+#pragma once
+
+#include "riglogic/TypeDefs.h"
+#include "riglogic/riglogic/RigInstance.h"
+#include "riglogic/riglogic/RigMetrics.h"
+
+#include <cstdint>
+
+namespace rl4 {
+
+class RigInstanceImpl : public RigInstance {
+    public:
+        RigInstanceImpl(const RigMetrics& metrics, MemoryResource* memRes_);
+
+        std::uint16_t getGUIControlCount() const override;
+        void setGUIControl(std::uint16_t index, float value) override;
+        void setGUIControlValues(const float* values) override;
+
+        std::uint16_t getRawControlCount() const override;
+        void setRawControl(std::uint16_t index, float value) override;
+        void setRawControlValues(const float* values) override;
+
+        std::uint16_t getLOD() const override;
+        void setLOD(std::uint16_t level) override;
+
+        ArrayView<float> getGUIControlValues();
+        ConstArrayView<float> getGUIControlValues() const;
+        ArrayView<float> getInputValues();
+        ConstArrayView<float> getInputValues() const;
+        ArrayView<float> getRawJointOutputs();
+        ConstArrayView<float> getRawJointOutputs() const override;
+        TransformationArrayView getJointOutputs() const override;
+        ArrayView<float> getBlendShapeOutputs();
+        ConstArrayView<float> getBlendShapeOutputs() const override;
+        ArrayView<float> getAnimatedMapOutputs();
+        ConstArrayView<float> getAnimatedMapOutputs() const override;
+
+        MemoryResource* getMemoryResource();
+
+    private:
+        MemoryResource* memRes;
+
+        std::uint16_t lodCount;
+        std::uint16_t lodLevel;
+        std::uint16_t guiControlCount;
+        std::uint16_t rawControlCount;
+
+        Vector<float> guiControlBuffer;
+        Vector<float> inputBuffer;
+        AlignedVector<float> jointOutputs;
+        Vector<float> blendShapeOutputs;
+        Vector<float> animatedMapOutputs;
+
+};
+
+}  // namespace rl4
diff --git a/Engine/Plugins/Animation/RigLogic/Source/RigLogicLib/Private/riglogic/riglogic/RigLogicImpl.cpp b/Engine/Plugins/Animation/RigLogic/Source/RigLogicLib/Private/riglogic/riglogic/RigLogicImpl.cpp
new file mode 100644
--- /dev/null
+++ b/Engine/Plugins/Animation/RigLogic/Source/RigLogicLib/Private/riglogic/riglogic/RigLogicImpl.cpp
@@ -0,0 +1,179 @@
+#include "riglogic/riglogic/RigLogicImpl.h"
+
+#include "riglogic/animatedmaps/AnimatedMapsFactory.h"
+#include "riglogic/blendshapes/BlendShapesFactory.h"
+#include "riglogic/controls/ControlsFactory.h"
+#include "riglogic/joints/JointsFactory.h"
+#include "riglogic/riglogic/RigInstanceImpl.h"
+#include "riglogic/riglogic/RigMetrics.h"
+
+#include <dna/layers/BehaviorReader.h>
+#include <pma/PolyAllocator.h>
+#include <pma/resources/AlignedMemoryResource.h>
+#include <terse/archives/binary/InputArchive.h>
+#include <terse/archives/binary/OutputArchive.h>
+#include <trio/Stream.h>
+
+#include <cstddef>
+#include <utility>
+
+namespace rl4 {
+
+static RigInstanceImpl* castInstance(RigInstance* instance) {
+    // NOLINTNEXTLINE(cppcoreguidelines-pro-type-static-cast-downcast)
+    return static_cast<RigInstanceImpl*>(instance);
+}
+
+RigLogic::~RigLogic() = default;
+
+RigLogic* RigLogic::create(const dna::BehaviorReader* reader, CalculationType calculationType, MemoryResource* memRes) {
+    RigMetrics metrics{};
+    metrics.lodCount = reader->getLODCount();
+    metrics.guiControlCount = reader->getGUIControlCount();
+    metrics.rawControlCount = reader->getRawControlCount();
+    metrics.psdCount = reader->getPSDCount();
+    metrics.jointAttributeCount = reader->getJointRowCount();
+    metrics.blendShapeCount = reader->getBlendShapeChannelCount();
+    metrics.animatedMapCount = reader->getAnimatedMapCount();
+
+    Configuration config{calculationType};
+    auto controls = ControlsFactory::create(reader, memRes);
+    auto joints = JointsFactory::create(config, reader, memRes);
+    auto blendShapes = BlendShapesFactory::create(reader, memRes);
+    auto animatedMaps = AnimatedMapsFactory::create(reader, memRes);
+
+    PolyAllocator<RigLogicImpl> alloc{memRes};
+    return alloc.newObject(config,
+                           metrics,
+                           std::move(controls),
+                           std::move(joints),
+                           std::move(blendShapes),
+                           std::move(animatedMaps),
+                           memRes);
+}
+
+RigLogic* RigLogic::restore(BoundedIOStream* source, MemoryResource* memRes) {
+    PolyAllocator<RigLogicImpl> alloc{memRes};
+
+    terse::BinaryInputArchive<BoundedIOStream> archive{source};
+
+    Configuration config;
+    archive >> config;
+
+    RigMetrics metrics;
+    archive >> metrics;
+
+    auto controls = ControlsFactory::create(memRes);
+    auto joints = JointsFactory::create(config, metrics, memRes);
+    auto blendShapes = BlendShapesFactory::create(memRes);
+    auto animatedMaps = AnimatedMapsFactory::create(memRes);
+
+    archive >> *controls >> *joints >> *blendShapes >> *animatedMaps;
+    return alloc.newObject(config,
+                           metrics,
+                           std::move(controls),
+                           std::move(joints),
+                           std::move(blendShapes),
+                           std::move(animatedMaps),
+                           memRes);
+}
+
+void RigLogic::destroy(RigLogic* instance) {
+    // NOLINTNEXTLINE(cppcoreguidelines-pro-type-static-cast-downcast)
+    auto ptr = static_cast<RigLogicImpl*>(instance);
+    PolyAllocator<RigLogicImpl> alloc{ptr->getMemoryResource()};
+    alloc.deleteObject(ptr);
+}
+
+RigLogicImpl::RigLogicImpl(Configuration config_,
+                           RigMetrics metrics_,
+                           ControlsPtr controls_,
+                           JointsPtr joints_,
+                           BlendShapesPtr blendShapes_,
+                           AnimatedMapsPtr animatedMaps_,
+                           MemoryResource* memRes_) :
+    memRes{memRes_},
+    config{config_},
+    metrics{metrics_},
+    controls{std::move(controls_)},
+    joints{std::move(joints_)},
+    blendShapes{std::move(blendShapes_)},
+    animatedMaps{std::move(animatedMaps_)} {
+}
+
+void RigLogicImpl::dump(BoundedIOStream* destination) const {
+    terse::BinaryOutputArchive<BoundedIOStream> archive{destination};
+    archive << config << metrics << *controls << *joints << *blendShapes << *animatedMaps;
+}
+
+const Configuration& RigLogicImpl::getConfiguration() const {
+    return config;
+}
+
+const RigMetrics& RigLogicImpl::getRigMetrics() const {
+    return metrics;
+}
+
+std::uint16_t RigLogicImpl::getLODCount() const {
+    return metrics.lodCount;
+}
+
+ConstArrayView<float> RigLogicImpl::getRawNeutralJointValues() const {
+    return joints->getRawNeutralValues();
+}
+
+TransformationArrayView RigLogicImpl::getNeutralJointValues() const {
+    return joints->getNeutralValues();
+}
+
+std::uint16_t RigLogicImpl::getJointGroupCount() const {
+    return joints->getJointGroupCount();
+}
+
+ConstArrayView<std::uint16_t> RigLogicImpl::getJointVariableAttributeIndices(std::uint16_t lod) const {
+    return joints->getVariableAttributeIndices(lod);
+}
+
+void RigLogicImpl::mapGUIToRawControls(RigInstance* instance) const {
+    auto pRigInstance = castInstance(instance);
+    controls->mapGUIToRaw(pRigInstance->getGUIControlValues(), pRigInstance->getInputValues());
+}
+
+void RigLogicImpl::calculateControls(RigInstance* instance) const {
+    auto pRigInstance = castInstance(instance);
+    controls->calculate(pRigInstance->getInputValues());
+}
+
+void RigLogicImpl::calculateJoints(RigInstance* instance) const {
+    auto pRigInstance = castInstance(instance);
+    joints->calculate(pRigInstance->getInputValues(), pRigInstance->getRawJointOutputs(), pRigInstance->getLOD());
+}
+
+void RigLogicImpl::calculateJoints(RigInstance* instance, std::uint16_t jointGroupIndex) const {
+    auto pRigInstance = castInstance(instance);
+    joints->calculate(pRigInstance->getInputValues(), pRigInstance->getRawJointOutputs(), pRigInstance->getLOD(),
+                      jointGroupIndex);
+}
+
+void RigLogicImpl::calculateBlendShapes(RigInstance* instance) const {
+    auto pRigInstance = castInstance(instance);
+    blendShapes->calculate(pRigInstance->getInputValues(), pRigInstance->getBlendShapeOutputs(), pRigInstance->getLOD());
+}
+
+void RigLogicImpl::calculateAnimatedMaps(RigInstance* instance) const {
+    auto pRigInstance = castInstance(instance);
+    animatedMaps->calculate(pRigInstance->getInputValues(), pRigInstance->getAnimatedMapOutputs(), pRigInstance->getLOD());
+}
+
+void RigLogicImpl::calculate(RigInstance* instance) const {
+    calculateControls(instance);
+    calculateJoints(instance);
+    calculateBlendShapes(instance);
+    calculateAnimatedMaps(instance);
+}
+
+MemoryResource* RigLogicImpl::getMemoryResource() {
+    return memRes;
+}
+
+}  // namespace rl4
diff --git a/Engine/Plugins/Animation/RigLogic/Source/RigLogicLib/Private/riglogic/riglogic/RigLogicImpl.h b/Engine/Plugins/Animation/RigLogic/Source/RigLogicLib/Private/riglogic/riglogic/RigLogicImpl.h
new file mode 100644
--- /dev/null
+++ b/Engine/Plugins/Animation/RigLogic/Source/RigLogicLib/Private/riglogic/riglogic/RigLogicImpl.h
@@ -0,0 +1,67 @@
+#pragma once
+
+#include "riglogic/animatedmaps/AnimatedMaps.h"
+#include "riglogic/blendshapes/BlendShapes.h"
+#include "riglogic/controls/Controls.h"
+#include "riglogic/joints/Joints.h"
+#include "riglogic/riglogic/Configuration.h"
+#include "riglogic/riglogic/RigLogic.h"
+#include "riglogic/riglogic/RigMetrics.h"
+
+#ifdef _MSC_VER
+    #pragma warning(push)
+    #pragma warning(disable : 4365 4987)
+#endif
+#include <functional>
+#include <memory>
+#ifdef _MSC_VER
+    #pragma warning(pop)
+#endif
+
+namespace rl4 {
+
+class RigLogicImpl : public RigLogic {
+    private:
+        using ControlsPtr = std::unique_ptr<Controls, std::function<void (Controls*)> >;
+        using JointsPtr = std::unique_ptr<Joints, std::function<void (Joints*)> >;
+        using BlendShapesPtr = std::unique_ptr<BlendShapes, std::function<void (BlendShapes*)> >;
+        using AnimatedMapsPtr = std::unique_ptr<AnimatedMaps, std::function<void (AnimatedMaps*)> >;
+
+    public:
+        RigLogicImpl(Configuration config_,
+                     RigMetrics metrics_,
+                     ControlsPtr controls_,
+                     JointsPtr joints_,
+                     BlendShapesPtr blendShapes_,
+                     AnimatedMapsPtr animatedMaps_,
+                     MemoryResource* memRes_);
+
+        void dump(BoundedIOStream* destination) const override;
+        const Configuration& getConfiguration() const;
+        const RigMetrics& getRigMetrics() const;
+        std::uint16_t getLODCount() const override;
+        ConstArrayView<float> getRawNeutralJointValues() const override;
+        TransformationArrayView getNeutralJointValues() const override;
+        std::uint16_t getJointGroupCount() const override;
+        ConstArrayView<std::uint16_t> getJointVariableAttributeIndices(std::uint16_t lod) const override;
+        void mapGUIToRawControls(RigInstance* instance) const override;
+        void calculateControls(RigInstance* instance) const override;
+        void calculateJoints(RigInstance* instance) const override;
+        void calculateJoints(RigInstance* instance, std::uint16_t jointGroupIndex) const override;
+        void calculateBlendShapes(RigInstance* instance) const override;
+        void calculateAnimatedMaps(RigInstance* instance) const override;
+        void calculate(RigInstance* instance) const override;
+
+        MemoryResource* getMemoryResource();
+
+    private:
+        MemoryResource* memRes;
+        Configuration config;
+        RigMetrics metrics;
+        ControlsPtr controls;
+        JointsPtr joints;
+        BlendShapesPtr blendShapes;
+        AnimatedMapsPtr animatedMaps;
+};
+
+}  // namespace rl4
diff --git a/Engine/Plugins/Animation/RigLogic/Source/RigLogicLib/Private/riglogic/riglogic/RigMetrics.h b/Engine/Plugins/Animation/RigLogic/Source/RigLogicLib/Private/riglogic/riglogic/RigMetrics.h
new file mode 100644
--- /dev/null
+++ b/Engine/Plugins/Animation/RigLogic/Source/RigLogicLib/Private/riglogic/riglogic/RigMetrics.h
@@ -0,0 +1,29 @@
+#pragma once
+
+#include <cstdint>
+
+namespace rl4 {
+
+struct RigMetrics {
+    std::uint16_t lodCount;
+    std::uint16_t guiControlCount;
+    std::uint16_t rawControlCount;
+    std::uint16_t psdCount;
+    std::uint16_t jointAttributeCount;
+    std::uint16_t blendShapeCount;
+    std::uint16_t animatedMapCount;
+
+    template<class Archive>
+    void serialize(Archive& archive) {
+        archive(lodCount,
+                guiControlCount,
+                rawControlCount,
+                psdCount,
+                jointAttributeCount,
+                blendShapeCount,
+                animatedMapCount);
+    }
+
+};
+
+}  // namespace rl4
diff --git a/Engine/Plugins/Animation/RigLogic/Source/RigLogicLib/Private/riglogic/system/simd/SIMD.h b/Engine/Plugins/Animation/RigLogic/Source/RigLogicLib/Private/riglogic/system/simd/SIMD.h
new file mode 100644
--- /dev/null
+++ b/Engine/Plugins/Animation/RigLogic/Source/RigLogicLib/Private/riglogic/system/simd/SIMD.h
@@ -0,0 +1,9 @@
+#pragma once
+
+#ifdef RL_BUILD_WITH_AVX
+    #include "riglogic/system/simd/AVX.h"
+#endif
+#ifdef RL_BUILD_WITH_SSE
+    #include "riglogic/system/simd/SSE.h"
+#endif
+#include "riglogic/system/simd/Scalar.h"
diff --git a/Engine/Plugins/Animation/RigLogic/Source/RigLogicLib/Private/riglogic/version/VersionInfo.cpp b/Engine/Plugins/Animation/RigLogic/Source/RigLogicLib/Private/riglogic/version/VersionInfo.cpp
new file mode 100644
--- /dev/null
+++ b/Engine/Plugins/Animation/RigLogic/Source/RigLogicLib/Private/riglogic/version/VersionInfo.cpp
@@ -0,0 +1,34 @@
+#include "riglogic/version/VersionInfo.h"
+
+#include "riglogic/version/Version.h"
+
+#include <cstring>
+
+namespace rl4 {
+
+namespace {
+
+constexpr int majorVersion = RL_MAJOR_VERSION;
+constexpr int minorVersion = RL_MINOR_VERSION;
+constexpr int patchVersion = RL_PATCH_VERSION;
+constexpr const char* versionString = RL_VERSION_STRING;
+
+}  // namespace
+
+int VersionInfo::getMajorVersion() {
+    return majorVersion;
+}
+
+int VersionInfo::getMinorVersion() {
+    return minorVersion;
+}
+
+int VersionInfo::getPatchVersion() {
+    return patchVersion;
+}
+
+StringView VersionInfo::getVersionString() {
+    return {versionString, std::strlen(versionString)};
+}
+
+}  // namespace rl4
diff --git a/Engine/Plugins/Animation/RigLogic/Source/RigLogicLib/Private/trio/streams/FileStreamImpl.cpp b/Engine/Plugins/Animation/RigLogic/Source/RigLogicLib/Private/trio/streams/FileStreamImpl.cpp
new file mode 100644
--- /dev/null
+++ b/Engine/Plugins/Animation/RigLogic/Source/RigLogicLib/Private/trio/streams/FileStreamImpl.cpp
@@ -0,0 +1,117 @@
+#include "trio/streams/FileStreamImpl.h"
+
+#include "trio/utils/ScopedEnumEx.h"
+
+#include <pma/PolyAllocator.h>
+#include <status/Provider.h>
+
+#ifdef _MSC_VER
+    #pragma warning(push)
+    #pragma warning(disable : 4365 4987)
+#endif
+#include <cassert>
+#include <cstddef>
+#include <cstring>
+#include <ios>
+#include <type_traits>
+#ifdef _MSC_VER
+    #pragma warning(pop)
+#endif
+
+namespace trio {
+
+const sc::StatusCode FileStream::OpenError{100, "Error opening file: %s"};
+const sc::StatusCode FileStream::ReadError{101, "Error reading file: %s"};
+const sc::StatusCode FileStream::WriteError{102, "Error writing file: %s"};
+const sc::StatusCode FileStream::AlreadyOpenError{103, "File already open: %s"};
+
+#ifdef __clang__
+    #pragma clang diagnostic push
+    #pragma clang diagnostic ignored "-Wglobal-constructors"
+#endif
+sc::StatusProvider FileStreamImpl::status{OpenError, ReadError, WriteError, AlreadyOpenError};
+#ifdef __clang__
+    #pragma clang diagnostic pop
+#endif
+
+FileStream::~FileStream() = default;
+
+FileStream* FileStream::create(const char* path, AccessMode accessMode, OpenMode openMode, MemoryResource* memRes) {
+    pma::PolyAllocator<FileStreamImpl> alloc{memRes};
+    return alloc.newObject(path, accessMode, openMode, memRes);
+}
+
+void FileStream::destroy(FileStream* instance) {
+    // NOLINTNEXTLINE(cppcoreguidelines-pro-type-static-cast-downcast)
+    auto fileStream = static_cast<FileStreamImpl*>(instance);
+    pma::PolyAllocator<FileStreamImpl> alloc{fileStream->getMemoryResource()};
+    alloc.deleteObject(fileStream);
+}
+
+FileStreamImpl::FileStreamImpl(const char* path_, AccessMode accessMode_, OpenMode openMode_, MemoryResource* memRes_) :
+    path{path_, memRes_},
+    accessMode{accessMode_},
+    openMode{openMode_},
+    fileSize{},
+    memRes{memRes_} {
+}
+
+void FileStreamImpl::open() {
+    status.reset();
+    if (file.is_open()) {
+        status.set(AlreadyOpenError, path.c_str());
+        return;
+    }
+
+    std::ios_base::openmode flags{};
+    flags |= (contains(accessMode, AccessMode::Read) ? std::ios_base::in : flags);
+    flags |= (contains(accessMode, AccessMode::Write) ? std::ios_base::out : flags);
+    flags |= (contains(openMode, OpenMode::Binary) ? std::ios_base::binary : flags);
+    flags |= std::ios_base::ate;
+
+    file.open(path.c_str(), flags);
+    if (!file.good()) {
+        status.set(OpenError, path.c_str());
+        return;
+    }
+    fileSize = static_cast<std::size_t>(file.tellg());
+    seek(0ul);
+}
+
+void FileStreamImpl::close() {
+    file.close();
+}
+
+std::size_t FileStreamImpl::tell() {
+    return static_cast<std::size_t>(file.tellp());
+}
+
+void FileStreamImpl::seek(std::size_t position) {
+    file.seekp(static_cast<std::streamoff>(position));
+}
+
+void FileStreamImpl::read(char* buffer, std::size_t size) {
+    status.reset();
+    file.read(buffer, static_cast<std::streamsize>(size));
+    if (!file.good() && !file.eof()) {
+        status.set(ReadError, path.c_str());
+    }
+}
+
+void FileStreamImpl::write(const char* buffer, std::size_t size) {
+    status.reset();
+    file.write(buffer, static_cast<std::streamsize>(size));
+    if (!file.good()) {
+        status.set(WriteError, path.c_str());
+    }
+}
+
+std::size_t FileStreamImpl::size() {
+    return fileSize;
+}
+
+MemoryResource* FileStreamImpl::getMemoryResource() {
+    return memRes;
+}
+
+}  // namespace trio
diff --git a/Engine/Plugins/Animation/RigLogic/Source/RigLogicLib/Private/trio/streams/FileStreamImpl.h b/Engine/Plugins/Animation/RigLogic/Source/RigLogicLib/Private/trio/streams/FileStreamImpl.h
new file mode 100644
--- /dev/null
+++ b/Engine/Plugins/Animation/RigLogic/Source/RigLogicLib/Private/trio/streams/FileStreamImpl.h
@@ -0,0 +1,48 @@
+#pragma once
+
+#include "trio/streams/FileStream.h"
+#include "trio/types/Aliases.h"
+
+#include <pma/TypeDefs.h>
+#include <status/Provider.h>
+
+#ifdef _MSC_VER
+    #pragma warning(push)
+    #pragma warning(disable : 4365 4987)
+#endif
+#include <cstddef>
+#include <fstream>
+#include <ios>
+#include <string>
+#ifdef _MSC_VER
+    #pragma warning(pop)
+#endif
+
+namespace trio {
+
+class FileStreamImpl : public FileStream {
+    public:
+        FileStreamImpl(const char* path_, AccessMode accessMode_, OpenMode openMode_, MemoryResource* memRes_);
+
+        void open() override;
+        void close() override;
+        std::size_t tell() override;
+        void seek(std::size_t position) override;
+        std::size_t size() override;
+        void read(char* buffer, std::size_t size) override;
+        void write(const char* buffer, std::size_t size) override;
+
+        MemoryResource* getMemoryResource();
+
+    private:
+        static sc::StatusProvider status;
+
+        std::fstream file;
+        pma::String<char> path;
+        AccessMode accessMode;
+        OpenMode openMode;
+        std::size_t fileSize;
+        MemoryResource* memRes;
+};
+
+}  // namespace trio
diff --git a/Engine/Plugins/Animation/RigLogic/Source/RigLogicLib/Private/trio/utils/ScopedEnumEx.h b/Engine/Plugins/Animation/RigLogic/Source/RigLogicLib/Private/trio/utils/ScopedEnumEx.h
new file mode 100644
--- /dev/null
+++ b/Engine/Plugins/Animation/RigLogic/Source/RigLogicLib/Private/trio/utils/ScopedEnumEx.h
@@ -0,0 +1,59 @@
+#pragma once
+
+#include <type_traits>
+
+namespace trio {
+
+template<typename TEnum>
+typename std::enable_if<std::is_enum<TEnum>::value, TEnum>::type
+operator&(TEnum lhs, TEnum rhs) {
+    using Underlying = typename std::underlying_type<TEnum>::type;
+    return static_cast<TEnum>(static_cast<Underlying>(lhs) & static_cast<Underlying>(rhs));
+}
+
+template<typename TEnum>
+typename std::enable_if<std::is_enum<TEnum>::value, TEnum>::type
+operator|(TEnum lhs, TEnum rhs) {
+    using Underlying = typename std::underlying_type<TEnum>::type;
+    return static_cast<TEnum>(static_cast<Underlying>(lhs) | static_cast<Underlying>(rhs));
+}
+
+template<typename TEnum>
+typename std::enable_if<std::is_enum<TEnum>::value, TEnum>::type
+operator^(TEnum lhs, TEnum rhs) {
+    using Underlying = typename std::underlying_type<TEnum>::type;
+    return static_cast<TEnum>(static_cast<Underlying>(lhs) ^ static_cast<Underlying>(rhs));
+}
+
+template<typename TEnum>
+typename std::enable_if<std::is_enum<TEnum>::value, TEnum>::type
+operator~(TEnum value) {
+    using Underlying = typename std::underlying_type<TEnum>::type;
+    return static_cast<TEnum>(~static_cast<Underlying>(value));
+}
+
+template<typename TEnum>
+typename std::enable_if<std::is_enum<TEnum>::value, TEnum>::type
+operator&=(TEnum& lhs, TEnum rhs) {
+    return lhs = (lhs & rhs);
+}
+
+template<typename TEnum>
+typename std::enable_if<std::is_enum<TEnum>::value, TEnum>::type
+operator|=(TEnum& lhs, TEnum rhs) {
+    return lhs = (lhs | rhs);
+}
+
+template<typename TEnum>
+typename std::enable_if<std::is_enum<TEnum>::value, TEnum>::type
+operator^=(TEnum& lhs, TEnum rhs) {
+    return lhs = (lhs ^ rhs);
+}
+
+template<typename TEnum>
+typename std::enable_if<std::is_enum<TEnum>::value, bool>::type
+contains(TEnum lhs, TEnum rhs) {
+    return (lhs & rhs) == rhs;
+}
+
+}  // namespace trio
diff --git a/Engine/Plugins/Animation/RigLogic/Source/RigLogicLib/Public/FMemoryResource.h b/Engine/Plugins/Animation/RigLogic/Source/RigLogicLib/Public/FMemoryResource.h
new file mode 100644
--- /dev/null
+++ b/Engine/Plugins/Animation/RigLogic/Source/RigLogicLib/Public/FMemoryResource.h
@@ -0,0 +1,25 @@
+// Copyright 1998-2019 Epic Games, Inc. All Rights Reserved.
+#pragma once
+
+#include "HAL/UnrealMemory.h"
+
+#include "riglogic/RigLogic.h"
+
+class FMemoryResource : public rl4::MemoryResource
+{
+public:
+	void* allocate(std::size_t size, std::size_t alignment) override
+	{
+		return FMemory::Malloc(size, alignment);
+	}
+
+	void deallocate(void* ptr, std::size_t size, std::size_t alignment) override
+	{
+		FMemory::Free(ptr);
+	}
+
+	static MemoryResource* Instance() {
+		static FMemoryResource MemoryResource;
+		return &MemoryResource;
+	}
+};
diff --git a/Engine/Plugins/Animation/RigLogic/Source/RigLogicLib/Public/RigLogicLib.h b/Engine/Plugins/Animation/RigLogic/Source/RigLogicLib/Public/RigLogicLib.h
new file mode 100644
--- /dev/null
+++ b/Engine/Plugins/Animation/RigLogic/Source/RigLogicLib/Public/RigLogicLib.h
@@ -0,0 +1,18 @@
+// Copyright 1998-2019 Epic Games, Inc. All Rights Reserved.
+#pragma once
+
+#ifndef RIGLOGIC_MODULE_DISCARD
+
+#include "CoreMinimal.h"
+#include "Modules/ModuleInterface.h"
+
+DECLARE_LOG_CATEGORY_EXTERN(LogRigLogicLib, Log, All);
+
+class FRigLogicLib : public IModuleInterface
+{
+public:
+	virtual void StartupModule() override;
+	virtual void ShutdownModule() override;
+};
+
+#endif  // RIGLOGIC_MODULE_DISCARD
diff --git a/Engine/Plugins/Animation/RigLogic/Source/RigLogicLib/Public/dna/DataLayer.h b/Engine/Plugins/Animation/RigLogic/Source/RigLogicLib/Public/dna/DataLayer.h
new file mode 100644
--- /dev/null
+++ b/Engine/Plugins/Animation/RigLogic/Source/RigLogicLib/Public/dna/DataLayer.h
@@ -0,0 +1,15 @@
+#pragma once
+
+namespace dna {
+
+enum class DataLayer {
+    Descriptor,
+    Definition,  // Includes Descriptor
+    Behavior,  // Includes Descriptor and Definition
+    Geometry,  // Includes Descriptor and Definition
+    GeometryWithoutBlendShapes,  // Includes Descriptor and Definition
+    AllWithoutBlendShapes,  // Includes everything except blend shapes from Geometry
+    All
+};
+
+}  // namespace dna
diff --git a/Engine/Plugins/Animation/RigLogic/Source/RigLogicLib/Public/dna/StreamReader.h b/Engine/Plugins/Animation/RigLogic/Source/RigLogicLib/Public/dna/StreamReader.h
new file mode 100644
--- /dev/null
+++ b/Engine/Plugins/Animation/RigLogic/Source/RigLogicLib/Public/dna/StreamReader.h
@@ -0,0 +1,140 @@
+#pragma once
+
+#include "dna/DataLayer.h"
+#include "dna/Defs.h"
+#include "dna/Reader.h"
+#include "dna/types/Aliases.h"
+
+namespace dna {
+
+class DNAAPI StreamReader : public Reader {
+    public:
+        static const sc::StatusCode SignatureMismatchError;
+        static const sc::StatusCode VersionMismatchError;
+        static const sc::StatusCode InvalidDataError;
+
+    public:
+        /**
+            @brief Factory method for creation of Reader
+            @param stream
+                Source stream from which data is going to be read.
+            @param layer
+                Specify the layer up to which the data needs to be loaded.
+            @note
+                The Definition data layer depends on and thus implicitly loads the Descriptor layer.
+                The Behavior data layer depends on and thus implicitly loads the Definition layer.
+                The Geometry data layer depends on and thus also implicitly loads the Definition layer.
+            @param maxLOD
+                The maximum level of details to be loaded.
+            @note
+                A value of zero indicates to load all LODs.
+            @warning
+                The maxLOD value must be less than the value returned by getLODCount.
+            @see getLODCount
+            @param memRes
+                Memory resource to be used for allocations.
+            @note
+                If a memory resource is not given, a default allocation mechanism will be used.
+            @warning
+                User is responsible for releasing the returned pointer by calling destroy.
+            @see destroy
+        */
+        static StreamReader* create(BoundedIOStream* stream,
+                                    DataLayer layer = DataLayer::All,
+                                    std::uint16_t maxLOD = 0u,
+                                    MemoryResource* memRes = nullptr);
+        /**
+            @brief Factory method for creation of Reader
+            @param stream
+                Source stream from which data is going to be read.
+            @param layer
+                Specify the layer up to which the data needs to be loaded.
+            @note
+                The Definition data layer depends on and thus implicitly loads the Descriptor layer.
+                The Behavior data layer depends on and thus implicitly loads the Definition layer.
+                The Geometry data layer depends on and thus also implicitly loads the Definition layer.
+            @param maxLOD
+                The maximum level of details to be loaded.
+            @param minLOD
+                The minimum level of details to be loaded.
+            @note
+                A range of [0, LOD count - 1] for maxLOD / minLOD respectively indicates to load all LODs.
+            @warning
+                Both maxLOD and minLOD values must be less than the value returned by getLODCount.
+            @see getLODCount
+            @param memRes
+                Memory resource to be used for allocations.
+            @note
+                If a memory resource is not given, a default allocation mechanism will be used.
+            @warning
+                User is responsible for releasing the returned pointer by calling destroy.
+            @see destroy
+        */
+        static StreamReader* create(BoundedIOStream* stream,
+                                    DataLayer layer,
+                                    std::uint16_t maxLOD,
+                                    std::uint16_t minLOD,
+                                    MemoryResource* memRes = nullptr);
+        /**
+            @brief Factory method for creation of Reader
+            @param stream
+                Source stream from which data is going to be read.
+            @param layer
+                Specify the layer up to which the data needs to be loaded.
+            @note
+                The Definition data layer depends on and thus implicitly loads the Descriptor layer.
+                The Behavior data layer depends on and thus implicitly loads the Definition layer.
+                The Geometry data layer depends on and thus also implicitly loads the Definition layer.
+            @param lods
+                An array specifying which exact lods to load.
+            @warning
+                All values in the array must be less than the value returned by getLODCount.
+            @see getLODCount
+            @param lodCount
+                The number of elements in the lods array.
+            @warning
+                There cannot be more elements in the array than the value returned by getLODCount.
+            @see getLODCount
+            @param memRes
+                Memory resource to be used for allocations.
+            @note
+                If a memory resource is not given, a default allocation mechanism will be used.
+            @warning
+                User is responsible for releasing the returned pointer by calling destroy.
+            @see destroy
+        */
+        static StreamReader* create(BoundedIOStream* stream,
+                                    DataLayer layer,
+                                    std::uint16_t* lods,
+                                    std::uint16_t lodCount,
+                                    MemoryResource* memRes = nullptr);
+        /**
+            @brief Method for freeing a Reader instance.
+            @param instance
+                Instance of Reader to be freed.
+            @see create
+        */
+        static void destroy(StreamReader* instance);
+
+        ~StreamReader() override;
+        /**
+           @brief read data from stream into internal structures.
+        */
+        virtual void read() = 0;
+};
+
+}  // namespace dna
+
+namespace pma {
+
+template<>
+struct DefaultInstanceCreator<dna::StreamReader> {
+    using type = pma::FactoryCreate<dna::StreamReader>;
+};
+
+template<>
+struct DefaultInstanceDestroyer<dna::StreamReader> {
+    using type = pma::FactoryDestroy<dna::StreamReader>;
+};
+
+}  // namespace pma
diff --git a/Engine/Plugins/Animation/RigLogic/Source/RigLogicLib/Public/dna/StreamWriter.h b/Engine/Plugins/Animation/RigLogic/Source/RigLogicLib/Public/dna/StreamWriter.h
new file mode 100644
--- /dev/null
+++ b/Engine/Plugins/Animation/RigLogic/Source/RigLogicLib/Public/dna/StreamWriter.h
@@ -0,0 +1,53 @@
+#pragma once
+
+#include "dna/Defs.h"
+#include "dna/Writer.h"
+#include "dna/types/Aliases.h"
+
+namespace dna {
+
+class DNAAPI StreamWriter : public Writer {
+    public:
+        /**
+            @brief Factory method for creation of Writer
+            @param stream
+                Stream into which the data is going to be written.
+            @param memRes
+                Memory resource to be used for allocations.
+            @note
+                If a memory resource is not given, a default allocation mechanism will be used.
+            @warning
+                User is responsible for releasing the returned pointer by calling destroy.
+            @see destroy
+        */
+        static StreamWriter* create(BoundedIOStream* stream, MemoryResource* memRes = nullptr);
+        /**
+            @brief Method for freeing a Writer instance.
+            @param instance
+                Instance of Writer to be freed.
+            @see create
+        */
+        static void destroy(StreamWriter* instance);
+
+        ~StreamWriter() override;
+        /**
+            @brief Write data to stream from internal structures.
+         */
+        virtual void write() = 0;
+};
+
+}  // namespace dna
+
+namespace pma {
+
+template<>
+struct DefaultInstanceCreator<dna::StreamWriter> {
+    using type = pma::FactoryCreate<dna::StreamWriter>;
+};
+
+template<>
+struct DefaultInstanceDestroyer<dna::StreamWriter> {
+    using type = pma::FactoryDestroy<dna::StreamWriter>;
+};
+
+}  // namespace pma
