Editgrating RigLogic plugin from Dev-CharacterTech to UE5/Main @ CL 13891876.

diff --git a/Engine/Plugins/Animation/RigLogic/Source/RigLogicLib/Public/dna/layers/Descriptor.h b/Engine/Plugins/Animation/RigLogic/Source/RigLogicLib/Public/dna/layers/Descriptor.h
new file mode 100644
--- /dev/null
+++ b/Engine/Plugins/Animation/RigLogic/Source/RigLogicLib/Public/dna/layers/Descriptor.h
@@ -0,0 +1,45 @@
+#pragma once
+
+namespace dna {
+
+enum class Archetype {
+    asian,
+    black,
+    caucasian,
+    hispanic,
+    alien,
+    other
+};
+
+enum class Gender {
+    male,
+    female,
+    other
+};
+
+enum class TranslationUnit {
+    cm,
+    m
+};
+
+enum class RotationUnit {
+    degrees,
+    radians
+};
+
+enum class Direction {
+    left,
+    right,
+    up,
+    down,
+    front,
+    back
+};
+
+struct CoordinateSystem {
+    Direction xAxis;
+    Direction yAxis;
+    Direction zAxis;
+};
+
+}  // namespace dna
diff --git a/Engine/Plugins/Animation/RigLogic/Source/RigLogicLib/Public/dna/types/Aliases.h b/Engine/Plugins/Animation/RigLogic/Source/RigLogicLib/Public/dna/types/Aliases.h
new file mode 100644
--- /dev/null
+++ b/Engine/Plugins/Animation/RigLogic/Source/RigLogicLib/Public/dna/types/Aliases.h
@@ -0,0 +1,27 @@
+#pragma once
+
+#include "dna/types/ArrayView.h"
+#include "dna/types/StringView.h"
+
+#include <pma/MemoryResource.h>
+#include <pma/ScopedPtr.h>
+#include <status/Status.h>
+#include <status/StatusCode.h>
+#include <trio/Stream.h>
+#include <trio/streams/FileStream.h>
+
+namespace dna {
+
+template<typename T>
+using ArrayView = trust::ArrayView<T>;
+
+template<typename T>
+using ConstArrayView = trust::ConstArrayView<T>;
+
+using trio::BoundedIOStream;
+using trio::FileStream;
+using sc::Status;
+
+using namespace pma;
+
+}  // namespace dna
diff --git a/Engine/Plugins/Animation/RigLogic/Source/RigLogicLib/Public/dna/types/StringView.h b/Engine/Plugins/Animation/RigLogic/Source/RigLogicLib/Public/dna/types/StringView.h
new file mode 100644
--- /dev/null
+++ b/Engine/Plugins/Animation/RigLogic/Source/RigLogicLib/Public/dna/types/StringView.h
@@ -0,0 +1,41 @@
+#pragma once
+
+#include "dna/types/Aliases.h"
+
+namespace dna {
+
+class StringView : public trust::ConstArrayView<char> {
+    public:
+        using Base = trust::ConstArrayView<char>;
+
+    public:
+        using Base::ArrayView;
+
+        const char* c_str() const {
+            return dataOrEmpty();
+        }
+
+        operator const char*() const {
+            return dataOrEmpty();
+        }
+
+        const char* operator*() const {
+            return dataOrEmpty();
+        }
+
+    private:
+        const char* dataOrEmpty() const {
+            return (data() == nullptr ? "" : data());
+        }
+
+};
+
+inline bool operator==(const StringView& lhs, const StringView& rhs) {
+    return (static_cast<const StringView::Base&>(lhs) == static_cast<const StringView::Base&>(rhs));
+}
+
+inline bool operator!=(const StringView& lhs, const StringView& rhs) {
+    return !(lhs == rhs);
+}
+
+}  // namespace dna
diff --git a/Engine/Plugins/Animation/RigLogic/Source/RigLogicLib/Public/dna/version/Version.h b/Engine/Plugins/Animation/RigLogic/Source/RigLogicLib/Public/dna/version/Version.h
new file mode 100644
--- /dev/null
+++ b/Engine/Plugins/Animation/RigLogic/Source/RigLogicLib/Public/dna/version/Version.h
@@ -0,0 +1,6 @@
+#pragma once
+
+#define DNA_MAJOR_VERSION 1
+#define DNA_MINOR_VERSION 1
+#define DNA_PATCH_VERSION 0
+#define DNA_VERSION_STRING "1.1.0"
diff --git a/Engine/Plugins/Animation/RigLogic/Source/RigLogicLib/Public/pma/TypeDefs.h b/Engine/Plugins/Animation/RigLogic/Source/RigLogicLib/Public/pma/TypeDefs.h
new file mode 100644
--- /dev/null
+++ b/Engine/Plugins/Animation/RigLogic/Source/RigLogicLib/Public/pma/TypeDefs.h
@@ -0,0 +1,43 @@
+#pragma once
+
+#include "pma/PolyAllocator.h"
+
+#ifdef _MSC_VER
+    #pragma warning(push)
+    #pragma warning(disable : 4365)
+#endif
+#include <cstddef>
+#include <map>
+#include <set>
+#include <string>
+#include <unordered_map>
+#include <unordered_set>
+#include <vector>
+#ifdef _MSC_VER
+    #pragma warning(pop)
+#endif
+
+namespace pma {
+
+template<typename T, typename Allocator = PolyAllocator<T> >
+using String = std::basic_string<T, std::char_traits<T>, Allocator>;
+
+template<typename T, typename Allocator = PolyAllocator<T> >
+using Vector = std::vector<T, Allocator>;
+
+template<typename T, typename Allocator = PolyAllocator<Vector<T> > >
+using Matrix = Vector<Vector<T>, Allocator>;
+
+template<typename T, typename Allocator = PolyAllocator<T> >
+using Set = std::set<T, std::less<T>, Allocator>;
+
+template<typename T, typename Allocator = PolyAllocator<T> >
+using UnorderedSet = std::unordered_set<T, std::hash<T>, std::equal_to<T>, Allocator>;
+
+template<typename K, typename V, typename Allocator = PolyAllocator<std::pair<const K, V> > >
+using Map = std::map<K, V, std::less<K>, Allocator>;
+
+template<typename K, typename V, typename Allocator = PolyAllocator<std::pair<const K, V> > >
+using UnorderedMap = std::unordered_map<K, V, std::hash<K>, std::equal_to<K>, Allocator>;
+
+}  // namespace pma
diff --git a/Engine/Plugins/Animation/RigLogic/Source/RigLogicLib/Public/pma/resources/AlignedMemoryResource.h b/Engine/Plugins/Animation/RigLogic/Source/RigLogicLib/Public/pma/resources/AlignedMemoryResource.h
new file mode 100644
--- /dev/null
+++ b/Engine/Plugins/Animation/RigLogic/Source/RigLogicLib/Public/pma/resources/AlignedMemoryResource.h
@@ -0,0 +1,20 @@
+#pragma once
+
+#include "pma/Defs.h"
+#include "pma/MemoryResource.h"
+
+#include <cstddef>
+
+namespace pma {
+
+/**
+    @brief A MemoryResource that honors alignment requirements.
+    @see MemoryResource
+*/
+class PMAAPI AlignedMemoryResource : public MemoryResource {
+    public:
+        void* allocate(std::size_t size, std::size_t alignment) override;
+        void deallocate(void* ptr, std::size_t size, std::size_t alignment) override;
+};
+
+}  // namespace pma
diff --git a/Engine/Plugins/Animation/RigLogic/Source/RigLogicLib/Public/pma/resources/DefaultMemoryResource.h b/Engine/Plugins/Animation/RigLogic/Source/RigLogicLib/Public/pma/resources/DefaultMemoryResource.h
new file mode 100644
--- /dev/null
+++ b/Engine/Plugins/Animation/RigLogic/Source/RigLogicLib/Public/pma/resources/DefaultMemoryResource.h
@@ -0,0 +1,20 @@
+#pragma once
+
+#include "pma/Defs.h"
+#include "pma/MemoryResource.h"
+
+#include <cstddef>
+
+namespace pma {
+
+/**
+    @brief A MemoryResource that delegates to malloc / free.
+    @see MemoryResource
+*/
+class PMAAPI DefaultMemoryResource : public MemoryResource {
+    public:
+        void* allocate(std::size_t size, std::size_t alignment) override;
+        void deallocate(void* ptr, std::size_t size, std::size_t alignment) override;
+};
+
+}  // namespace pma
diff --git a/Engine/Plugins/Animation/RigLogic/Source/RigLogicLib/Public/pma/version/Version.h b/Engine/Plugins/Animation/RigLogic/Source/RigLogicLib/Public/pma/version/Version.h
new file mode 100644
--- /dev/null
+++ b/Engine/Plugins/Animation/RigLogic/Source/RigLogicLib/Public/pma/version/Version.h
@@ -0,0 +1,6 @@
+#pragma once
+
+#define PMA_MAJOR_VERSION 1
+#define PMA_MINOR_VERSION 0
+#define PMA_PATCH_VERSION 1
+#define PMA_VERSION_STRING "1.0.1"
diff --git a/Engine/Plugins/Animation/RigLogic/Source/RigLogicLib/Public/riglogic/riglogic/CalculationType.h b/Engine/Plugins/Animation/RigLogic/Source/RigLogicLib/Public/riglogic/riglogic/CalculationType.h
new file mode 100644
--- /dev/null
+++ b/Engine/Plugins/Animation/RigLogic/Source/RigLogicLib/Public/riglogic/riglogic/CalculationType.h
@@ -0,0 +1,15 @@
+#pragma once
+
+namespace rl4 {
+
+/**
+    @brief Implementation type for RigLogic calculations.
+*/
+enum class CalculationType {
+    Scalar,  ///< scalar CPU algorithm
+    SSE,  ///< vectorized (SSE) CPU algorithm
+    AVX  ///< vectorized (AVX) CPU algorithm (RigLogic must be built with AVX support,
+         ///< otherwise it falls back to using the Scalar version)
+};
+
+}  // namespace rl4
diff --git a/Engine/Plugins/Animation/RigLogic/Source/RigLogicLib/Public/riglogic/transformation/Transformation.h b/Engine/Plugins/Animation/RigLogic/Source/RigLogicLib/Public/riglogic/transformation/Transformation.h
new file mode 100644
--- /dev/null
+++ b/Engine/Plugins/Animation/RigLogic/Source/RigLogicLib/Public/riglogic/transformation/Transformation.h
@@ -0,0 +1,81 @@
+#pragma once
+
+#include "riglogic/types/Aliases.h"
+
+#include <cassert>
+#include <cstddef>
+
+namespace rl4 {
+
+class Transformation {
+    public:
+        explicit Transformation(const float* chunk_) : chunk{chunk_} {
+        }
+
+        static std::size_t size() {
+            return 9ul;
+        }
+
+        const Transformation& operator*() const {
+            return *this;
+        }
+
+        Transformation& operator++() {
+            chunk += size();
+            return *this;
+        }
+
+        bool operator==(const Transformation& rhs) const {
+            return (chunk == rhs.chunk);
+        }
+
+        bool operator!=(const Transformation& rhs) const {
+            return !(*this == rhs);
+        }
+
+        Vector3 getTranslation() const {
+            return {chunk[0ul], chunk[1ul], chunk[2ul]};
+        }
+
+        Vector3 getRotation() const {
+            return {chunk[3ul], chunk[4ul], chunk[5ul]};
+        }
+
+        Vector3 getScale() const {
+            return {chunk[6ul], chunk[7ul], chunk[8ul]};
+        }
+
+    private:
+        const float* chunk;
+};
+
+class TransformationArrayView {
+    public:
+        TransformationArrayView(const float* values_, std::size_t count_) :
+            values{values_},
+            count{count_} {
+        }
+
+        std::size_t size() const {
+            assert(count % Transformation::size() == 0ul);
+            return count / Transformation::size();
+        }
+
+        Transformation operator[](std::size_t index) const {
+            return Transformation{values + (index * Transformation::size())};
+        }
+
+        Transformation begin() const {
+            return Transformation{values};
+        }
+
+        Transformation end() const {
+            return Transformation{values + count};
+        }
+
+    private:
+        const float* values;
+        std::size_t count;
+};
+
+}  // namespace rl4
diff --git a/Engine/Plugins/Animation/RigLogic/Source/RigLogicLib/Public/riglogic/version/Version.h b/Engine/Plugins/Animation/RigLogic/Source/RigLogicLib/Public/riglogic/version/Version.h
new file mode 100644
--- /dev/null
+++ b/Engine/Plugins/Animation/RigLogic/Source/RigLogicLib/Public/riglogic/version/Version.h
@@ -0,0 +1,6 @@
+#pragma once
+
+#define RL_MAJOR_VERSION 5
+#define RL_MINOR_VERSION 1
+#define RL_PATCH_VERSION 0
+#define RL_VERSION_STRING "5.1.0"
diff --git a/Engine/Plugins/Animation/RigLogic/Source/RigLogicLib/Public/status/Provider.h b/Engine/Plugins/Animation/RigLogic/Source/RigLogicLib/Public/status/Provider.h
new file mode 100644
--- /dev/null
+++ b/Engine/Plugins/Animation/RigLogic/Source/RigLogicLib/Public/status/Provider.h
@@ -0,0 +1,48 @@
+#pragma once
+
+#include "status/Defs.h"
+#include "status/StatusCode.h"
+
+#ifdef _MSC_VER
+    #pragma warning(push)
+    #pragma warning(disable : 4365 4987)
+#endif
+#include <algorithm>
+#include <array>
+#include <cstdio>
+#ifdef _MSC_VER
+    #pragma warning(pop)
+#endif
+
+namespace sc {
+
+class SCAPI StatusProvider {
+    public:
+        explicit StatusProvider(std::initializer_list<StatusCode> statuses);
+
+        static void reset();
+        static StatusCode get();
+        static bool isOk();
+        static void set(StatusCode status);
+
+        template<typename ... Args>
+        static void set(StatusCode status, Args&& ... args) {
+            std::array<char, 512> buffer{};
+            #if !defined(__clang__) && defined(__GNUC__)
+                #pragma GCC diagnostic push
+                #pragma GCC diagnostic ignored "-Wformat-security"
+            #endif
+            // The returned number of bytes to be written does not include the null terminator
+            const auto neededSize = std::snprintf(nullptr, 0ul, status.message, args ...) + 1;
+            const auto size = std::min(buffer.size(), static_cast<std::size_t>(neededSize));
+            std::snprintf(buffer.data(), size, status.message, args ...);
+            #if !defined(__clang__) && defined(__GNUC__)
+                #pragma GCC diagnostic pop
+            #endif
+            status.message = buffer.data();
+            set(status);
+        }
+
+};
+
+}  // namespace sc
diff --git a/Engine/Plugins/Animation/RigLogic/Source/RigLogicLib/Public/status/version/Version.h b/Engine/Plugins/Animation/RigLogic/Source/RigLogicLib/Public/status/version/Version.h
new file mode 100644
--- /dev/null
+++ b/Engine/Plugins/Animation/RigLogic/Source/RigLogicLib/Public/status/version/Version.h
@@ -0,0 +1,6 @@
+#pragma once
+
+#define SC_MAJOR_VERSION 1
+#define SC_MINOR_VERSION 0
+#define SC_PATCH_VERSION 0
+#define SC_VERSION_STRING "1.0.0"
diff --git a/Engine/Plugins/Animation/RigLogic/Source/RigLogicLib/Public/terse/archives/binary/InputArchive.h b/Engine/Plugins/Animation/RigLogic/Source/RigLogicLib/Public/terse/archives/binary/InputArchive.h
new file mode 100644
--- /dev/null
+++ b/Engine/Plugins/Animation/RigLogic/Source/RigLogicLib/Public/terse/archives/binary/InputArchive.h
@@ -0,0 +1,242 @@
+#pragma once
+
+#include "terse/Archive.h"
+#include "terse/archives/binary/Traits.h"
+#include "terse/utils/ArchiveOffset.h"
+#include "terse/utils/Endianness.h"
+
+#ifdef _MSC_VER
+    #pragma warning(push)
+    #pragma warning(disable : 4365 4987)
+#endif
+#include <array>
+#include <cassert>
+#include <cstddef>
+#include <cstdint>
+#include <cstring>
+#include <memory>
+#include <string>
+#include <tuple>
+#include <type_traits>
+#include <utility>
+#include <vector>
+#ifdef _MSC_VER
+    #pragma warning(pop)
+#endif
+
+namespace pma {
+
+class MemoryResource;
+
+}  // namespace pma
+
+namespace terse {
+
+namespace impl {
+
+template<typename T>
+struct ValueFactory {
+    static constexpr bool NeedsAllocator = traits::needs_allocator<T>::value;
+    static constexpr bool NeedsMemoryResource = std::is_constructible<T, pma::MemoryResource*>::value;
+    static constexpr bool IsPair = traits::is_pair<T>::value;
+    static constexpr bool IsTuple = traits::is_tuple<T>::value;
+    static constexpr bool IsPrimitive = (!NeedsAllocator && !NeedsMemoryResource & !IsPair & !IsTuple);
+
+    template<class ParentAllocator, bool IsPrimitive = IsPrimitive>
+    static typename std::enable_if<IsPrimitive, T>::type create(const ParentAllocator&  /*unused*/) {
+        return T{};
+    }
+
+    template<class ParentAllocator, bool IsPrimitive = IsPrimitive>
+    static typename std::enable_if<!IsPrimitive && NeedsMemoryResource, T>::type create(const ParentAllocator& alloc) {
+        return T{alloc.getMemoryResource()};
+    }
+
+    template<class ParentAllocator, bool IsPrimitive = IsPrimitive>
+    static typename std::enable_if<!IsPrimitive && !NeedsMemoryResource && NeedsAllocator, T>::type create(
+        const ParentAllocator& alloc) {
+        using TAllocator = typename std::allocator_traits<ParentAllocator>::template rebind_alloc<typename T::value_type>;
+        return T{TAllocator{alloc}};
+    }
+
+    template<class ParentAllocator, bool IsPrimitive = IsPrimitive>
+    static typename std::enable_if<!IsPrimitive && !NeedsMemoryResource && !NeedsAllocator && IsPair, T>::type create(
+        const ParentAllocator& alloc) {
+        using K = typename T::first_type;
+        using V = typename T::second_type;
+        return {ValueFactory<K>::create(alloc), ValueFactory<V>::create(alloc)};
+    }
+
+    template<class ParentAllocator, bool IsPrimitive = IsPrimitive>
+    static typename std::enable_if<!IsPrimitive && !NeedsMemoryResource && !NeedsAllocator && !IsPair && IsTuple, T>::type create(
+        const ParentAllocator& alloc) {
+        using K = typename std::tuple_element<0, T>::type;
+        using V = typename std::tuple_element<0, T>::type;
+        return {ValueFactory<K>::create(alloc), ValueFactory<V>::create(alloc)};
+    }
+
+};
+
+}  // namespace impl
+
+template<class TExtender, class TStream, typename TSize, typename TOffset>
+class ExtendableBinaryInputArchive : public Archive<TExtender> {
+    public:
+        // Given the possibility of both 32 and 64bit platforms, use a fixed width type during serialization
+        using SizeType = TSize;
+        using OffsetType = TOffset;
+
+    private:
+        using BaseArchive = Archive<TExtender>;
+
+    public:
+        ExtendableBinaryInputArchive(TExtender* extender, TStream* stream_) : BaseArchive{extender}, stream{stream_} {
+        }
+
+    protected:
+        template<typename T>
+        void reconstruct(T& value) {
+            using UIntType = typename traits::uint_of_size<sizeof(T)>::type;
+            static_assert(sizeof(T) == sizeof(UIntType), "No matching unsigned integral type found for the given type.");
+            // Using memcpy is the only well-defined way of reconstructing arbitrary types from raw bytes.
+            // The seemingly unnecessary copies and memcpy calls are all optimized away,
+            // compiler knows what's up.
+            UIntType hostOrder;
+            std::memcpy(std::addressof(hostOrder), std::addressof(value), sizeof(T));
+            hostOrder = ntoh(hostOrder);
+            std::memcpy(std::addressof(value), std::addressof(hostOrder), sizeof(T));
+        }
+
+        void process(ArchiveOffset<OffsetType>& dest) {
+            // Store the position of the offset itself, so it can be seeked to when writing the stream
+            dest.position = stream->tell();
+            // Load the offset value itself (this points forward within the stream to the position of
+            // the data with which the offset is associated)
+            process(dest.value);
+            // Sanity check for making sure there is an associated proxy with the offset
+            assert(dest.proxy != nullptr);
+        }
+
+        void process(typename ArchiveOffset<OffsetType>::Proxy& dest) {
+            // Rely on the offset value stored in the associated `ArchiveOffset` and seek to it
+            stream->seek(dest.target->value);
+        }
+
+        template<typename T>
+        typename std::enable_if<traits::has_load<T>::value && traits::has_serialize<T>::value, void>::type process(T& dest) {
+            dest.load(*static_cast<TExtender*>(this));
+        }
+
+        template<typename T>
+        typename std::enable_if<traits::has_load<T>::value && !traits::has_serialize<T>::value, void>::type process(T& dest) {
+            dest.load(*static_cast<TExtender*>(this));
+        }
+
+        template<typename T>
+        typename std::enable_if<!traits::has_load<T>::value && traits::has_serialize<T>::value, void>::type process(T& dest) {
+            dest.serialize(*static_cast<TExtender*>(this));
+        }
+
+        template<typename T>
+        typename std::enable_if<!traits::has_load<T>::value && !traits::has_serialize<T>::value, void>::type process(T& dest) {
+            // NOLINTNEXTLINE(cppcoreguidelines-pro-type-reinterpret-cast)
+            stream->read(reinterpret_cast<char*>(&dest), sizeof(T));
+            reconstruct(dest);
+        }
+
+        template<typename T, std::size_t N>
+        void process(std::array<T, N>& dest) {
+            for (auto& element : dest) {
+                BaseArchive::dispatch(element);
+            }
+        }
+
+        template<typename T, typename ... Args>
+        void process(std::vector<T, Args...>& dest) {
+            const auto size = processSize();
+            processElements(dest, size);
+        }
+
+        template<typename T, typename ... Args>
+        void process(std::basic_string<T, Args...>& dest) {
+            const auto size = processSize();
+            processElements(dest, size);
+        }
+
+        template<typename K, typename V>
+        void process(std::pair<K, V>& dest) {
+            BaseArchive::dispatch(dest.first);
+            BaseArchive::dispatch(dest.second);
+        }
+
+        template<typename K, typename V>
+        void process(std::tuple<K, V>& dest) {
+            BaseArchive::dispatch(std::get<0>(dest));
+            BaseArchive::dispatch(std::get<1>(dest));
+        }
+
+        std::size_t processSize() {
+            SizeType size{};
+            process(size);
+            return static_cast<std::size_t>(size);
+        }
+
+        template<class TContainer>
+        typename std::enable_if<!traits::is_batchable<TContainer>::value>::type
+        processElements(TContainer& dest, std::size_t size) {
+            using ValueType = typename TContainer::value_type;
+            dest.reserve(size);
+            for (std::size_t i = 0ul; i < size; ++i) {
+                dest.push_back(impl::ValueFactory<ValueType>::create(dest.get_allocator()));
+                BaseArchive::dispatch(dest.back());
+            }
+        }
+
+        template<class TContainer>
+        typename std::enable_if<traits::is_batchable<TContainer>::value && traits::has_wide_elements<TContainer>::value>::type
+        processElements(TContainer& dest, std::size_t size) {
+            using ValueType = typename TContainer::value_type;
+            if (size != 0ul) {
+                dest.resize(size);
+                // NOLINTNEXTLINE(cppcoreguidelines-pro-type-reinterpret-cast)
+                stream->read(reinterpret_cast<char*>(std::addressof(dest[0])), size * sizeof(ValueType));
+                for (std::size_t i = 0ul; i < size; ++i) {
+                    reconstruct(dest[i]);
+                }
+            }
+        }
+
+        template<class TContainer>
+        typename std::enable_if<traits::is_batchable<TContainer>::value && !traits::has_wide_elements<TContainer>::value>::type
+        processElements(TContainer& dest, std::size_t size) {
+            using ValueType = typename TContainer::value_type;
+            if (size != 0ul) {
+                dest.resize(size);
+                // NOLINTNEXTLINE(cppcoreguidelines-pro-type-reinterpret-cast)
+                stream->read(reinterpret_cast<char*>(std::addressof(dest[0])), size * sizeof(ValueType));
+            }
+        }
+
+    private:
+        TStream* stream;
+};
+
+template<class TStream, typename TSize = std::uint32_t, typename TOffset = TSize>
+class BinaryInputArchive : public ExtendableBinaryInputArchive<BinaryInputArchive<TStream>, TStream, TSize, TOffset> {
+    private:
+        using BaseArchive = ExtendableBinaryInputArchive<BinaryInputArchive, TStream, TSize, TOffset>;
+        friend Archive<BinaryInputArchive>;
+
+    public:
+        explicit BinaryInputArchive(TStream* stream_) : BaseArchive{this, stream_} {
+        }
+
+    private:
+        template<typename T>
+        void process(T&& dest) {
+            BaseArchive::process(std::forward<T>(dest));
+        }
+
+};
+
+}  // namespace terse
diff --git a/Engine/Plugins/Animation/RigLogic/Source/RigLogicLib/Public/terse/archives/binary/OutputArchive.h b/Engine/Plugins/Animation/RigLogic/Source/RigLogicLib/Public/terse/archives/binary/OutputArchive.h
new file mode 100644
--- /dev/null
+++ b/Engine/Plugins/Animation/RigLogic/Source/RigLogicLib/Public/terse/archives/binary/OutputArchive.h
@@ -0,0 +1,163 @@
+#pragma once
+
+#include "terse/Archive.h"
+#include "terse/archives/binary/Traits.h"
+#include "terse/utils/ArchiveOffset.h"
+#include "terse/utils/Endianness.h"
+
+#ifdef _MSC_VER
+    #pragma warning(push)
+    #pragma warning(disable : 4365 4987)
+#endif
+#include <array>
+#include <cassert>
+#include <cstddef>
+#include <cstdint>
+#include <cstring>
+#include <limits>
+#include <memory>
+#include <string>
+#include <tuple>
+#include <type_traits>
+#include <utility>
+#include <vector>
+#ifdef _MSC_VER
+    #pragma warning(pop)
+#endif
+
+namespace terse {
+
+template<class TExtender, class TStream, typename TSize, typename TOffset>
+class ExtendableBinaryOutputArchive : public Archive<TExtender> {
+    private:
+        using BaseArchive = Archive<TExtender>;
+
+    public:
+        // Given the possibility of both 32 and 64bit platforms, use a fixed width type during serialization
+        using SizeType = TSize;
+        using OffsetType = TOffset;
+
+    public:
+        ExtendableBinaryOutputArchive(TExtender* extender, TStream* stream_) : BaseArchive{extender}, stream{stream_} {
+        }
+
+    protected:
+        void process(ArchiveOffset<OffsetType>& source) {
+            // Record the position where the offset is going to be written so it can be seeked to later when
+            // it's proxy is found
+            source.position = stream->tell();
+            // Since the actual offset value is not yet known at the time when it's declaration is encountered,
+            // fill it's place with zeros as a placeholder, and it will be later populated when it's associated
+            // proxy is found
+            process(typename ArchiveOffset<OffsetType>::ValueType{});
+        }
+
+        void process(typename ArchiveOffset<OffsetType>::Proxy& source) {
+            // The current position of the stream needs to be written to the position where the associated
+            // `ArchiveOffset` is found.
+            auto current = stream->tell();
+            assert(current <= std::numeric_limits<typename ArchiveOffset<OffsetType>::ValueType>::max());
+            // Seek to the actual position of the offset marker and write the stream's above captured position there
+            stream->seek(source.target->position);
+            process(static_cast<typename ArchiveOffset<OffsetType>::ValueType>(current));
+            // Return to the earlier captured stream position so processing can safely resume from there
+            stream->seek(current);
+        }
+
+        template<typename T>
+        typename std::enable_if<traits::has_save<T>::value && traits::has_serialize<T>::value, void>::type process(
+            const T& source) {
+            const_cast<T&>(source).save(*static_cast<TExtender*>(this));
+        }
+
+        template<typename T>
+        typename std::enable_if<traits::has_save<T>::value && !traits::has_serialize<T>::value, void>::type process(
+            const T& source) {
+            const_cast<T&>(source).save(*static_cast<TExtender*>(this));
+        }
+
+        template<typename T>
+        typename std::enable_if<!traits::has_save<T>::value && traits::has_serialize<T>::value, void>::type process(
+            const T& source) {
+            const_cast<T&>(source).serialize(*static_cast<TExtender*>(this));
+        }
+
+        template<typename T>
+        typename std::enable_if<!traits::has_save<T>::value && !traits::has_serialize<T>::value,
+                                void>::type process(const T& source) {
+            // Using memcpy is the only well-defined way of reconstructing arbitrary types from raw bytes.
+            // The seemingly unnecessary copies and memcpy calls are all optimized away,
+            // compiler knows what's up.
+            using UIntType = typename traits::uint_of_size<sizeof(T)>::type;
+            static_assert(sizeof(T) == sizeof(UIntType), "No matching unsigned integral type found for the given type.");
+            UIntType networkOrder;
+            std::memcpy(std::addressof(networkOrder), std::addressof(source), sizeof(T));
+            networkOrder = hton(networkOrder);
+            // NOLINTNEXTLINE(cppcoreguidelines-pro-type-reinterpret-cast)
+            stream->write(reinterpret_cast<char*>(&networkOrder), sizeof(T));
+        }
+
+        template<typename T, std::size_t N>
+        void process(const std::array<T, N>& source) {
+            processElements(source);
+        }
+
+        template<typename T, typename ... Args>
+        void process(const std::vector<T, Args...>& source) {
+            processSize(source.size());
+            processElements(source);
+        }
+
+        template<typename T, typename ... Args>
+        void process(const std::basic_string<T, Args...>& source) {
+            processSize(source.size());
+            processElements(source);
+        }
+
+        template<typename K, typename V>
+        void process(const std::pair<K, V>& source) {
+            BaseArchive::dispatch(source.first);
+            BaseArchive::dispatch(source.second);
+        }
+
+        template<typename K, typename V>
+        void process(const std::tuple<K, V>& source) {
+            BaseArchive::dispatch(std::get<0>(source));
+            BaseArchive::dispatch(std::get<1>(source));
+        }
+
+        void processSize(std::size_t size) {
+            assert(size <= std::numeric_limits<SizeType>::max());
+            process(static_cast<SizeType>(size));
+        }
+
+        template<class TContainer>
+        void processElements(const TContainer& source) {
+            for (const auto& element : source) {
+                BaseArchive::dispatch(element);
+            }
+        }
+
+    private:
+        TStream* stream;
+};
+
+template<class TStream, typename TSize = std::uint32_t, typename TOffset = TSize>
+class BinaryOutputArchive : public ExtendableBinaryOutputArchive<BinaryOutputArchive<TStream>, TStream, TSize, TOffset> {
+    public:
+        using BaseArchive = ExtendableBinaryOutputArchive<BinaryOutputArchive, TStream, TSize, TOffset>;
+        friend Archive<BinaryOutputArchive>;
+
+    public:
+        explicit BinaryOutputArchive(TStream* stream_) : BaseArchive{this, stream_} {
+        }
+
+    private:
+        template<typename T>
+        void process(T&& dest) {
+            BaseArchive::process(std::forward<T>(dest));
+        }
+
+};
+
+}  // namespace terse
diff --git a/Engine/Plugins/Animation/RigLogic/Source/RigLogicLib/Public/terse/archives/binary/Traits.h b/Engine/Plugins/Animation/RigLogic/Source/RigLogicLib/Public/terse/archives/binary/Traits.h
new file mode 100644
--- /dev/null
+++ b/Engine/Plugins/Animation/RigLogic/Source/RigLogicLib/Public/terse/archives/binary/Traits.h
@@ -0,0 +1,85 @@
+#pragma once
+
+#include <cstdint>
+#include <type_traits>
+#include <utility>
+
+static_assert(sizeof(char) == 1ul, "Unsupported platform, char is not 8-bits wide.");
+
+namespace terse {
+
+namespace traits {
+
+template<typename ...>
+struct sink {
+    using type = void;
+};
+
+template<typename T, typename = void>
+struct needs_allocator : std::false_type {};
+
+template<typename T>
+struct needs_allocator<T, typename sink<typename T::allocator_type,
+                                        decltype(std::declval<T>().get_allocator())>::type> : std::true_type {};
+
+template<class>
+struct true_sink : std::true_type {};
+
+template<class T>
+static auto test_serialize(std::int32_t)->true_sink<decltype(std::declval<T>().serialize(std::declval<T&>()))>;
+
+template<class>
+static auto test_serialize(std::uint32_t)->std::false_type;
+
+template<class T>
+struct has_serialize : decltype(test_serialize<T>(0)) {};
+
+template<class T>
+static auto test_load(std::int32_t)->true_sink<decltype(std::declval<T>().load(std::declval<T&>()))>;
+
+template<class>
+static auto test_load(std::uint32_t)->std::false_type;
+
+template<class T>
+struct has_load : decltype(test_load<T>(0)) {};
+
+template<class T>
+static auto test_save(std::int32_t)->true_sink<decltype(std::declval<T>().save(std::declval<T&>()))>;
+
+template<class>
+static auto test_save(std::uint32_t)->std::false_type;
+
+template<class T>
+struct has_save : decltype(test_save<T>(0)) {};
+
+template<std::size_t size>
+struct uint_of_size {
+    using type = typename std::conditional<(size == 1ul), std::uint8_t,
+                                           typename std::conditional<(size == 2ul), std::uint16_t,
+                                                                     typename std::conditional<(size <= 4ul), std::uint32_t,
+                                                                                               std::uint64_t>::type>::type>::type;
+};
+
+template<typename TContainer>
+using is_batchable = std::is_scalar<typename TContainer::value_type>;
+
+template<typename TContainer>
+struct has_wide_elements {
+    static constexpr bool value = (sizeof(typename TContainer::value_type) > 1ul);
+};
+
+template<typename T>
+struct is_pair : public std::false_type {};
+
+template<typename K, typename V>
+struct is_pair<std::pair<K, V> > : public std::true_type {};
+
+template<typename T>
+struct is_tuple : public std::false_type {};
+
+template<typename K, typename V>
+struct is_tuple<std::tuple<K, V> > : public std::true_type {};
+
+}  // namespace traits
+
+}  // namespace terse
diff --git a/Engine/Plugins/Animation/RigLogic/Source/RigLogicLib/Public/terse/utils/Endianness.h b/Engine/Plugins/Animation/RigLogic/Source/RigLogicLib/Public/terse/utils/Endianness.h
new file mode 100644
--- /dev/null
+++ b/Engine/Plugins/Animation/RigLogic/Source/RigLogicLib/Public/terse/utils/Endianness.h
@@ -0,0 +1,171 @@
+#pragma once
+
+#include <cstdint>
+
+#if !defined(NO_ENDIAN_H)
+    #if defined(USE_ENDIAN_H) || defined(USE_MACHINE_ENDIAN_H) || defined(USE_SYS_ENDIAN_H)
+        #define OVERRIDDEN_ENDIAN_H
+    #endif
+
+    #if !defined(OVERRIDDEN_ENDIAN_H)
+        #if defined(__linux__) || defined(__GLIBC__) || defined(__CYGWIN__) || defined(__ANDROID__)
+            #define USE_ENDIAN_H
+        #elif defined(__APPLE__)
+            #define USE_MACHINE_ENDIAN_H
+        #elif defined(__FreeBSD__) || defined(__NetBSD__) || defined(__OpenBSD__)
+            #define USE_SYS_ENDIAN_H
+        #endif
+    #endif
+
+    #if defined(USE_ENDIAN_H)
+        #include <endian.h>
+    #elif defined(USE_MACHINE_ENDIAN_H)
+        #include <machine/endian.h>
+    #elif defined(USE_SYS_ENDIAN_H)
+        #include <sys/endian.h>
+    #elif defined(__sun)
+        #include <sys/isa_defs.h>
+    #elif defined(__MINGW32__) || defined(__MINGW64__) || !(defined(_WIN64) || defined(_WIN32))
+        #include <sys/param.h>
+    #endif
+#endif
+
+#if !defined(TARGET_LITTLE_ENDIAN) && !defined(TARGET_BIG_ENDIAN)
+    #if (defined(__BYTE_ORDER__) && __BYTE_ORDER__ == __ORDER_LITTLE_ENDIAN__)  /*GCC*/ || \
+    (defined(__BYTE_ORDER) && __BYTE_ORDER == __LITTLE_ENDIAN)  /*Linux*/ || \
+    (defined(_BYTE_ORDER) && _BYTE_ORDER == _LITTLE_ENDIAN)  /*xBSD,Sun*/ || \
+    (defined(BYTE_ORDER) && BYTE_ORDER == LITTLE_ENDIAN)  /*Apple,MingW*/ || \
+    defined(__LITTLE_ENDIAN__)  /*GCC Mac*/ || defined(__ARMEL__)  /*GCC,Clang*/ || \
+    defined(__THUMBEL__)  /*GCC,Clang*/ || defined(__AARCH64EL__)  /*GCC,Clang*/ || \
+    defined(_MIPSEL)  /*GCC,Clang*/ || defined(__MIPSEL)  /*GCC,Clang*/ || \
+    defined(__MIPSEL__)  /*GCC,Clang*/ || defined(_M_IX86)  /*MSVC*/ || \
+    defined(_M_X64)  /*MSVC*/ || defined(_M_IA64)  /*MSVC*/ || \
+    defined(_M_AMD64)  /*MSVC*/ || defined(_M_ARM)  /*MSVC*/ || \
+    defined(_M_ARM64)  /*MSVC*/
+        #define TARGET_LITTLE_ENDIAN
+    #elif (defined(__BYTE_ORDER__) && __BYTE_ORDER__ == __ORDER_BIG_ENDIAN__)  /*GCC*/ || \
+    (defined(__BYTE_ORDER) && __BYTE_ORDER == __BIG_ENDIAN)  /*Linux*/ || \
+    (defined(_BYTE_ORDER) && _BYTE_ORDER == _BIG_ENDIAN)  /*xBSD,Sun*/ || \
+    (defined(BYTE_ORDER) && BYTE_ORDER == BIG_ENDIAN)  /*Apple,MingW*/ || \
+    defined(_M_PPC)  /*MSVC for XBox-360*/ || defined(__BIG_ENDIAN__)  /*GCC Mac*/ || \
+    defined(__ARMEB__)  /*GCC,Clang*/ || defined(__THUMBEB__)  /*GCC,Clang*/ || \
+    defined(__AARCH64EB__)  /*GCC,Clang*/ || defined(_MIPSEB)  /*GCC,Clang*/ || \
+    defined(__MIPSEB)  /*GCC,Clang*/ || defined(__MIPSEB__)  /*GCC,Clang*/
+        #define TARGET_BIG_ENDIAN
+    #elif defined(_WIN32)
+        #define TARGET_LITTLE_ENDIAN
+    #endif  // End of byte order checks
+#endif  // End of guard for explicitly defined endianness
+
+/*
+ * Swap intrinsics
+ */
+#if defined(__clang__) || (defined(__GNUC__) && \
+    ((__GNUC__ == 4 && __GNUC_MINOR__ >= 8) || __GNUC__ >= 5))
+    #define bswap16(x) __builtin_bswap16((x))
+    #define bswap32(x) __builtin_bswap32((x))
+    #define bswap64(x) __builtin_bswap64((x))
+#elif defined(__linux__) || defined(__GLIBC__)
+    #include <byteswap.h>
+    #define bswap16(x) bswap_16((x))
+    #define bswap32(x) bswap_32((x))
+    #define bswap64(x) bswap_64((x))
+#elif defined(_MSC_VER)
+    #include <stdlib.h>
+    #define bswap16(x) _byteswap_ushort((x))
+    #define bswap32(x) _byteswap_ulong((x))
+    #define bswap64(x) _byteswap_uint64((x))
+#elif defined(__APPLE__)
+    #include <libkern/OSByteOrder.h>
+    #define bswap16(x) OSSwapInt16((x))
+    #define bswap32(x) OSSwapInt32((x))
+    #define bswap64(x) OSSwapInt64((x))
+#elif defined(__FreeBSD__) || defined(__NetBSD__)
+    #include <sys/endian.h>  // This defines the intrinsics as per the chosen naming convention
+#elif defined(__OpenBSD__)
+    #include <sys/endian.h>
+    #define bswap16(x) swap16((x))
+    #define bswap32(x) swap32((x))
+    #define bswap64(x) swap64((x))
+#elif defined(__sun) || defined(sun)
+    #include <sys/byteorder.h>
+    #define bswap16(x) BSWAP_16((x))
+    #define bswap32(x) BSWAP_32((x))
+    #define bswap64(x) BSWAP_64((x))
+#else
+    static inline std::uint16_t bswap16(std::uint16_t x) {
+        return (((x& std::uint16_t{0x00FF}) << 8) |
+                ((x& std::uint16_t{0xFF00}) >> 8));
+    }
+
+    static inline std::uint32_t bswap32(std::uint32_t x) {
+        return (((x& std::uint32_t{0x000000FF}) << 24) |
+                ((x& std::uint32_t{0x0000FF00}) << 8) |
+                ((x& std::uint32_t{0x00FF0000}) >> 8) |
+                ((x& std::uint32_t{0xFF000000}) >> 24));
+    }
+
+    static inline std::uint64_t bswap64(std::uint64_t x) {
+        return (((x& std::uint64_t{0x00000000000000FF}) << 56) |
+                ((x& std::uint64_t{0x000000000000FF00}) << 40) |
+                ((x& std::uint64_t{0x0000000000FF0000}) << 24) |
+                ((x& std::uint64_t{0x00000000FF000000}) << 8) |
+                ((x& std::uint64_t{0x000000FF00000000}) >> 8) |
+                ((x& std::uint64_t{0x0000FF0000000000}) >> 24) |
+                ((x& std::uint64_t{0x00FF000000000000}) >> 40) |
+                ((x& std::uint64_t{0xFF00000000000000}) >> 56));
+    }
+
+#endif
+
+// Target architecture specific ntoh and hton for all relevant sizes
+// In case of big endian architectures this is a noop
+#if defined(TARGET_LITTLE_ENDIAN)
+    #define ntoh16(x) bswap16((x))
+    #define hton16(x) bswap16((x))
+    #define ntoh32(x) bswap32((x))
+    #define hton32(x) bswap32((x))
+    #define ntoh64(x) bswap64((x))
+    #define hton64(x) bswap64((x))
+#elif defined(TARGET_BIG_ENDIAN)
+    #define ntoh16(x) (x)
+    #define hton16(x) (x)
+    #define ntoh32(x) (x)
+    #define hton32(x) (x)
+    #define ntoh64(x) (x)
+    #define hton64(x) (x)
+#else
+    #error "Platform not supported, no byte swap functions defined."
+#endif
+
+inline std::uint8_t ntoh(std::uint8_t x) {
+    return x;
+}
+
+inline std::uint16_t ntoh(std::uint16_t x) {
+    return ntoh16(x);
+}
+
+inline std::uint32_t ntoh(std::uint32_t x) {
+    return ntoh32(x);
+}
+
+inline std::uint64_t ntoh(std::uint64_t x) {
+    return ntoh64(x);
+}
+
+inline std::uint8_t hton(std::uint8_t x) {
+    return x;
+}
+
+inline std::uint16_t hton(std::uint16_t x) {
+    return hton16(x);
+}
+
+inline std::uint32_t hton(std::uint32_t x) {
+    return hton32(x);
+}
+
+inline std::uint64_t hton(std::uint64_t x) {
+    return hton64(x);
+}
diff --git a/Engine/Plugins/Animation/RigLogic/Source/RigLogicLib/Public/terse/version/Version.h b/Engine/Plugins/Animation/RigLogic/Source/RigLogicLib/Public/terse/version/Version.h
new file mode 100644
--- /dev/null
+++ b/Engine/Plugins/Animation/RigLogic/Source/RigLogicLib/Public/terse/version/Version.h
@@ -0,0 +1,6 @@
+#pragma once
+
+#define TERSE_MAJOR_VERSION 1
+#define TERSE_MINOR_VERSION 1
+#define TERSE_PATCH_VERSION 0
+#define TERSE_VERSION_STRING "1.1.0"
diff --git a/Engine/Plugins/Animation/RigLogic/Source/RigLogicLib/Public/trio/Stream.h b/Engine/Plugins/Animation/RigLogic/Source/RigLogicLib/Public/trio/Stream.h
new file mode 100644
--- /dev/null
+++ b/Engine/Plugins/Animation/RigLogic/Source/RigLogicLib/Public/trio/Stream.h
@@ -0,0 +1,21 @@
+#pragma once
+
+#include "trio/Concepts.h"
+#include "trio/Defs.h"
+
+#include <cstddef>
+
+namespace trio {
+
+class TRIOAPI BoundedIOStream : public Controllable, public Readable, public Writable, public Seekable {
+    public:
+        virtual ~BoundedIOStream();
+        /**
+            @brief Obtain size of stream in bytes.
+            @return
+                Size in bytes.
+        */
+        virtual std::size_t size() = 0;
+};
+
+}  // namespace trio
diff --git a/Engine/Plugins/Animation/RigLogic/Source/RigLogicLib/Public/trio/version/Version.h b/Engine/Plugins/Animation/RigLogic/Source/RigLogicLib/Public/trio/version/Version.h
new file mode 100644
--- /dev/null
+++ b/Engine/Plugins/Animation/RigLogic/Source/RigLogicLib/Public/trio/version/Version.h
@@ -0,0 +1,6 @@
+#pragma once
+
+#define TRIO_MAJOR_VERSION 1
+#define TRIO_MINOR_VERSION 0
+#define TRIO_PATCH_VERSION 1
+#define TRIO_VERSION_STRING "1.0.1"
diff --git a/Engine/Plugins/Animation/RigLogic/Source/RigLogicLib/RigLogicLib.Build.cs b/Engine/Plugins/Animation/RigLogic/Source/RigLogicLib/RigLogicLib.Build.cs
new file mode 100644
--- /dev/null
+++ b/Engine/Plugins/Animation/RigLogic/Source/RigLogicLib/RigLogicLib.Build.cs
@@ -0,0 +1,44 @@
+// Copyright 1998-2019 Epic Games, Inc. All Rights Reserved.
+
+using UnrealBuildTool;
+using System.IO;
+
+public class RigLogicLib : ModuleRules
+{
+    public string ModulePath
+    {
+        get { return ModuleDirectory; }
+    }
+
+    public RigLogicLib(ReadOnlyTargetRules Target) : base(Target)
+    {
+        PublicDependencyModuleNames.AddRange(
+            new string[]
+            {
+                    "Core",
+                    "CoreUObject",
+                    "Engine",
+                    "ControlRig"
+            }
+        );
+
+        if (Target.Type == TargetType.Editor)
+        {
+            PublicDependencyModuleNames.Add("UnrealEd");
+        }
+
+        Type = ModuleType.CPlusPlus;
+
+        PrivateDefinitions.Add("RL_BUILD_SHARED=1");  // used instead of #define so cpp files can do conditional compilation
+
+        if (Target.Platform == UnrealTargetPlatform.Win32 ||
+                Target.Platform == UnrealTargetPlatform.Win64 ||
+                Target.Platform == UnrealTargetPlatform.Linux ||
+                Target.Platform == UnrealTargetPlatform.Mac ||
+                Target.Platform == UnrealTargetPlatform.PS4 ||
+                Target.Platform == UnrealTargetPlatform.XboxOne)
+        {
+            PrivateDefinitions.Add("RL_BUILD_WITH_SSE=1");
+        }
+    }
+}
diff --git a/Engine/Plugins/Animation/RigLogic/Source/RigLogicLibTest/Private/RigLogicLibTest.cpp b/Engine/Plugins/Animation/RigLogic/Source/RigLogicLibTest/Private/RigLogicLibTest.cpp
new file mode 100644
--- /dev/null
+++ b/Engine/Plugins/Animation/RigLogic/Source/RigLogicLibTest/Private/RigLogicLibTest.cpp
@@ -0,0 +1,14 @@
+// Copyright 1998-2019 Epic Games, Inc. All Rights Reserved.
+
+#include "RigLogicLibTest.h"
+
+#include "Runtime/Core/Public/Modules/ModuleInterface.h"
+#include "Runtime/Core/Public/Modules/ModuleManager.h"
+
+#define LOCTEXT_NAMESPACE "FRigLogicLibTest"
+
+DEFINE_LOG_CATEGORY(LogRigLogicLibTest);
+
+#undef LOCTEXT_NAMESPACE
+
+IMPLEMENT_MODULE(FRigLogicLibTest, RigLogicLibTest)
diff --git a/Engine/Plugins/Animation/RigLogic/Source/RigLogicLibTest/Private/RigLogicLibTestSuite.cpp b/Engine/Plugins/Animation/RigLogic/Source/RigLogicLibTest/Private/RigLogicLibTestSuite.cpp
new file mode 100644
--- /dev/null
+++ b/Engine/Plugins/Animation/RigLogic/Source/RigLogicLibTest/Private/RigLogicLibTestSuite.cpp
@@ -0,0 +1,51 @@
+// Copyright 1998-2019 Epic Games, Inc. All Rights Reserved.
+
+#if WITH_DEV_AUTOMATION_TESTS
+#pragma warning(disable:4668)
+#pragma warning(disable:4996)
+#include "gtest/gtest.h"
+#pragma warning(default:4668)
+#pragma warning(default:4996)
+
+#include "CoreMinimal.h"
+#include "Misc/AutomationTest.h"
+
+class FTestCollector : public ::testing::EmptyTestEventListener
+{
+public:
+    explicit FTestCollector(TArray<FString>* Names) : TestNames{Names}
+    {
+    }
+
+private:
+    void OnTestStart(const ::testing::TestInfo& TestInfo) override
+    {
+        FString TestCaseName = FString::Printf(TEXT("%s.%s"), *FString(TestInfo.test_case_name()), *FString(TestInfo.name()));
+        TestNames->Add(TestCaseName);
+    }
+
+private:
+    TArray<FString>* TestNames;
+};
+
+IMPLEMENT_COMPLEX_AUTOMATION_TEST(FRigLogicLibTestSuite, "ControlRig.Units.FRigUnit_RigLogicLib", EAutomationTestFlags::EditorContext | EAutomationTestFlags::EngineFilter)
+
+void FRigLogicLibTestSuite::GetTests(TArray<FString>& OutBeautifiedNames, TArray<FString>& OutTestCommands) const
+{
+    FTestCollector Collector{&OutBeautifiedNames};
+    ::testing::InitGoogleTest();
+    ::testing::TestEventListeners& Listeners = ::testing::UnitTest::GetInstance()->listeners();
+    Listeners.Append(&Collector);
+    auto _ = RUN_ALL_TESTS();
+    OutTestCommands = OutBeautifiedNames;
+    Listeners.Release(&Collector);
+}
+
+bool FRigLogicLibTestSuite::RunTest(const FString& Parameters)
+{
+    TArray<ANSICHAR> TestFilter{TCHAR_TO_ANSI(*Parameters), Parameters.Len() + 1};
+    ::testing::GTEST_FLAG(filter) = TestFilter.GetData();
+    return (RUN_ALL_TESTS() == 0);
+}
+
+#endif  // WITH_DEV_AUTOMATION_TESTS
diff --git a/Engine/Plugins/Animation/RigLogic/Source/RigLogicLibTest/Private/dnatests/Defs.h b/Engine/Plugins/Animation/RigLogic/Source/RigLogicLibTest/Private/dnatests/Defs.h
new file mode 100644
--- /dev/null
+++ b/Engine/Plugins/Animation/RigLogic/Source/RigLogicLibTest/Private/dnatests/Defs.h
@@ -0,0 +1,48 @@
+#pragma once
+
+#ifdef _MSC_VER
+    #pragma warning(disable:4668)
+    #pragma warning(disable:4996)
+    #pragma warning(push)
+    #pragma warning(disable : 4365 4987)
+#endif
+#ifdef __clang__
+    #pragma clang diagnostic push
+    #pragma clang diagnostic ignored "-Wglobal-constructors"
+#endif
+#include <gtest/gtest.h>
+#ifdef __clang__
+    #pragma clang diagnostic pop
+#endif
+#ifdef _MSC_VER
+    #pragma warning(pop)
+    #pragma warning(default:4668)
+    #pragma warning(default:4996)
+#endif
+
+#ifndef INSTANTIATE_TEST_SUITE_P
+    #define INSTANTIATE_TEST_SUITE_P INSTANTIATE_TEST_CASE_P
+#endif
+
+#ifndef TYPED_TEST_SUITE_P
+    #define TYPED_TEST_SUITE_P TYPED_TEST_CASE_P
+#endif
+
+#ifndef REGISTER_TYPED_TEST_SUITE_P
+    #define REGISTER_TYPED_TEST_SUITE_P REGISTER_TYPED_TEST_CASE_P
+#endif
+
+#ifndef INSTANTIATE_TYPED_TEST_SUITE_P
+    #define INSTANTIATE_TYPED_TEST_SUITE_P INSTANTIATE_TYPED_TEST_CASE_P
+#endif
+
+#define ASSERT_ELEMENTS_EQ(result, expected, count)   \
+    for (std::size_t i = 0ul; i < count; ++i) {         \
+        ASSERT_EQ(result[i], expected[i]);              \
+    }
+
+#define ASSERT_ELEMENTS_NEAR(lhs, rhs, threshold)       \
+    ASSERT_EQ(lhs.size(), rhs.size());                  \
+    for (std::size_t i = 0ul; i < lhs.size(); ++i) {    \
+        ASSERT_NEAR(lhs[i], rhs[i], threshold);         \
+    }
diff --git a/Engine/Plugins/Animation/RigLogic/Source/RigLogicLibTest/Private/dnatests/FakeStream.h b/Engine/Plugins/Animation/RigLogic/Source/RigLogicLibTest/Private/dnatests/FakeStream.h
new file mode 100644
--- /dev/null
+++ b/Engine/Plugins/Animation/RigLogic/Source/RigLogicLibTest/Private/dnatests/FakeStream.h
@@ -0,0 +1,63 @@
+#pragma once
+
+#include <trio/Stream.h>
+
+#ifdef _MSC_VER
+    #pragma warning(push)
+    #pragma warning(disable : 4365 4987)
+#endif
+#include <cassert>
+#include <cstring>
+#include <cstddef>
+#include <vector>
+#ifdef _MSC_VER
+    #pragma warning(pop)
+#endif
+
+namespace dnatests {
+
+struct FakeStream : public trio::BoundedIOStream {
+
+    FakeStream() : position{} {
+    }
+
+    void open() override {
+        position = 0ul;
+    }
+
+    void close() override {
+        position = 0ul;
+    }
+
+    std::size_t tell() override {
+        return position;
+    }
+
+    void seek(std::size_t position_) override {
+        position = position_;
+    }
+
+    void read(char* buffer, std::size_t size) override {
+        assert(position + size <= data.size());
+        std::memcpy(buffer, &data[position], size);
+        position += size;
+    }
+
+    void write(const char* buffer, std::size_t size) override {
+        std::size_t available = data.size() - position;
+        if (available < size) {
+            data.resize(data.size() + (size - available));
+        }
+        std::memcpy(&data[position], buffer, size);
+        position += size;
+    }
+
+    std::size_t size() override {
+        return data.size();
+    }
+
+    std::vector<char> data;
+    std::size_t position;
+};
+
+}  // namespace dnatests
