Editgrating RigLogic plugin from Dev-CharacterTech to UE5/Main @ CL 13891876.

diff --git a/Engine/Plugins/Animation/RigLogic/Source/RigLogicLibTest/Private/dnatests/TestStreamReadWriteIntegration.cpp b/Engine/Plugins/Animation/RigLogic/Source/RigLogicLibTest/Private/dnatests/TestStreamReadWriteIntegration.cpp
new file mode 100644
--- /dev/null
+++ b/Engine/Plugins/Animation/RigLogic/Source/RigLogicLibTest/Private/dnatests/TestStreamReadWriteIntegration.cpp
@@ -0,0 +1,371 @@
+#include "dnatests/TestStreamReadWriteIntegration.h"
+
+#include "dnatests/Defs.h"
+#include "dnatests/FakeStream.h"
+#include "dnatests/Fixtures.h"
+
+#include "dna/DataLayer.h"
+#include "dna/DNA.h"
+#include "dna/StreamReader.h"
+#include "dna/StreamWriter.h"
+
+#include <pma/resources/AlignedMemoryResource.h>
+
+namespace dna {
+
+static void verifyDescriptor(DescriptorReader* reader, const LODParameters& params) {
+    ASSERT_EQ(reader->getName(), StringView{decoded::name});
+    ASSERT_EQ(reader->getArchetype(), decoded::archetype);
+    ASSERT_EQ(reader->getGender(), decoded::gender);
+    ASSERT_EQ(reader->getAge(), decoded::age);
+
+    const auto metaDataCount = reader->getMetaDataCount();
+    ASSERT_EQ(metaDataCount, 2u);
+    for (std::uint32_t i = 0u; i < metaDataCount; ++i) {
+        const auto key = reader->getMetaDataKey(i);
+        const auto value = reader->getMetaDataValue(key);
+        ASSERT_EQ(key, StringView{decoded::metadata[i].first});
+        ASSERT_EQ(value, StringView{decoded::metadata[i].second});
+    }
+
+    ASSERT_EQ(reader->getTranslationUnit(), decoded::translationUnit);
+    ASSERT_EQ(reader->getRotationUnit(), decoded::rotationUnit);
+
+    const auto coordinateSystem = reader->getCoordinateSystem();
+    ASSERT_EQ(coordinateSystem.xAxis, decoded::coordinateSystem.xAxis);
+    ASSERT_EQ(coordinateSystem.yAxis, decoded::coordinateSystem.yAxis);
+    ASSERT_EQ(coordinateSystem.zAxis, decoded::coordinateSystem.zAxis);
+
+    ASSERT_EQ(reader->getLODCount(), decoded::lodCount[params.maxLOD]);
+    ASSERT_EQ(reader->getDBMaxLOD(), decoded::maxLODs[params.maxLOD]);
+    ASSERT_EQ(reader->getDBComplexity(), StringView{decoded::complexity});
+    ASSERT_EQ(reader->getDBName(), StringView{decoded::dbName});
+}
+
+static void verifyDefinition(DefinitionReader* reader, const LODParameters& params) {
+    const auto guiControlCount = reader->getGUIControlCount();
+    ASSERT_EQ(guiControlCount, decoded::guiControlNames.size());
+    for (std::uint16_t i = 0u; i < guiControlCount; ++i) {
+        ASSERT_EQ(reader->getGUIControlName(i), StringView{decoded::guiControlNames[i]});
+    }
+
+    const auto rawControlCount = reader->getRawControlCount();
+    ASSERT_EQ(rawControlCount, decoded::rawControlNames.size());
+    for (std::uint16_t i = 0u; i < rawControlCount; ++i) {
+        ASSERT_EQ(reader->getRawControlName(i), StringView{decoded::rawControlNames[i]});
+    }
+
+    ASSERT_EQ(reader->getJointCount(), decoded::jointNames[params.maxLOD][0ul].size());
+    const auto& expectedJointNames = decoded::jointNames[params.maxLOD][params.currentLOD];
+    const auto jointIndices = reader->getJointIndicesForLOD(params.currentLOD);
+    ASSERT_EQ(jointIndices.size(), expectedJointNames.size());
+    for (std::size_t i = 0ul; i < jointIndices.size(); ++i) {
+        ASSERT_EQ(reader->getJointName(jointIndices[i]), StringView{expectedJointNames[i]});
+    }
+
+    for (std::uint16_t i = 0u; i < reader->getJointCount(); ++i) {
+        ASSERT_EQ(reader->getJointParentIndex(i), decoded::jointHierarchy[params.maxLOD][i]);
+    }
+
+    ASSERT_EQ(reader->getBlendShapeChannelCount(), decoded::blendShapeNames[params.maxLOD][0ul].size());
+    const auto& expectedBlendShapeNames = decoded::blendShapeNames[params.maxLOD][params.currentLOD];
+    const auto blendShapeIndices = reader->getBlendShapeChannelIndicesForLOD(params.currentLOD);
+    ASSERT_EQ(blendShapeIndices.size(), expectedBlendShapeNames.size());
+    for (std::size_t i = 0ul; i < blendShapeIndices.size(); ++i) {
+        ASSERT_EQ(reader->getBlendShapeChannelName(blendShapeIndices[i]), StringView{expectedBlendShapeNames[i]});
+    }
+
+    ASSERT_EQ(reader->getAnimatedMapCount(), decoded::animatedMapNames[params.maxLOD][0ul].size());
+    const auto& expectedAnimatedMapNames = decoded::animatedMapNames[params.maxLOD][params.currentLOD];
+    const auto animatedMapIndices = reader->getAnimatedMapIndicesForLOD(params.currentLOD);
+    ASSERT_EQ(animatedMapIndices.size(), expectedAnimatedMapNames.size());
+    for (std::size_t i = 0ul; i < animatedMapIndices.size(); ++i) {
+        ASSERT_EQ(reader->getAnimatedMapName(animatedMapIndices[i]), StringView{expectedAnimatedMapNames[i]});
+    }
+
+    ASSERT_EQ(reader->getMeshCount(), decoded::meshNames[params.maxLOD][0ul].size());
+    const auto& expectedMeshNames = decoded::meshNames[params.maxLOD][params.currentLOD];
+    const auto meshIndices = reader->getMeshIndicesForLOD(params.currentLOD);
+    ASSERT_EQ(meshIndices.size(), expectedMeshNames.size());
+    for (std::size_t i = 0ul; i < meshIndices.size(); ++i) {
+        ASSERT_EQ(reader->getMeshName(meshIndices[i]), StringView{expectedMeshNames[i]});
+    }
+
+    ASSERT_EQ(reader->getMeshBlendShapeChannelMappingCount(), decoded::meshBlendShapeIndices[params.maxLOD][0ul].size());
+    const auto meshBlendShapeIndices = reader->getMeshBlendShapeChannelMappingIndicesForLOD(params.currentLOD);
+    const auto& expectedMeshBlendShapeIndices = decoded::meshBlendShapeIndices[params.maxLOD][params.currentLOD];
+    ASSERT_EQ(meshBlendShapeIndices, ConstArrayView<std::uint16_t>{expectedMeshBlendShapeIndices});
+
+    const auto& expectedNeutralJointTranslations = decoded::neutralJointTranslations[params.maxLOD][params.currentLOD];
+    ASSERT_EQ(jointIndices.size(), expectedNeutralJointTranslations.size());
+    for (std::size_t i = 0ul; i < jointIndices.size(); ++i) {
+        ASSERT_EQ(reader->getNeutralJointTranslation(jointIndices[i]), expectedNeutralJointTranslations[i]);
+    }
+
+    const auto& expectedNeutralJointRotations = decoded::neutralJointRotations[params.maxLOD][params.currentLOD];
+    ASSERT_EQ(jointIndices.size(), expectedNeutralJointRotations.size());
+    for (std::size_t i = 0ul; i < jointIndices.size(); ++i) {
+        ASSERT_EQ(reader->getNeutralJointRotation(jointIndices[i]), expectedNeutralJointRotations[i]);
+    }
+}
+
+static void verifyBehavior(BehaviorReader* reader, const LODParameters& params) {
+    const auto guiToRawInputIndices = reader->getGUIToRawInputIndices();
+    const auto& expectedG2RInputIndices = decoded::conditionalInputIndices[0ul][0ul];
+    ASSERT_EQ(guiToRawInputIndices, ConstArrayView<std::uint16_t>{expectedG2RInputIndices});
+
+    const auto guiToRawOutputIndices = reader->getGUIToRawOutputIndices();
+    const auto& expectedG2ROutputIndices = decoded::conditionalOutputIndices[0ul][0ul];
+    ASSERT_EQ(guiToRawOutputIndices, ConstArrayView<std::uint16_t>{expectedG2ROutputIndices});
+
+    const auto guiToRawFromValues = reader->getGUIToRawFromValues();
+    const auto& expectedG2RFromValues = decoded::conditionalFromValues[0ul][0ul];
+    ASSERT_EQ(guiToRawFromValues, ConstArrayView<float>{expectedG2RFromValues});
+
+    const auto guiToRawToValues = reader->getGUIToRawToValues();
+    const auto& expectedG2RToValues = decoded::conditionalToValues[0ul][0ul];
+    ASSERT_EQ(guiToRawToValues, ConstArrayView<float>{expectedG2RToValues});
+
+    const auto guiToRawSlopeValues = reader->getGUIToRawSlopeValues();
+    const auto& expectedG2RSlopeValues = decoded::conditionalSlopeValues[0ul][0ul];
+    ASSERT_EQ(guiToRawSlopeValues, ConstArrayView<float>{expectedG2RSlopeValues});
+
+    const auto guiToRawCutValues = reader->getGUIToRawCutValues();
+    const auto& expectedG2RCutValues = decoded::conditionalCutValues[0ul][0ul];
+    ASSERT_EQ(guiToRawCutValues, ConstArrayView<float>{expectedG2RCutValues});
+
+    const auto psdRowIndices = reader->getPSDRowIndices();
+    ASSERT_EQ(psdRowIndices, ConstArrayView<std::uint16_t>{decoded::psdRowIndices});
+
+    const auto psdColumnIndices = reader->getPSDColumnIndices();
+    ASSERT_EQ(psdColumnIndices, ConstArrayView<std::uint16_t>{decoded::psdColumnIndices});
+
+    const auto psdValues = reader->getPSDValues();
+    ASSERT_EQ(psdValues, ConstArrayView<float>{decoded::psdValues});
+
+    ASSERT_EQ(reader->getPSDCount(), decoded::psdCount);
+    ASSERT_EQ(reader->getJointRowCount(), decoded::jointRowCount[params.maxLOD]);
+    ASSERT_EQ(reader->getJointColumnCount(), decoded::jointColumnCount);
+
+    const auto jointVariableAttrIndices = reader->getJointVariableAttributeIndices(params.currentLOD);
+    const auto& expectedJointVariableAttrIndices = decoded::jointVariableIndices[params.maxLOD][params.currentLOD];
+    ASSERT_EQ(jointVariableAttrIndices, ConstArrayView<std::uint16_t>{expectedJointVariableAttrIndices});
+
+    const auto jointGroupCount = reader->getJointGroupCount();
+    ASSERT_EQ(jointGroupCount, decoded::jointGroupLODs.size());
+
+    for (std::uint16_t i = 0u; i < jointGroupCount; ++i) {
+        const auto& expectedLODs = decoded::jointGroupLODs[i][params.maxLOD];
+        ASSERT_EQ(reader->getJointGroupLODs(i), ConstArrayView<std::uint16_t>{expectedLODs});
+
+        const auto& expectedInputIndices = decoded::jointGroupInputIndices[i][params.maxLOD][0ul];
+        ASSERT_EQ(reader->getJointGroupInputIndices(i), ConstArrayView<std::uint16_t>{expectedInputIndices});
+
+        const auto outputIndices = reader->getJointGroupOutputIndices(i);
+        ASSERT_EQ(outputIndices.size(), expectedLODs[0ul]);
+
+        ConstArrayView<std::uint16_t> outputIndicesForLOD{outputIndices.data(), expectedLODs[params.currentLOD]};
+        const auto& expectedOutputIndices = decoded::jointGroupOutputIndices[i][params.maxLOD][params.currentLOD];
+        ASSERT_EQ(outputIndicesForLOD, ConstArrayView<std::uint16_t>{expectedOutputIndices});
+
+        const auto values = reader->getJointGroupValues(i);
+        ASSERT_EQ(values.size(), expectedLODs[0ul] * expectedInputIndices.size());
+
+        ConstArrayView<float> valuesForLOD{values.data(), expectedLODs[params.currentLOD] * expectedInputIndices.size()};
+        const auto& expectedValues = decoded::jointGroupValues[i][params.maxLOD][params.currentLOD];
+        ASSERT_EQ(valuesForLOD, ConstArrayView<float>{expectedValues});
+
+        const auto& expectedJointIndices = decoded::jointGroupJointIndices[i][params.maxLOD][0ul];
+        ASSERT_EQ(reader->getJointGroupJointIndices(i), ConstArrayView<std::uint16_t>{expectedJointIndices});
+    }
+
+    ASSERT_EQ(reader->getBlendShapeChannelLODs(),
+              ConstArrayView<std::uint16_t>{decoded::blendShapeLODs[params.maxLOD]});
+
+    const auto blendShapeChannelInputIndices = reader->getBlendShapeChannelInputIndices();
+    ASSERT_EQ(blendShapeChannelInputIndices.size(), decoded::blendShapeLODs[params.maxLOD][0ul]);
+    ConstArrayView<std::uint16_t> blendShapeInputIndicesForLOD{
+        blendShapeChannelInputIndices.data(),
+        decoded::blendShapeLODs[params.maxLOD][params.currentLOD]
+    };
+    ASSERT_EQ(blendShapeInputIndicesForLOD,
+              ConstArrayView<std::uint16_t>{decoded::blendShapeInputIndices[params.maxLOD][params.currentLOD]});
+
+    const auto blendShapeChannelOutputIndices = reader->getBlendShapeChannelOutputIndices();
+    ASSERT_EQ(blendShapeChannelOutputIndices.size(), decoded::blendShapeLODs[params.maxLOD][0ul]);
+    ConstArrayView<std::uint16_t> blendShapeOutputIndicesForLOD{
+        blendShapeChannelOutputIndices.data(),
+        decoded::blendShapeLODs[params.maxLOD][params.currentLOD]
+    };
+    ASSERT_EQ(blendShapeOutputIndicesForLOD,
+              ConstArrayView<std::uint16_t>{decoded::blendShapeOutputIndices[params.maxLOD][params.currentLOD]});
+
+    ASSERT_EQ(reader->getAnimatedMapLODs(),
+              ConstArrayView<std::uint16_t>{decoded::animatedMapLODs[params.maxLOD]});
+
+    ASSERT_EQ(reader->getAnimatedMapCount(), decoded::animatedMapCount[params.maxLOD]);
+
+    const auto animatdMapLOD = decoded::animatedMapLODs[params.maxLOD][params.currentLOD];
+
+    const auto animatedMapInputIndices = reader->getAnimatedMapInputIndices();
+    ASSERT_EQ(animatedMapInputIndices.size(), decoded::animatedMapLODs[params.maxLOD][0ul]);
+    ConstArrayView<std::uint16_t> animatedMapInputIndicesForLOD{animatedMapInputIndices.data(), animatdMapLOD};
+    ASSERT_EQ(animatedMapInputIndicesForLOD,
+              ConstArrayView<std::uint16_t>{decoded::conditionalInputIndices[params.maxLOD][params.currentLOD]});
+
+    const auto animatedMapOutputIndices = reader->getAnimatedMapOutputIndices();
+    ASSERT_EQ(animatedMapOutputIndices.size(), decoded::animatedMapLODs[params.maxLOD][0ul]);
+    ConstArrayView<std::uint16_t> animatedMapOutputIndicesForLOD{animatedMapOutputIndices.data(), animatdMapLOD};
+    ASSERT_EQ(animatedMapOutputIndicesForLOD,
+              ConstArrayView<std::uint16_t>{decoded::conditionalOutputIndices[params.maxLOD][params.currentLOD]});
+
+    const auto animatedMapFromValues = reader->getAnimatedMapFromValues();
+    ASSERT_EQ(animatedMapFromValues.size(), decoded::animatedMapLODs[params.maxLOD][0ul]);
+    ConstArrayView<float> animatedMapFromValuesForLOD{animatedMapFromValues.data(), animatdMapLOD};
+    ASSERT_EQ(animatedMapFromValuesForLOD,
+              ConstArrayView<float>{decoded::conditionalFromValues[params.maxLOD][params.currentLOD]});
+
+    const auto animatedMapToValues = reader->getAnimatedMapToValues();
+    ASSERT_EQ(animatedMapToValues.size(), decoded::animatedMapLODs[params.maxLOD][0ul]);
+    ConstArrayView<float> animatedMapToValuesForLOD{animatedMapToValues.data(), animatdMapLOD};
+    ASSERT_EQ(animatedMapToValuesForLOD,
+              ConstArrayView<float>{decoded::conditionalToValues[params.maxLOD][params.currentLOD]});
+
+    const auto animatedMapSlopeValues = reader->getAnimatedMapSlopeValues();
+    ASSERT_EQ(animatedMapSlopeValues.size(), decoded::animatedMapLODs[params.maxLOD][0ul]);
+    ConstArrayView<float> animatedMapSlopeValuesForLOD{animatedMapSlopeValues.data(), animatdMapLOD};
+    ASSERT_EQ(animatedMapSlopeValuesForLOD,
+              ConstArrayView<float>{decoded::conditionalSlopeValues[params.maxLOD][params.currentLOD]});
+
+    const auto animatedMapCutValues = reader->getAnimatedMapCutValues();
+    ASSERT_EQ(animatedMapCutValues.size(), decoded::animatedMapLODs[params.maxLOD][0ul]);
+    ConstArrayView<float> animatedMapCutValuesForLOD{animatedMapCutValues.data(), animatdMapLOD};
+    ASSERT_EQ(animatedMapCutValuesForLOD,
+              ConstArrayView<float>{decoded::conditionalCutValues[params.maxLOD][params.currentLOD]});
+}
+
+static void verifyGeometry(GeometryReader* reader, const LODParameters& params) {
+    const auto meshCount = reader->getMeshCount();
+    ASSERT_EQ(meshCount, decoded::meshCount[params.maxLOD]);
+    for (std::uint16_t meshIndex = 0u; meshIndex < meshCount; ++meshIndex) {
+        const auto vertexPositionCount = reader->getVertexPositionCount(meshIndex);
+        ASSERT_EQ(vertexPositionCount, decoded::vertexPositions[params.maxLOD][meshIndex].size());
+        for (std::uint32_t vertexIndex = 0u; vertexIndex < vertexPositionCount; ++vertexIndex) {
+            ASSERT_EQ(reader->getVertexPosition(meshIndex, vertexIndex),
+                      decoded::vertexPositions[params.maxLOD][meshIndex][vertexIndex]);
+        }
+
+        const auto vertexTextureCoordinateCount = reader->getVertexTextureCoordinateCount(meshIndex);
+        ASSERT_EQ(vertexTextureCoordinateCount, decoded::vertexTextureCoordinates[params.maxLOD][meshIndex].size());
+        for (std::uint32_t texCoordIndex = 0u; texCoordIndex < vertexTextureCoordinateCount; ++texCoordIndex) {
+            const auto& textureCoordinate = reader->getVertexTextureCoordinate(meshIndex, texCoordIndex);
+            const auto& expectedTextureCoordinate =
+                decoded::vertexTextureCoordinates[params.maxLOD][meshIndex][texCoordIndex];
+            ASSERT_EQ(textureCoordinate.u, expectedTextureCoordinate.u);
+            ASSERT_EQ(textureCoordinate.v, expectedTextureCoordinate.v);
+        }
+
+        const auto vertexNormalCount = reader->getVertexNormalCount(meshIndex);
+        ASSERT_EQ(vertexNormalCount, decoded::vertexNormals[params.maxLOD][meshIndex].size());
+        for (std::uint32_t normalIndex = 0u; normalIndex < vertexNormalCount; ++normalIndex) {
+            ASSERT_EQ(reader->getVertexNormal(meshIndex, normalIndex),
+                      decoded::vertexNormals[params.maxLOD][meshIndex][normalIndex]);
+        }
+
+        const auto vertexLayoutCount = reader->getVertexLayoutCount(meshIndex);
+        ASSERT_EQ(vertexLayoutCount, decoded::vertexLayouts[params.maxLOD][meshIndex].size());
+        for (std::uint32_t layoutIndex = 0u; layoutIndex < vertexLayoutCount; ++layoutIndex) {
+            const auto& layout = reader->getVertexLayout(meshIndex, layoutIndex);
+            const auto& expectedLayout = decoded::vertexLayouts[params.maxLOD][meshIndex][layoutIndex];
+            ASSERT_EQ(layout.position, expectedLayout.position);
+            ASSERT_EQ(layout.textureCoordinate, expectedLayout.textureCoordinate);
+            ASSERT_EQ(layout.normal, expectedLayout.normal);
+        }
+
+        const auto faceCount = reader->getFaceCount(meshIndex);
+        ASSERT_EQ(faceCount, decoded::faces[params.maxLOD][meshIndex].size());
+        for (std::uint32_t faceIndex = 0u; faceIndex < faceCount; ++faceIndex) {
+            ASSERT_EQ(reader->getFaceVertexLayoutIndices(meshIndex, faceIndex),
+                      ConstArrayView<std::uint32_t>{decoded::faces[params.maxLOD][meshIndex][faceIndex]});
+        }
+
+        ASSERT_EQ(reader->getMaximumInfluencePerVertex(meshIndex),
+                  decoded::maxInfluencePerVertex[params.maxLOD][meshIndex]);
+
+        for (std::uint32_t vertexIndex = 0u; vertexIndex < vertexPositionCount; ++vertexIndex) {
+            const auto skinWeights = reader->getSkinWeightsValues(meshIndex, vertexIndex);
+            const auto& expectedSkinWeights = decoded::skinWeightsValues[params.maxLOD][meshIndex][vertexIndex];
+            ASSERT_EQ(skinWeights, ConstArrayView<float>{expectedSkinWeights});
+
+            const auto jointIndices = reader->getSkinWeightsJointIndices(meshIndex, vertexIndex);
+            const auto& expectedJointIndices = decoded::skinWeightsJointIndices[params.maxLOD][meshIndex][vertexIndex];
+            ASSERT_EQ(jointIndices, ConstArrayView<std::uint16_t>{expectedJointIndices});
+        }
+
+        const auto blendShapeCount = reader->getBlendShapeTargetCount(meshIndex);
+        ASSERT_EQ(blendShapeCount, decoded::correctiveBlendShapeDeltas[params.maxLOD][meshIndex].size());
+        for (std::uint16_t blendShapeTargetIndex = 0u; blendShapeTargetIndex < blendShapeCount; ++blendShapeTargetIndex) {
+            const auto channelIndex = reader->getBlendShapeChannelIndex(meshIndex, blendShapeTargetIndex);
+            ASSERT_EQ(channelIndex, decoded::correctiveBlendShapeIndices[params.maxLOD][meshIndex][blendShapeTargetIndex]);
+
+            const auto deltaCount = reader->getBlendShapeTargetDeltaCount(meshIndex, blendShapeTargetIndex);
+            ASSERT_EQ(deltaCount,
+                      decoded::correctiveBlendShapeDeltas[params.maxLOD][meshIndex][blendShapeTargetIndex].size());
+
+            for (std::uint32_t deltaIndex = 0u; deltaIndex < deltaCount; ++deltaIndex) {
+                const auto& delta = reader->getBlendShapeTargetDelta(meshIndex, blendShapeTargetIndex, deltaIndex);
+                const auto& expectedDelta =
+                    decoded::correctiveBlendShapeDeltas[params.maxLOD][meshIndex][blendShapeTargetIndex][deltaIndex];
+                ASSERT_EQ(delta, expectedDelta);
+            }
+
+            const auto vertexIndices = reader->getBlendShapeTargetVertexIndices(meshIndex, blendShapeTargetIndex);
+            const auto& expectedVertexIndices =
+                decoded::correctiveBlendShapeVertexIndices[params.maxLOD][meshIndex][blendShapeTargetIndex];
+            ASSERT_EQ(vertexIndices, ConstArrayView<std::uint32_t>{expectedVertexIndices});
+        }
+    }
+}
+
+static void assertReaderHasAllData(Reader* reader, const LODParameters& params) {
+    verifyDescriptor(reader, params);
+    verifyDefinition(reader, params);
+    verifyBehavior(reader, params);
+    verifyGeometry(reader, params);
+}
+
+TEST_P(StreamReadWriteIntegrationTest, VerifyAllDNADataAfterSetFromIsUsed) {
+    pma::AlignedMemoryResource memRes;
+
+    dnatests::FakeStream source;
+    const auto bytes = raw::getBytes();
+    source.write(reinterpret_cast<const char*>(bytes.data()), bytes.size());
+    source.seek(0);
+
+    auto sourceReader = StreamReader::create(&source, DataLayer::All, 0u, &memRes);
+    sourceReader->read();
+
+    dnatests::FakeStream clone;
+    auto cloneWriter = StreamWriter::create(&clone, &memRes);
+    cloneWriter->setFrom(sourceReader);
+    cloneWriter->write();
+
+    clone.seek(0ul);
+    auto cloneReader = StreamReader::create(&clone, DataLayer::All, params.maxLOD, &memRes);
+    cloneReader->read();
+
+    assertReaderHasAllData(cloneReader, params);
+
+    StreamReader::destroy(cloneReader);
+    StreamWriter::destroy(cloneWriter);
+    StreamReader::destroy(sourceReader);
+}
+
+INSTANTIATE_TEST_SUITE_P(StreamReadWriteIntegrationTest, StreamReadWriteIntegrationTest, ::testing::Values(
+                             // maxLOD, currentLOD
+                             LODParameters{0u, 0u},
+                             LODParameters{0u, 1u},
+                             LODParameters{1u, 0u}
+                             ));
+
+}  // namespace dna
diff --git a/Engine/Plugins/Animation/RigLogic/Source/RigLogicLibTest/Private/dnatests/TestStreamReadWriteIntegration.h b/Engine/Plugins/Animation/RigLogic/Source/RigLogicLibTest/Private/dnatests/TestStreamReadWriteIntegration.h
new file mode 100644
--- /dev/null
+++ b/Engine/Plugins/Animation/RigLogic/Source/RigLogicLibTest/Private/dnatests/TestStreamReadWriteIntegration.h
@@ -0,0 +1,21 @@
+#pragma once
+
+#include "dnatests/Defs.h"
+
+#include "dna/StreamReader.h"
+
+
+struct LODParameters {
+    std::uint16_t maxLOD;
+    std::uint16_t currentLOD;
+};
+
+class StreamReadWriteIntegrationTest : public ::testing::TestWithParam<LODParameters> {
+    protected:
+        void SetUp() override {
+            params = GetParam();
+        }
+
+    protected:
+        LODParameters params;
+};
diff --git a/Engine/Plugins/Animation/RigLogic/Source/RigLogicLibTest/Private/dnatests/TestStreamWriter.cpp b/Engine/Plugins/Animation/RigLogic/Source/RigLogicLibTest/Private/dnatests/TestStreamWriter.cpp
new file mode 100644
--- /dev/null
+++ b/Engine/Plugins/Animation/RigLogic/Source/RigLogicLibTest/Private/dnatests/TestStreamWriter.cpp
@@ -0,0 +1,30 @@
+#include "dnatests/TestStreamWriter.h"
+
+TEST_F(StreamWriterTest, SetName) {
+    writer->setLODCount(1);
+    writer->setName("test");
+    writer->write();
+
+    reader->read();
+    auto gotName = reader->getName();
+    EXPECT_STREQ(gotName, "test");
+}
+
+TEST_F(StreamWriterTest, SetControlName) {
+    writer->setLODCount(1);
+    writer->setRawControlName(0ul, "test");
+    writer->write();
+
+    reader->read();
+    EXPECT_EQ(reader->getRawControlCount(), 1ul);
+    EXPECT_STREQ(reader->getRawControlName(0ul), "test");
+}
+
+TEST_F(StreamWriterTest, SetJointRowCount) {
+    writer->setLODCount(1);
+    writer->setJointRowCount(1ul);
+    writer->write();
+
+    reader->read();
+    EXPECT_EQ(reader->getJointRowCount(), 1ul);
+}
diff --git a/Engine/Plugins/Animation/RigLogic/Source/RigLogicLibTest/Private/dnatests/TestStreamWriter.h b/Engine/Plugins/Animation/RigLogic/Source/RigLogicLibTest/Private/dnatests/TestStreamWriter.h
new file mode 100644
--- /dev/null
+++ b/Engine/Plugins/Animation/RigLogic/Source/RigLogicLibTest/Private/dnatests/TestStreamWriter.h
@@ -0,0 +1,28 @@
+#pragma once
+
+#include "dnatests/Defs.h"
+#include "dnatests/FakeStream.h"
+
+#include "dna/StreamReader.h"
+#include "dna/StreamWriter.h"
+
+#include <pma/resources/DefaultMemoryResource.h>
+
+class StreamWriterTest : public ::testing::Test {
+    protected:
+        void SetUp() override {
+            writer = dna::StreamWriter::create(&stream, &memRes);
+            reader = dna::StreamReader::create(&stream, dna::DataLayer::All, 0u, &memRes);
+        }
+
+        void TearDown() override {
+            dna::StreamReader::destroy(reader);
+            dna::StreamWriter::destroy(writer);
+        }
+
+    protected:
+        dnatests::FakeStream stream;
+        pma::DefaultMemoryResource memRes;
+        dna::StreamWriter* writer;
+        dna::StreamReader* reader;
+};
diff --git a/Engine/Plugins/Animation/RigLogic/Source/RigLogicLibTest/Private/pmatests/Defs.h b/Engine/Plugins/Animation/RigLogic/Source/RigLogicLibTest/Private/pmatests/Defs.h
new file mode 100644
--- /dev/null
+++ b/Engine/Plugins/Animation/RigLogic/Source/RigLogicLibTest/Private/pmatests/Defs.h
@@ -0,0 +1,48 @@
+#pragma once
+
+#ifdef _MSC_VER
+    #pragma warning(disable:4668)
+    #pragma warning(disable:4996)
+    #pragma warning(push)
+    #pragma warning(disable : 4365 4987)
+#endif
+#ifdef __clang__
+    #pragma clang diagnostic push
+    #pragma clang diagnostic ignored "-Wglobal-constructors"
+#endif
+#include <gtest/gtest.h>
+#ifdef __clang__
+    #pragma clang diagnostic pop
+#endif
+#ifdef _MSC_VER
+    #pragma warning(pop)
+    #pragma warning(default:4668)
+    #pragma warning(default:4996)
+#endif
+
+#ifndef INSTANTIATE_TEST_SUITE_P
+    #define INSTANTIATE_TEST_SUITE_P INSTANTIATE_TEST_CASE_P
+#endif
+
+#ifndef TYPED_TEST_SUITE_P
+    #define TYPED_TEST_SUITE_P TYPED_TEST_CASE_P
+#endif
+
+#ifndef REGISTER_TYPED_TEST_SUITE_P
+    #define REGISTER_TYPED_TEST_SUITE_P REGISTER_TYPED_TEST_CASE_P
+#endif
+
+#ifndef INSTANTIATE_TYPED_TEST_SUITE_P
+    #define INSTANTIATE_TYPED_TEST_SUITE_P INSTANTIATE_TYPED_TEST_CASE_P
+#endif
+
+#define ASSERT_ELEMENTS_EQ(result, expected, count)   \
+    for (std::size_t i = 0ul; i < count; ++i) {         \
+        ASSERT_EQ(result[i], expected[i]);              \
+    }
+
+#define ASSERT_ELEMENTS_NEAR(lhs, rhs, threshold)       \
+    ASSERT_EQ(lhs.size(), rhs.size());                  \
+    for (std::size_t i = 0ul; i < lhs.size(); ++i) {    \
+        ASSERT_NEAR(lhs[i], rhs[i], threshold);         \
+    }
diff --git a/Engine/Plugins/Animation/RigLogic/Source/RigLogicLibTest/Private/pmatests/resources/TestAlignedMemoryResource.cpp b/Engine/Plugins/Animation/RigLogic/Source/RigLogicLibTest/Private/pmatests/resources/TestAlignedMemoryResource.cpp
new file mode 100644
--- /dev/null
+++ b/Engine/Plugins/Animation/RigLogic/Source/RigLogicLibTest/Private/pmatests/resources/TestAlignedMemoryResource.cpp
@@ -0,0 +1,19 @@
+#include "pmatests/Defs.h"
+
+#include <pma/resources/AlignedMemoryResource.h>
+
+#include <cstdint>
+
+class AlignedMemoryResourceTest : public ::testing::TestWithParam<std::size_t> {
+};
+
+TEST_P(AlignedMemoryResourceTest, AlignmentRequirementHonored) {
+    std::size_t allocSize = 1024u;
+    std::size_t alignment = GetParam();
+    pma::AlignedMemoryResource amr;
+    void* ptr = amr.allocate(allocSize, alignment);
+    ASSERT_EQ(reinterpret_cast<std::uintptr_t>(ptr) & (alignment - 1), 0);
+    amr.deallocate(ptr, allocSize, alignment);
+}
+
+INSTANTIATE_TEST_SUITE_P(AlignedMemoryResourceTest, AlignedMemoryResourceTest, ::testing::Values(8, 64, 1024, 4096));
diff --git a/Engine/Plugins/Animation/RigLogic/Source/RigLogicLibTest/Private/pmatests/utils/TestManagedInstance.cpp b/Engine/Plugins/Animation/RigLogic/Source/RigLogicLibTest/Private/pmatests/utils/TestManagedInstance.cpp
new file mode 100644
--- /dev/null
+++ b/Engine/Plugins/Animation/RigLogic/Source/RigLogicLibTest/Private/pmatests/utils/TestManagedInstance.cpp
@@ -0,0 +1,71 @@
+#include "pmatests/Defs.h"
+
+#include "pma/utils/ManagedInstance.h"
+
+namespace pmatests {
+
+struct Stats {
+    int constructed;
+    int destructed;
+};
+
+class BaseType {
+    public:
+        virtual ~BaseType() = default;
+};
+
+class ManagableType : public BaseType {
+    public:
+        ManagableType(Stats& stats_) : stats{&stats_} {
+            stats->constructed += 1;
+        }
+
+        ~ManagableType() {
+            stats->destructed += 1;
+        }
+
+        ManagableType(const ManagableType&) = default;
+        ManagableType& operator=(const ManagableType&) = default;
+
+        ManagableType(ManagableType&&) = default;
+        ManagableType& operator=(ManagableType&&) = default;
+
+    private:
+        Stats* stats;
+};
+
+}  // namespace pmatests
+
+TEST(ManagedInstanceTest, UniqueInstance) {
+    pma::DefaultMemoryResource memRes;
+    pmatests::Stats stats{};
+    ASSERT_EQ(stats.constructed, 0);
+    ASSERT_EQ(stats.destructed, 0);
+    {
+        auto ptr = pma::UniqueInstance<pmatests::ManagableType, pmatests::BaseType>::with(&memRes).create(stats);
+        ASSERT_EQ(stats.constructed, 1);
+        ASSERT_EQ(stats.destructed, 0);
+    }
+    ASSERT_EQ(stats.constructed, 1);
+    ASSERT_EQ(stats.destructed, 1);
+}
+
+TEST(ManagedInstanceTest, SharedInstance) {
+    pma::DefaultMemoryResource memRes;
+    pmatests::Stats stats{};
+    ASSERT_EQ(stats.constructed, 0);
+    ASSERT_EQ(stats.destructed, 0);
+    {
+        std::shared_ptr<pmatests::BaseType> ptr;
+        {
+            auto ptr2 = pma::SharedInstance<pmatests::ManagableType, pmatests::BaseType>::with(&memRes).create(stats);
+            ptr = ptr2;
+            ASSERT_EQ(stats.constructed, 1);
+            ASSERT_EQ(stats.destructed, 0);
+        }
+        ASSERT_EQ(stats.constructed, 1);
+        ASSERT_EQ(stats.destructed, 0);
+    }
+    ASSERT_EQ(stats.constructed, 1);
+    ASSERT_EQ(stats.destructed, 1);
+}
diff --git a/Engine/Plugins/Animation/RigLogic/Source/RigLogicLibTest/Private/rltests/animatedmaps/AnimatedMapFixtures.cpp b/Engine/Plugins/Animation/RigLogic/Source/RigLogicLibTest/Private/rltests/animatedmaps/AnimatedMapFixtures.cpp
new file mode 100644
--- /dev/null
+++ b/Engine/Plugins/Animation/RigLogic/Source/RigLogicLibTest/Private/rltests/animatedmaps/AnimatedMapFixtures.cpp
@@ -0,0 +1,19 @@
+#include "rltests/animatedmaps/AnimatedMapFixtures.h"
+
+#include "riglogic/TypeDefs.h"
+#include "riglogic/animatedmaps/AnimatedMaps.h"
+#include "riglogic/conditionaltable/ConditionalTable.h"
+
+namespace defaults {
+
+const std::size_t lodCount = 1ul;
+const std::uint16_t lods[lodCount] = {2};
+
+}  // namespace calculation
+
+rl4::AnimatedMaps AnimatedMapsFactory::withDefaults(rl4::ConditionalTable&& conditionals, rl4::MemoryResource* memRes) {
+    return rl4::AnimatedMaps{
+        rl4::Vector<std::uint16_t>{defaults::lods, defaults::lods + defaults::lodCount, memRes},
+        std::move(conditionals)
+    };
+}
diff --git a/Engine/Plugins/Animation/RigLogic/Source/RigLogicLibTest/Private/rltests/animatedmaps/AnimatedMapFixtures.h b/Engine/Plugins/Animation/RigLogic/Source/RigLogicLibTest/Private/rltests/animatedmaps/AnimatedMapFixtures.h
new file mode 100644
--- /dev/null
+++ b/Engine/Plugins/Animation/RigLogic/Source/RigLogicLibTest/Private/rltests/animatedmaps/AnimatedMapFixtures.h
@@ -0,0 +1,13 @@
+#pragma once
+
+#include "riglogic/TypeDefs.h"
+#include "riglogic/animatedmaps/AnimatedMaps.h"
+#include "riglogic/conditionaltable/ConditionalTable.h"
+
+#include <cstddef>
+
+struct AnimatedMapsFactory {
+
+    static rl4::AnimatedMaps withDefaults(rl4::ConditionalTable&& conditionals, rl4::MemoryResource* memRes);
+
+};
diff --git a/Engine/Plugins/Animation/RigLogic/Source/RigLogicLibTest/Private/rltests/animatedmaps/TestAnimatedMaps.cpp b/Engine/Plugins/Animation/RigLogic/Source/RigLogicLibTest/Private/rltests/animatedmaps/TestAnimatedMaps.cpp
new file mode 100644
--- /dev/null
+++ b/Engine/Plugins/Animation/RigLogic/Source/RigLogicLibTest/Private/rltests/animatedmaps/TestAnimatedMaps.cpp
@@ -0,0 +1,33 @@
+#include "rltests/Defs.h"
+#include "rltests/animatedmaps/AnimatedMapFixtures.h"
+#include "rltests/conditionaltable/ConditionalTableFixtures.h"
+
+#include "riglogic/animatedmaps/AnimatedMaps.h"
+#include "riglogic/types/Aliases.h"
+
+#include <pma/resources/AlignedMemoryResource.h>
+
+#include <array>
+
+class AnimatedMapsTest : public ::testing::TestWithParam<std::uint16_t> {
+};
+
+TEST_P(AnimatedMapsTest, LODLimitsCondTableSize) {
+    pma::AlignedMemoryResource amr;
+    rl4::Vector<std::uint16_t> lods{2, 1};
+    auto conditionals = ConditionalTableFactory::withMultipleIODefaults(&amr);
+    rl4::AnimatedMaps animatedMaps{std::move(lods), std::move(conditionals)};
+
+    const auto lod = GetParam();
+    const float expected[2][2] = {
+        {0.3f, 0.6f},  // LOD0
+        {0.3f, 0.0f}  // LOD1
+    };
+    std::array<float, 2ul> outputs;
+    animatedMaps.calculate(rl4::ConstArrayView<float>{conditionalTableInputs},
+                           rl4::ArrayView<float>{outputs},
+                           lod);
+    ASSERT_ELEMENTS_EQ(outputs, expected[lod], 2ul);
+}
+
+INSTANTIATE_TEST_SUITE_P(AnimatedMapsTest, AnimatedMapsTest, ::testing::Values(0u, 1u));
diff --git a/Engine/Plugins/Animation/RigLogic/Source/RigLogicLibTest/Private/rltests/blendshapes/BlendShapeFixtures.cpp b/Engine/Plugins/Animation/RigLogic/Source/RigLogicLibTest/Private/rltests/blendshapes/BlendShapeFixtures.cpp
new file mode 100644
--- /dev/null
+++ b/Engine/Plugins/Animation/RigLogic/Source/RigLogicLibTest/Private/rltests/blendshapes/BlendShapeFixtures.cpp
@@ -0,0 +1,12 @@
+#include "rltests/blendshapes/BlendShapeFixtures.h"
+
+#include "riglogic/TypeDefs.h"
+#include "riglogic/blendshapes/BlendShapes.h"
+
+rl4::BlendShapes createTestBlendShapes(rl4::MemoryResource* memRes) {
+    return rl4::BlendShapes{
+        rl4::Vector<std::uint16_t>{std::begin(LODs), std::end(LODs), memRes},
+        rl4::Vector<std::uint16_t>{std::begin(blendShapeInputIndices), std::end(blendShapeInputIndices), memRes},
+        rl4::Vector<std::uint16_t>{std::begin(blendShapeOutputIndices), std::end(blendShapeOutputIndices), memRes}
+    };
+}
diff --git a/Engine/Plugins/Animation/RigLogic/Source/RigLogicLibTest/Private/rltests/blendshapes/BlendShapeFixtures.h b/Engine/Plugins/Animation/RigLogic/Source/RigLogicLibTest/Private/rltests/blendshapes/BlendShapeFixtures.h
new file mode 100644
--- /dev/null
+++ b/Engine/Plugins/Animation/RigLogic/Source/RigLogicLibTest/Private/rltests/blendshapes/BlendShapeFixtures.h
@@ -0,0 +1,26 @@
+#pragma once
+
+#include "riglogic/TypeDefs.h"
+#include "riglogic/blendshapes/BlendShapes.h"
+
+#include <pma/MemoryResource.h>
+
+#include <array>
+#include <cstddef>
+
+static const std::size_t LODCount = 4ul;
+static const std::uint16_t LOD0 = 10u;
+static const std::uint16_t LOD1 = 7u;
+static const std::uint16_t LOD2 = 4u;
+static const std::uint16_t LOD3 = 1u;
+
+static const std::uint16_t LODs[] = {LOD0, LOD1, LOD2, LOD3};
+
+static const std::array<float, 15ul> blendShapeInputs = {
+    1.0f, 2.0f, 3.0f, 4.0f, 5.0f, 6.0f, 7.0f, 8.0f, 9.0f, 10.0f, 11.0f, 12.0f, 13.0f, 14.0f, 15.0f
+};
+static const std::uint16_t blendShapeInputIndices[] = {0, 1, 3, 4, 6, 7, 8, 10, 12, 14};
+static const std::uint16_t blendShapeOutputIndices[] = {0, 2, 1, 3, 4, 5, 6, 7, 8, 9};
+static const std::array<float, 10ul> blendShapeOutputs = {1.0f, 4.0f, 2.0f, 5.0f, 7.0f, 8.0f, 9.0f, 11.0f, 13.0f, 15.0f};
+
+rl4::BlendShapes createTestBlendShapes(rl4::MemoryResource* memRes);
diff --git a/Engine/Plugins/Animation/RigLogic/Source/RigLogicLibTest/Private/rltests/blendshapes/TestBlendShapes.cpp b/Engine/Plugins/Animation/RigLogic/Source/RigLogicLibTest/Private/rltests/blendshapes/TestBlendShapes.cpp
new file mode 100644
--- /dev/null
+++ b/Engine/Plugins/Animation/RigLogic/Source/RigLogicLibTest/Private/rltests/blendshapes/TestBlendShapes.cpp
@@ -0,0 +1,36 @@
+#include "rltests/Defs.h"
+#include "rltests/blendshapes/BlendShapeFixtures.h"
+
+#include "riglogic/blendshapes/BlendShapes.h"
+#include "riglogic/types/Aliases.h"
+#include "riglogic/utils/Extd.h"
+
+#include <pma/resources/AlignedMemoryResource.h>
+
+#include <array>
+
+class BlendShapesTest : public ::testing::TestWithParam<std::uint16_t> {
+};
+
+TEST_P(BlendShapesTest, Calculate) {
+    // Test input parameter - Which LOD level to test
+    const auto lod = GetParam();
+
+    pma::AlignedMemoryResource amr;
+    auto blendShapes = createTestBlendShapes(&amr);
+
+    std::vector<float> outputs(blendShapeOutputs.size());
+    blendShapes.calculate(rl4::ConstArrayView<float>{blendShapeInputs},
+                          rl4::ArrayView<float>{outputs},
+                          lod);
+
+    std::vector<float> expectedData{blendShapeOutputs.begin(), blendShapeOutputs.end()};
+    std::fill(extd::advanced(expectedData.begin(), LODs[lod]), expectedData.end(), 0.0f);
+    rl4::ConstArrayView<float> expected{expectedData};
+
+    for (std::size_t i = 0ul; i < blendShapeOutputs.size(); ++i) {
+        ASSERT_EQ(outputs[i], expected[i]);
+    }
+}
+
+INSTANTIATE_TEST_SUITE_P(BlendShapesTest, BlendShapesTest, ::testing::Values(0u, 1u, 2u, 3u));
diff --git a/Engine/Plugins/Animation/RigLogic/Source/RigLogicLibTest/Private/rltests/conditionaltable/TestConditionalTable.cpp b/Engine/Plugins/Animation/RigLogic/Source/RigLogicLibTest/Private/rltests/conditionaltable/TestConditionalTable.cpp
new file mode 100644
--- /dev/null
+++ b/Engine/Plugins/Animation/RigLogic/Source/RigLogicLibTest/Private/rltests/conditionaltable/TestConditionalTable.cpp
@@ -0,0 +1,86 @@
+#include "rltests/Defs.h"
+#include "rltests/conditionaltable/ConditionalTableFixtures.h"
+
+#include "riglogic/conditionaltable/ConditionalTable.h"
+
+#include <pma/resources/AlignedMemoryResource.h>
+
+struct CalcTestData {
+    float fromValues[2];
+    float toValues[2];
+    float expected[1];
+};
+
+class ConditionalTableTest : public ::testing::TestWithParam<CalcTestData> {
+};
+
+TEST_P(ConditionalTableTest, CheckCalculationBorderCases) {
+    pma::AlignedMemoryResource amr;
+    const std::size_t conditionalsSize = 2ul;
+    auto testData = GetParam();
+    auto conditionals = ConditionalTableFactory::withSingleIO(
+        rl4::Vector<float>{testData.fromValues, testData.fromValues + conditionalsSize, &amr},
+        rl4::Vector<float>{testData.toValues, testData.toValues + conditionalsSize, &amr},
+        &amr
+        );
+    float outputs[] = {0.0f};
+    conditionals.calculate(conditionalTableInputs.data(), outputs);
+    ASSERT_ELEMENTS_EQ(outputs, testData.expected, 1);
+}
+
+INSTANTIATE_TEST_SUITE_P(ConditionalTableTest,
+                         ConditionalTableTest,
+                         ::testing::Values(
+                             // {{fromValues}, {toValues}, {expected}}
+                             // In-value below from-value
+                             CalcTestData{{0.3f, 0.6f}, {0.6f, 1.0f}, {0.0f}},
+                             // In-value equals from-value
+                             CalcTestData{{0.1f, 0.6f}, {0.6f, 1.0f}, {0.0f}},
+                             // In-value equals to-value
+                             CalcTestData{{0.0f, 0.1f}, {0.6f, 1.0f}, {0.2f}},
+                             // In-value equals both from-value and to-value
+                             CalcTestData{{0.0f, 0.1f}, {0.1f, 1.0f}, {0.2f}},
+                             // In-value between from-value and to-value
+                             CalcTestData{{0.0f, 0.6f}, {0.6f, 1.0f}, {0.2f}},
+                             // In-value above to-value
+                             CalcTestData{{0.0f, 0.04f}, {0.04f, 0.09f}, {0.0f}},
+                             // Original output not reset, but summed with new value
+                             CalcTestData{{0.0f, 0.0f}, {0.2f, 0.2f}, {0.6f}}
+                             ));
+
+TEST(ConditionalTableTest, OutputClamped) {
+    const std::uint16_t inputCount = 1u;
+    const std::uint16_t outputCount = 1u;
+    rl4::Vector<std::uint16_t> inputIndices{0};
+    rl4::Vector<std::uint16_t> outputIndices{0};
+    rl4::Vector<float> slopeValues{1.0f};
+    rl4::Vector<float> cutValues{2.0f};
+    rl4::Vector<float> fromValues{0.0f};
+    rl4::Vector<float> toValues{1.0f};
+    float outputs[1] = {};
+    rl4::ConditionalTable conditionals{std::move(inputIndices),
+                                       std::move(outputIndices),
+                                       std::move(fromValues),
+                                       std::move(toValues),
+                                       std::move(slopeValues),
+                                       std::move(cutValues),
+                                       inputCount,
+                                       outputCount};
+    conditionals.calculate(conditionalTableInputs.data(), outputs);
+    const float expected[] = {1.0f};
+    ASSERT_ELEMENTS_EQ(outputs, expected, 1ul);
+}
+
+TEST(ConditionalTableTest, OutputIsResetOnEachCalculation) {
+    pma::AlignedMemoryResource amr;
+    auto conditionals = ConditionalTableFactory::withSingleIODefaults(&amr);
+
+    float outputs[2ul];
+    const float expected[] = {0.6f};
+
+    conditionals.calculate(conditionalTableInputs.data(), outputs);
+    ASSERT_ELEMENTS_EQ(outputs, expected, 1ul);
+
+    conditionals.calculate(conditionalTableInputs.data(), outputs);
+    ASSERT_ELEMENTS_EQ(outputs, expected, 1ul);
+}
diff --git a/Engine/Plugins/Animation/RigLogic/Source/RigLogicLibTest/Private/rltests/controls/TestControls.cpp b/Engine/Plugins/Animation/RigLogic/Source/RigLogicLibTest/Private/rltests/controls/TestControls.cpp
new file mode 100644
--- /dev/null
+++ b/Engine/Plugins/Animation/RigLogic/Source/RigLogicLibTest/Private/rltests/controls/TestControls.cpp
@@ -0,0 +1,43 @@
+#include "rltests/Defs.h"
+#include "rltests/conditionaltable/ConditionalTableFixtures.h"
+
+#include "riglogic/controls/Controls.h"
+#include "riglogic/types/Aliases.h"
+
+#include <pma/resources/AlignedMemoryResource.h>
+
+#include <array>
+
+TEST(ControlsTest, GUIToRawMapping) {
+    pma::AlignedMemoryResource amr;
+    auto conditionals = ConditionalTableFactory::withMultipleIODefaults(&amr);
+
+    rl4::PSDMatrix psds{0u, {}, {}, {}};
+    rl4::Controls controls{std::move(conditionals), std::move(psds)};
+
+    const rl4::Vector<float> guiControls{0.1f, 0.2f};
+    const rl4::Vector<float> expected{0.3f, 0.6f};
+
+    std::array<float, 2ul> outputs;
+    controls.mapGUIToRaw(rl4::ConstArrayView<float>{guiControls}, rl4::ArrayView<float>{outputs});
+
+    ASSERT_ELEMENTS_EQ(outputs, expected, 2ul);
+}
+
+TEST(ControlsTest, PSDsAppendToOutput) {
+    pma::AlignedMemoryResource amr;
+    auto conditionals = ConditionalTableFactory::withMultipleIODefaults(&amr);
+
+    rl4::Vector<std::uint16_t> rows{2u, 3u};
+    rl4::Vector<std::uint16_t> cols{0u, 1u};
+    rl4::Vector<float> values{4.0f, 3.0f};
+    rl4::PSDMatrix psds{2u, std::move(rows), std::move(cols), std::move(values)};
+
+    rl4::Controls controls{std::move(conditionals), std::move(psds)};
+
+    rl4::Vector<float> rawControls{0.1f, 0.2f, 0.0f, 0.0f};
+    const rl4::Vector<float> expected{0.1f, 0.2f, 0.4f, 0.6f};
+    controls.calculate(rl4::ArrayView<float>{rawControls});
+
+    ASSERT_ELEMENTS_EQ(rawControls, expected, 4ul);
+}
