Editgrating RigLogic plugin from Dev-CharacterTech to UE5/Main @ CL 13891876.

diff --git a/Engine/Plugins/Animation/RigLogic/Source/RigLogicLibTest/Private/rltests/joints/bpcm/JointFixturesBlock4.cpp b/Engine/Plugins/Animation/RigLogic/Source/RigLogicLibTest/Private/rltests/joints/bpcm/JointFixturesBlock4.cpp
new file mode 100644
--- /dev/null
+++ b/Engine/Plugins/Animation/RigLogic/Source/RigLogicLibTest/Private/rltests/joints/bpcm/JointFixturesBlock4.cpp
@@ -0,0 +1,929 @@
+#include "rltests/joints/bpcm/JointFixturesBlock4.h"
+
+#include "riglogic/TypeDefs.h"
+#include "riglogic/joints/bpcm/Extent.h"
+#include "riglogic/joints/bpcm/JointsEvaluator.h"
+
+#include <cstring>
+
+namespace block4 {
+
+namespace unoptimized {
+
+const std::uint16_t lodCount = 4u;
+const bpcm::Extent dimensions{75ul, 13ul};
+
+const Matrix<float> values = {
+    {  // Joint group 0 - Pad to block-4 => block-4
+        1.0f, 1.0f, 1.0f, 1.0f, 1.0f, 1.0f, 1.0f, 1.0f, 1.0f, 1.0f, 1.0f, 1.0f, 1.0f
+    },
+    {  // Joint group 1 - No padding => block-4
+        2.0f, 2.0f, 2.0f, 2.0f, 2.0f, 2.0f, 2.0f, 2.0f, 2.0f, 2.0f, 2.0f, 2.0f, 2.0f,
+        3.0f, 3.0f, 3.0f, 3.0f, 3.0f, 3.0f, 3.0f, 3.0f, 3.0f, 3.0f, 3.0f, 3.0f, 3.0f,
+        4.0f, 4.0f, 4.0f, 4.0f, 4.0f, 4.0f, 4.0f, 4.0f, 4.0f, 4.0f, 4.0f, 4.0f, 4.0f,
+        5.0f, 5.0f, 5.0f, 5.0f, 5.0f, 5.0f, 5.0f, 5.0f, 5.0f, 5.0f, 5.0f, 5.0f, 5.0f
+    },
+    {  // Joint group 2 - Pad to block-8 => block-8
+        6.0f, 6.0f, 6.0f, 6.0f, 6.0f, 6.0f, 6.0f, 6.0f, 6.0f, 6.0f, 6.0f, 6.0f, 6.0f,
+        7.0f, 7.0f, 7.0f, 7.0f, 7.0f, 7.0f, 7.0f, 7.0f, 7.0f, 7.0f, 7.0f, 7.0f, 7.0f,
+        8.0f, 8.0f, 8.0f, 8.0f, 8.0f, 8.0f, 8.0f, 8.0f, 8.0f, 8.0f, 8.0f, 8.0f, 8.0f,
+        9.0f, 9.0f, 9.0f, 9.0f, 9.0f, 9.0f, 9.0f, 9.0f, 9.0f, 9.0f, 9.0f, 9.0f, 9.0f,
+        10.0f, 10.0f, 10.0f, 10.0f, 10.0f, 10.0f, 10.0f, 10.0f, 10.0f, 10.0f, 10.0f, 10.0f, 10.0f,
+        11.0f, 11.0f, 11.0f, 11.0f, 11.0f, 11.0f, 11.0f, 11.0f, 11.0f, 11.0f, 11.0f, 11.0f, 11.0f,
+        12.0f, 12.0f, 12.0f, 12.0f, 12.0f, 12.0f, 12.0f, 12.0f, 12.0f, 12.0f, 12.0f, 12.0f, 12.0f
+    },
+    {  // Joint group 3 - No padding => block-8
+        13.0f, 13.0f, 13.0f, 13.0f, 13.0f, 13.0f, 13.0f, 13.0f, 13.0f, 13.0f, 13.0f, 13.0f, 13.0f,
+        14.0f, 14.0f, 14.0f, 14.0f, 14.0f, 14.0f, 14.0f, 14.0f, 14.0f, 14.0f, 14.0f, 14.0f, 14.0f,
+        15.0f, 15.0f, 15.0f, 15.0f, 15.0f, 15.0f, 15.0f, 15.0f, 15.0f, 15.0f, 15.0f, 15.0f, 15.0f,
+        16.0f, 16.0f, 16.0f, 16.0f, 16.0f, 16.0f, 16.0f, 16.0f, 16.0f, 16.0f, 16.0f, 16.0f, 16.0f,
+        17.0f, 17.0f, 17.0f, 17.0f, 17.0f, 17.0f, 17.0f, 17.0f, 17.0f, 17.0f, 17.0f, 17.0f, 17.0f,
+        18.0f, 18.0f, 18.0f, 18.0f, 18.0f, 18.0f, 18.0f, 18.0f, 18.0f, 18.0f, 18.0f, 18.0f, 18.0f,
+        19.0f, 19.0f, 19.0f, 19.0f, 19.0f, 19.0f, 19.0f, 19.0f, 19.0f, 19.0f, 19.0f, 19.0f, 19.0f,
+        20.0f, 20.0f, 20.0f, 20.0f, 20.0f, 20.0f, 20.0f, 20.0f, 20.0f, 20.0f, 20.0f, 20.0f, 20.0f
+    },
+    {  // Joint group 4 - Pad to block-4 => block-8 + block-4
+        21.0f, 21.0f, 21.0f, 21.0f, 21.0f, 21.0f, 21.0f, 21.0f, 21.0f, 21.0f, 21.0f, 21.0f, 21.0f,
+        22.0f, 22.0f, 22.0f, 22.0f, 22.0f, 22.0f, 22.0f, 22.0f, 22.0f, 22.0f, 22.0f, 22.0f, 22.0f,
+        23.0f, 23.0f, 23.0f, 23.0f, 23.0f, 23.0f, 23.0f, 23.0f, 23.0f, 23.0f, 23.0f, 23.0f, 23.0f,
+        24.0f, 24.0f, 24.0f, 24.0f, 24.0f, 24.0f, 24.0f, 24.0f, 24.0f, 24.0f, 24.0f, 24.0f, 24.0f,
+        25.0f, 25.0f, 25.0f, 25.0f, 25.0f, 25.0f, 25.0f, 25.0f, 25.0f, 25.0f, 25.0f, 25.0f, 25.0f,
+        26.0f, 26.0f, 26.0f, 26.0f, 26.0f, 26.0f, 26.0f, 26.0f, 26.0f, 26.0f, 26.0f, 26.0f, 26.0f,
+        27.0f, 27.0f, 27.0f, 27.0f, 27.0f, 27.0f, 27.0f, 27.0f, 27.0f, 27.0f, 27.0f, 27.0f, 27.0f,
+        28.0f, 28.0f, 28.0f, 28.0f, 28.0f, 28.0f, 28.0f, 28.0f, 28.0f, 28.0f, 28.0f, 28.0f, 28.0f,
+        29.0f, 29.0f, 29.0f, 29.0f, 29.0f, 29.0f, 29.0f, 29.0f, 29.0f, 29.0f, 29.0f, 29.0f, 29.0f,
+    },
+    {  // Joint group 5 - No padding => block-8 + block-4
+        30.0f, 30.0f, 30.0f, 30.0f, 30.0f, 30.0f, 30.0f, 30.0f, 30.0f, 30.0f, 30.0f, 30.0f, 30.0f,
+        31.0f, 31.0f, 31.0f, 31.0f, 31.0f, 31.0f, 31.0f, 31.0f, 31.0f, 31.0f, 31.0f, 31.0f, 31.0f,
+        32.0f, 32.0f, 32.0f, 32.0f, 32.0f, 32.0f, 32.0f, 32.0f, 32.0f, 32.0f, 32.0f, 32.0f, 32.0f,
+        33.0f, 33.0f, 33.0f, 33.0f, 33.0f, 33.0f, 33.0f, 33.0f, 33.0f, 33.0f, 33.0f, 33.0f, 33.0f,
+        34.0f, 34.0f, 34.0f, 34.0f, 34.0f, 34.0f, 34.0f, 34.0f, 34.0f, 34.0f, 34.0f, 34.0f, 34.0f,
+        35.0f, 35.0f, 35.0f, 35.0f, 35.0f, 35.0f, 35.0f, 35.0f, 35.0f, 35.0f, 35.0f, 35.0f, 35.0f,
+        36.0f, 36.0f, 36.0f, 36.0f, 36.0f, 36.0f, 36.0f, 36.0f, 36.0f, 36.0f, 36.0f, 36.0f, 36.0f,
+        37.0f, 37.0f, 37.0f, 37.0f, 37.0f, 37.0f, 37.0f, 37.0f, 37.0f, 37.0f, 37.0f, 37.0f, 37.0f,
+        38.0f, 38.0f, 38.0f, 38.0f, 38.0f, 38.0f, 38.0f, 38.0f, 38.0f, 38.0f, 38.0f, 38.0f, 38.0f,
+        39.0f, 39.0f, 39.0f, 39.0f, 39.0f, 39.0f, 39.0f, 39.0f, 39.0f, 39.0f, 39.0f, 39.0f, 39.0f,
+        40.0f, 40.0f, 40.0f, 40.0f, 40.0f, 40.0f, 40.0f, 40.0f, 40.0f, 40.0f, 40.0f, 40.0f, 40.0f,
+        41.0f, 41.0f, 41.0f, 41.0f, 41.0f, 41.0f, 41.0f, 41.0f, 41.0f, 41.0f, 41.0f, 41.0f, 41.0f,
+    },
+    {  // Joint group 6 - Pad to block-8 => block-8 + block-8
+        42.0f, 42.0f, 42.0f, 42.0f, 42.0f, 42.0f, 42.0f, 42.0f, 42.0f, 42.0f, 42.0f, 42.0f, 42.0f,
+        43.0f, 43.0f, 43.0f, 43.0f, 43.0f, 43.0f, 43.0f, 43.0f, 43.0f, 43.0f, 43.0f, 43.0f, 43.0f,
+        44.0f, 44.0f, 44.0f, 44.0f, 44.0f, 44.0f, 44.0f, 44.0f, 44.0f, 44.0f, 44.0f, 44.0f, 44.0f,
+        45.0f, 45.0f, 45.0f, 45.0f, 45.0f, 45.0f, 45.0f, 45.0f, 45.0f, 45.0f, 45.0f, 45.0f, 45.0f,
+        46.0f, 46.0f, 46.0f, 46.0f, 46.0f, 46.0f, 46.0f, 46.0f, 46.0f, 46.0f, 46.0f, 46.0f, 46.0f,
+        47.0f, 47.0f, 47.0f, 47.0f, 47.0f, 47.0f, 47.0f, 47.0f, 47.0f, 47.0f, 47.0f, 47.0f, 47.0f,
+        48.0f, 48.0f, 48.0f, 48.0f, 48.0f, 48.0f, 48.0f, 48.0f, 48.0f, 48.0f, 48.0f, 48.0f, 48.0f,
+        49.0f, 49.0f, 49.0f, 49.0f, 49.0f, 49.0f, 49.0f, 49.0f, 49.0f, 49.0f, 49.0f, 49.0f, 49.0f,
+        50.0f, 50.0f, 50.0f, 50.0f, 50.0f, 50.0f, 50.0f, 50.0f, 50.0f, 50.0f, 50.0f, 50.0f, 50.0f,
+        51.0f, 51.0f, 51.0f, 51.0f, 51.0f, 51.0f, 51.0f, 51.0f, 51.0f, 51.0f, 51.0f, 51.0f, 51.0f,
+        52.0f, 52.0f, 52.0f, 52.0f, 52.0f, 52.0f, 52.0f, 52.0f, 52.0f, 52.0f, 52.0f, 52.0f, 52.0f,
+        53.0f, 53.0f, 53.0f, 53.0f, 53.0f, 53.0f, 53.0f, 53.0f, 53.0f, 53.0f, 53.0f, 53.0f, 53.0f,
+        54.0f, 54.0f, 54.0f, 54.0f, 54.0f, 54.0f, 54.0f, 54.0f, 54.0f, 54.0f, 54.0f, 54.0f, 54.0f,
+        55.0f, 55.0f, 55.0f, 55.0f, 55.0f, 55.0f, 55.0f, 55.0f, 55.0f, 55.0f, 55.0f, 55.0f, 55.0f,
+    },
+    {  // Joint group 7 - No padding => block-8 + block-8
+        56.0f, 56.0f, 56.0f, 56.0f, 56.0f, 56.0f, 56.0f, 56.0f, 56.0f, 56.0f, 56.0f, 56.0f, 56.0f,
+        57.0f, 57.0f, 57.0f, 57.0f, 57.0f, 57.0f, 57.0f, 57.0f, 57.0f, 57.0f, 57.0f, 57.0f, 57.0f,
+        58.0f, 58.0f, 58.0f, 58.0f, 58.0f, 58.0f, 58.0f, 58.0f, 58.0f, 58.0f, 58.0f, 58.0f, 58.0f,
+        59.0f, 59.0f, 59.0f, 59.0f, 59.0f, 59.0f, 59.0f, 59.0f, 59.0f, 59.0f, 59.0f, 59.0f, 59.0f,
+        60.0f, 60.0f, 60.0f, 60.0f, 60.0f, 60.0f, 60.0f, 60.0f, 60.0f, 60.0f, 60.0f, 60.0f, 60.0f,
+        61.0f, 61.0f, 61.0f, 61.0f, 61.0f, 61.0f, 61.0f, 61.0f, 61.0f, 61.0f, 61.0f, 61.0f, 61.0f,
+        62.0f, 62.0f, 62.0f, 62.0f, 62.0f, 62.0f, 62.0f, 62.0f, 62.0f, 62.0f, 62.0f, 62.0f, 62.0f,
+        63.0f, 63.0f, 63.0f, 63.0f, 63.0f, 63.0f, 63.0f, 63.0f, 63.0f, 63.0f, 63.0f, 63.0f, 63.0f,
+        64.0f, 64.0f, 64.0f, 64.0f, 64.0f, 64.0f, 64.0f, 64.0f, 64.0f, 64.0f, 64.0f, 64.0f, 64.0f,
+        65.0f, 65.0f, 65.0f, 65.0f, 65.0f, 65.0f, 65.0f, 65.0f, 65.0f, 65.0f, 65.0f, 65.0f, 65.0f,
+        66.0f, 66.0f, 66.0f, 66.0f, 66.0f, 66.0f, 66.0f, 66.0f, 66.0f, 66.0f, 66.0f, 66.0f, 66.0f,
+        67.0f, 67.0f, 67.0f, 67.0f, 67.0f, 67.0f, 67.0f, 67.0f, 67.0f, 67.0f, 67.0f, 67.0f, 67.0f,
+        68.0f, 68.0f, 68.0f, 68.0f, 68.0f, 68.0f, 68.0f, 68.0f, 68.0f, 68.0f, 68.0f, 68.0f, 68.0f,
+        69.0f, 69.0f, 69.0f, 69.0f, 69.0f, 69.0f, 69.0f, 69.0f, 69.0f, 69.0f, 69.0f, 69.0f, 69.0f,
+        70.0f, 70.0f, 70.0f, 70.0f, 70.0f, 70.0f, 70.0f, 70.0f, 70.0f, 70.0f, 70.0f, 70.0f, 70.0f,
+        71.0f, 71.0f, 71.0f, 71.0f, 71.0f, 71.0f, 71.0f, 71.0f, 71.0f, 71.0f, 71.0f, 71.0f, 71.0f,
+    },
+    {  // Joint group 8 - Pad to block-4 => block-4
+        1.0f, 1.0f, 1.0f, 1.0f, 1.0f,
+        2.0f, 2.0f, 2.0f, 2.0f, 2.0f,
+    },
+    {  // Joint group 9 - Pad to block-4 => block-4
+        1.0f, 1.0f, 1.0f, 1.0f, 1.0f,
+        2.0f, 2.0f, 2.0f, 2.0f, 2.0f
+    }
+};
+const Matrix<std::uint16_t> inputIndices = {
+    {  // Joint group 0
+        12, 11, 10, 9, 8, 7, 6, 5, 4, 3, 2, 1, 0
+    },
+    {  // Joint group 1
+        12, 11, 10, 9, 8, 7, 6, 5, 4, 3, 2, 1, 0
+    },
+    {  // Joint group 2
+        12, 11, 10, 9, 8, 7, 6, 5, 4, 3, 2, 1, 0
+    },
+    {  // Joint group 3
+        12, 11, 10, 9, 8, 7, 6, 5, 4, 3, 2, 1, 0
+    },
+    {  // Joint group 4
+        12, 11, 10, 9, 8, 7, 6, 5, 4, 3, 2, 1, 0
+    },
+    {  // Joint group 5
+        12, 11, 10, 9, 8, 7, 6, 5, 4, 3, 2, 1, 0
+    },
+    {  // Joint group 6
+        12, 11, 10, 9, 8, 7, 6, 5, 4, 3, 2, 1, 0
+    },
+    {  // Joint group 7
+        12, 11, 10, 9, 8, 7, 6, 5, 4, 3, 2, 1, 0
+    },
+    {  // Joint group 8
+        1, 2, 3, 4, 5
+    },
+    {  // Joint group 9
+        7, 8, 9, 10, 11
+    }
+};
+const Matrix<std::uint16_t> outputIndices = {
+    {  // Joint group 0
+        0
+    },
+    {  // Joint group 1
+        1, 2, 3, 4
+    },
+    {  // Joint group 2
+        5, 6, 7, 8, 9, 10, 11
+    },
+    {  // Joint group 3
+        12, 13, 14, 15, 16, 17, 18, 19
+    },
+    {  // Joint group 4
+        20, 21, 22, 23, 24, 25, 26, 27, 28
+    },
+    {  // Joint group 5
+        29, 30, 31, 32, 33, 34, 35, 36, 37, 38, 39, 40
+    },
+    {  // Joint group 6
+        41, 42, 43, 44, 45, 46, 47, 48, 49, 50, 51, 52, 53, 54
+    },
+    {  // Joint group 7
+        55, 56, 57, 58, 59, 60, 61, 62, 63, 64, 65, 66, 67, 68, 69, 70
+    },
+    {  // Joint group 8
+        71, 72
+    },
+    {  // Joint group 9
+        73, 74
+    }
+};
+const Matrix<std::uint16_t> lods = {
+    {  // Joint group 0
+        1, 1, 1, 1
+    },
+    {  // Joint group 1
+        4, 3, 2, 1
+    },
+    {  // Joint group 2
+        7, 5, 4, 1
+    },
+    {  // Joint group 3
+        8, 6, 4, 3
+    },
+    {  // Joint group 4
+        9, 7, 4, 2
+    },
+    {  // Joint group 5
+        12, 9, 7, 3
+    },
+    {  // Joint group 6
+        14, 13, 11, 9
+    },
+    {  // Joint group 7
+        16, 12, 8, 4
+    },
+    {  // Joint group 8
+        2, 2, 2, 2
+    },
+    {  // Joint group 9
+        2, 2, 2, 2
+    }
+};
+const Vector<bpcm::Extent> subMatrices = {
+    {1, 13}, {4, 13}, {7, 13}, {8, 13}, {9, 13}, {12, 13}, {14, 13}, {16, 13}, {2, 5}, {2, 5}
+};
+
+const Vector<float> neutralJointTranslationXs = {0.0f, 6.0f, 12.0f};
+const Vector<float> neutralJointTranslationYs = {1.0f, 7.0f, 13.0f};
+const Vector<float> neutralJointTranslationZs = {2.0f, 8.0f, 14.0f};
+const Vector<float> neutralJointRotationXs = {3.0f, 9.0f, 15.0f};
+const Vector<float> neutralJointRotationYs = {4.0f, 10.0f, 16.0f};
+const Vector<float> neutralJointRotationZs = {5.0f, 11.0f, 17.0f};
+
+}  // namespace unoptimized
+
+namespace optimized {
+
+const bpcm::Extent dimensions{88ul, 13ul};
+
+const AlignedMatrix<float> floatValues = {
+    {  // Joint group 0
+        1.0f, 0.0f, 0.0f, 0.0f,
+        1.0f, 0.0f, 0.0f, 0.0f,
+        1.0f, 0.0f, 0.0f, 0.0f,
+        1.0f, 0.0f, 0.0f, 0.0f,
+        1.0f, 0.0f, 0.0f, 0.0f,
+        1.0f, 0.0f, 0.0f, 0.0f,
+        1.0f, 0.0f, 0.0f, 0.0f,
+        1.0f, 0.0f, 0.0f, 0.0f,
+        1.0f, 0.0f, 0.0f, 0.0f,
+        1.0f, 0.0f, 0.0f, 0.0f,
+        1.0f, 0.0f, 0.0f, 0.0f,
+        1.0f, 0.0f, 0.0f, 0.0f,
+        1.0f, 0.0f, 0.0f, 0.0f
+    },
+    {  // Joint group 1
+        2.0f, 3.0f, 4.0f, 5.0f,
+        2.0f, 3.0f, 4.0f, 5.0f,
+        2.0f, 3.0f, 4.0f, 5.0f,
+        2.0f, 3.0f, 4.0f, 5.0f,
+        2.0f, 3.0f, 4.0f, 5.0f,
+        2.0f, 3.0f, 4.0f, 5.0f,
+        2.0f, 3.0f, 4.0f, 5.0f,
+        2.0f, 3.0f, 4.0f, 5.0f,
+        2.0f, 3.0f, 4.0f, 5.0f,
+        2.0f, 3.0f, 4.0f, 5.0f,
+        2.0f, 3.0f, 4.0f, 5.0f,
+        2.0f, 3.0f, 4.0f, 5.0f,
+        2.0f, 3.0f, 4.0f, 5.0f
+    },
+    {  // Joint group 2
+        6.0f, 7.0f, 8.0f, 9.0f, 10.0f, 11.0f, 12.0f, 0.0f,
+        6.0f, 7.0f, 8.0f, 9.0f, 10.0f, 11.0f, 12.0f, 0.0f,
+        6.0f, 7.0f, 8.0f, 9.0f, 10.0f, 11.0f, 12.0f, 0.0f,
+        6.0f, 7.0f, 8.0f, 9.0f, 10.0f, 11.0f, 12.0f, 0.0f,
+        6.0f, 7.0f, 8.0f, 9.0f, 10.0f, 11.0f, 12.0f, 0.0f,
+        6.0f, 7.0f, 8.0f, 9.0f, 10.0f, 11.0f, 12.0f, 0.0f,
+        6.0f, 7.0f, 8.0f, 9.0f, 10.0f, 11.0f, 12.0f, 0.0f,
+        6.0f, 7.0f, 8.0f, 9.0f, 10.0f, 11.0f, 12.0f, 0.0f,
+        6.0f, 7.0f, 8.0f, 9.0f, 10.0f, 11.0f, 12.0f, 0.0f,
+        6.0f, 7.0f, 8.0f, 9.0f, 10.0f, 11.0f, 12.0f, 0.0f,
+        6.0f, 7.0f, 8.0f, 9.0f, 10.0f, 11.0f, 12.0f, 0.0f,
+        6.0f, 7.0f, 8.0f, 9.0f, 10.0f, 11.0f, 12.0f, 0.0f,
+        6.0f, 7.0f, 8.0f, 9.0f, 10.0f, 11.0f, 12.0f, 0.0f,
+    },
+    {  // Joint group 3
+        13.0f, 14.0f, 15.0f, 16.0f, 17.0f, 18.0f, 19.0f, 20.0f,
+        13.0f, 14.0f, 15.0f, 16.0f, 17.0f, 18.0f, 19.0f, 20.0f,
+        13.0f, 14.0f, 15.0f, 16.0f, 17.0f, 18.0f, 19.0f, 20.0f,
+        13.0f, 14.0f, 15.0f, 16.0f, 17.0f, 18.0f, 19.0f, 20.0f,
+        13.0f, 14.0f, 15.0f, 16.0f, 17.0f, 18.0f, 19.0f, 20.0f,
+        13.0f, 14.0f, 15.0f, 16.0f, 17.0f, 18.0f, 19.0f, 20.0f,
+        13.0f, 14.0f, 15.0f, 16.0f, 17.0f, 18.0f, 19.0f, 20.0f,
+        13.0f, 14.0f, 15.0f, 16.0f, 17.0f, 18.0f, 19.0f, 20.0f,
+        13.0f, 14.0f, 15.0f, 16.0f, 17.0f, 18.0f, 19.0f, 20.0f,
+        13.0f, 14.0f, 15.0f, 16.0f, 17.0f, 18.0f, 19.0f, 20.0f,
+        13.0f, 14.0f, 15.0f, 16.0f, 17.0f, 18.0f, 19.0f, 20.0f,
+        13.0f, 14.0f, 15.0f, 16.0f, 17.0f, 18.0f, 19.0f, 20.0f,
+        13.0f, 14.0f, 15.0f, 16.0f, 17.0f, 18.0f, 19.0f, 20.0f,
+    },
+    {  // Joint group 4
+        21.0f, 22.0f, 23.0f, 24.0f, 25.0f, 26.0f, 27.0f, 28.0f,
+        21.0f, 22.0f, 23.0f, 24.0f, 25.0f, 26.0f, 27.0f, 28.0f,
+        21.0f, 22.0f, 23.0f, 24.0f, 25.0f, 26.0f, 27.0f, 28.0f,
+        21.0f, 22.0f, 23.0f, 24.0f, 25.0f, 26.0f, 27.0f, 28.0f,
+        21.0f, 22.0f, 23.0f, 24.0f, 25.0f, 26.0f, 27.0f, 28.0f,
+        21.0f, 22.0f, 23.0f, 24.0f, 25.0f, 26.0f, 27.0f, 28.0f,
+        21.0f, 22.0f, 23.0f, 24.0f, 25.0f, 26.0f, 27.0f, 28.0f,
+        21.0f, 22.0f, 23.0f, 24.0f, 25.0f, 26.0f, 27.0f, 28.0f,
+        21.0f, 22.0f, 23.0f, 24.0f, 25.0f, 26.0f, 27.0f, 28.0f,
+        21.0f, 22.0f, 23.0f, 24.0f, 25.0f, 26.0f, 27.0f, 28.0f,
+        21.0f, 22.0f, 23.0f, 24.0f, 25.0f, 26.0f, 27.0f, 28.0f,
+        21.0f, 22.0f, 23.0f, 24.0f, 25.0f, 26.0f, 27.0f, 28.0f,
+        21.0f, 22.0f, 23.0f, 24.0f, 25.0f, 26.0f, 27.0f, 28.0f,
+        29.0f, 0.0f, 0.0f, 0.0f,
+        29.0f, 0.0f, 0.0f, 0.0f,
+        29.0f, 0.0f, 0.0f, 0.0f,
+        29.0f, 0.0f, 0.0f, 0.0f,
+        29.0f, 0.0f, 0.0f, 0.0f,
+        29.0f, 0.0f, 0.0f, 0.0f,
+        29.0f, 0.0f, 0.0f, 0.0f,
+        29.0f, 0.0f, 0.0f, 0.0f,
+        29.0f, 0.0f, 0.0f, 0.0f,
+        29.0f, 0.0f, 0.0f, 0.0f,
+        29.0f, 0.0f, 0.0f, 0.0f,
+        29.0f, 0.0f, 0.0f, 0.0f,
+        29.0f, 0.0f, 0.0f, 0.0f,
+    },
+    {  // Joint group 5
+        30.0f, 31.0f, 32.0f, 33.0f, 34.0f, 35.0f, 36.0f, 37.0f,
+        30.0f, 31.0f, 32.0f, 33.0f, 34.0f, 35.0f, 36.0f, 37.0f,
+        30.0f, 31.0f, 32.0f, 33.0f, 34.0f, 35.0f, 36.0f, 37.0f,
+        30.0f, 31.0f, 32.0f, 33.0f, 34.0f, 35.0f, 36.0f, 37.0f,
+        30.0f, 31.0f, 32.0f, 33.0f, 34.0f, 35.0f, 36.0f, 37.0f,
+        30.0f, 31.0f, 32.0f, 33.0f, 34.0f, 35.0f, 36.0f, 37.0f,
+        30.0f, 31.0f, 32.0f, 33.0f, 34.0f, 35.0f, 36.0f, 37.0f,
+        30.0f, 31.0f, 32.0f, 33.0f, 34.0f, 35.0f, 36.0f, 37.0f,
+        30.0f, 31.0f, 32.0f, 33.0f, 34.0f, 35.0f, 36.0f, 37.0f,
+        30.0f, 31.0f, 32.0f, 33.0f, 34.0f, 35.0f, 36.0f, 37.0f,
+        30.0f, 31.0f, 32.0f, 33.0f, 34.0f, 35.0f, 36.0f, 37.0f,
+        30.0f, 31.0f, 32.0f, 33.0f, 34.0f, 35.0f, 36.0f, 37.0f,
+        30.0f, 31.0f, 32.0f, 33.0f, 34.0f, 35.0f, 36.0f, 37.0f,
+        38.0f, 39.0f, 40.0f, 41.0f,
+        38.0f, 39.0f, 40.0f, 41.0f,
+        38.0f, 39.0f, 40.0f, 41.0f,
+        38.0f, 39.0f, 40.0f, 41.0f,
+        38.0f, 39.0f, 40.0f, 41.0f,
+        38.0f, 39.0f, 40.0f, 41.0f,
+        38.0f, 39.0f, 40.0f, 41.0f,
+        38.0f, 39.0f, 40.0f, 41.0f,
+        38.0f, 39.0f, 40.0f, 41.0f,
+        38.0f, 39.0f, 40.0f, 41.0f,
+        38.0f, 39.0f, 40.0f, 41.0f,
+        38.0f, 39.0f, 40.0f, 41.0f,
+        38.0f, 39.0f, 40.0f, 41.0f,
+    },
+    {  // Joint group 6
+        42.0f, 43.0f, 44.0f, 45.0f, 46.0f, 47.0f, 48.0f, 49.0f,
+        42.0f, 43.0f, 44.0f, 45.0f, 46.0f, 47.0f, 48.0f, 49.0f,
+        42.0f, 43.0f, 44.0f, 45.0f, 46.0f, 47.0f, 48.0f, 49.0f,
+        42.0f, 43.0f, 44.0f, 45.0f, 46.0f, 47.0f, 48.0f, 49.0f,
+        42.0f, 43.0f, 44.0f, 45.0f, 46.0f, 47.0f, 48.0f, 49.0f,
+        42.0f, 43.0f, 44.0f, 45.0f, 46.0f, 47.0f, 48.0f, 49.0f,
+        42.0f, 43.0f, 44.0f, 45.0f, 46.0f, 47.0f, 48.0f, 49.0f,
+        42.0f, 43.0f, 44.0f, 45.0f, 46.0f, 47.0f, 48.0f, 49.0f,
+        42.0f, 43.0f, 44.0f, 45.0f, 46.0f, 47.0f, 48.0f, 49.0f,
+        42.0f, 43.0f, 44.0f, 45.0f, 46.0f, 47.0f, 48.0f, 49.0f,
+        42.0f, 43.0f, 44.0f, 45.0f, 46.0f, 47.0f, 48.0f, 49.0f,
+        42.0f, 43.0f, 44.0f, 45.0f, 46.0f, 47.0f, 48.0f, 49.0f,
+        42.0f, 43.0f, 44.0f, 45.0f, 46.0f, 47.0f, 48.0f, 49.0f,
+        50.0f, 51.0f, 52.0f, 53.0f, 54.0f, 55.0f, 0.0f, 0.0f,
+        50.0f, 51.0f, 52.0f, 53.0f, 54.0f, 55.0f, 0.0f, 0.0f,
+        50.0f, 51.0f, 52.0f, 53.0f, 54.0f, 55.0f, 0.0f, 0.0f,
+        50.0f, 51.0f, 52.0f, 53.0f, 54.0f, 55.0f, 0.0f, 0.0f,
+        50.0f, 51.0f, 52.0f, 53.0f, 54.0f, 55.0f, 0.0f, 0.0f,
+        50.0f, 51.0f, 52.0f, 53.0f, 54.0f, 55.0f, 0.0f, 0.0f,
+        50.0f, 51.0f, 52.0f, 53.0f, 54.0f, 55.0f, 0.0f, 0.0f,
+        50.0f, 51.0f, 52.0f, 53.0f, 54.0f, 55.0f, 0.0f, 0.0f,
+        50.0f, 51.0f, 52.0f, 53.0f, 54.0f, 55.0f, 0.0f, 0.0f,
+        50.0f, 51.0f, 52.0f, 53.0f, 54.0f, 55.0f, 0.0f, 0.0f,
+        50.0f, 51.0f, 52.0f, 53.0f, 54.0f, 55.0f, 0.0f, 0.0f,
+        50.0f, 51.0f, 52.0f, 53.0f, 54.0f, 55.0f, 0.0f, 0.0f,
+        50.0f, 51.0f, 52.0f, 53.0f, 54.0f, 55.0f, 0.0f, 0.0f,
+    },
+    {  // Joint group 7
+        56.0f, 57.0f, 58.0f, 59.0f, 60.0f, 61.0f, 62.0f, 63.0f,
+        56.0f, 57.0f, 58.0f, 59.0f, 60.0f, 61.0f, 62.0f, 63.0f,
+        56.0f, 57.0f, 58.0f, 59.0f, 60.0f, 61.0f, 62.0f, 63.0f,
+        56.0f, 57.0f, 58.0f, 59.0f, 60.0f, 61.0f, 62.0f, 63.0f,
+        56.0f, 57.0f, 58.0f, 59.0f, 60.0f, 61.0f, 62.0f, 63.0f,
+        56.0f, 57.0f, 58.0f, 59.0f, 60.0f, 61.0f, 62.0f, 63.0f,
+        56.0f, 57.0f, 58.0f, 59.0f, 60.0f, 61.0f, 62.0f, 63.0f,
+        56.0f, 57.0f, 58.0f, 59.0f, 60.0f, 61.0f, 62.0f, 63.0f,
+        56.0f, 57.0f, 58.0f, 59.0f, 60.0f, 61.0f, 62.0f, 63.0f,
+        56.0f, 57.0f, 58.0f, 59.0f, 60.0f, 61.0f, 62.0f, 63.0f,
+        56.0f, 57.0f, 58.0f, 59.0f, 60.0f, 61.0f, 62.0f, 63.0f,
+        56.0f, 57.0f, 58.0f, 59.0f, 60.0f, 61.0f, 62.0f, 63.0f,
+        56.0f, 57.0f, 58.0f, 59.0f, 60.0f, 61.0f, 62.0f, 63.0f,
+        64.0f, 65.0f, 66.0f, 67.0f, 68.0f, 69.0f, 70.0f, 71.0f,
+        64.0f, 65.0f, 66.0f, 67.0f, 68.0f, 69.0f, 70.0f, 71.0f,
+        64.0f, 65.0f, 66.0f, 67.0f, 68.0f, 69.0f, 70.0f, 71.0f,
+        64.0f, 65.0f, 66.0f, 67.0f, 68.0f, 69.0f, 70.0f, 71.0f,
+        64.0f, 65.0f, 66.0f, 67.0f, 68.0f, 69.0f, 70.0f, 71.0f,
+        64.0f, 65.0f, 66.0f, 67.0f, 68.0f, 69.0f, 70.0f, 71.0f,
+        64.0f, 65.0f, 66.0f, 67.0f, 68.0f, 69.0f, 70.0f, 71.0f,
+        64.0f, 65.0f, 66.0f, 67.0f, 68.0f, 69.0f, 70.0f, 71.0f,
+        64.0f, 65.0f, 66.0f, 67.0f, 68.0f, 69.0f, 70.0f, 71.0f,
+        64.0f, 65.0f, 66.0f, 67.0f, 68.0f, 69.0f, 70.0f, 71.0f,
+        64.0f, 65.0f, 66.0f, 67.0f, 68.0f, 69.0f, 70.0f, 71.0f,
+        64.0f, 65.0f, 66.0f, 67.0f, 68.0f, 69.0f, 70.0f, 71.0f,
+        64.0f, 65.0f, 66.0f, 67.0f, 68.0f, 69.0f, 70.0f, 71.0f,
+    },
+    {  // Joint group 8
+        1.0f, 2.0f, 0.0f, 0.0f,
+        1.0f, 2.0f, 0.0f, 0.0f,
+        1.0f, 2.0f, 0.0f, 0.0f,
+        1.0f, 2.0f, 0.0f, 0.0f,
+        1.0f, 2.0f, 0.0f, 0.0f,
+    },
+    {  // Joint group 9
+        1.0f, 2.0f, 0.0f, 0.0f,
+        1.0f, 2.0f, 0.0f, 0.0f,
+        1.0f, 2.0f, 0.0f, 0.0f,
+        1.0f, 2.0f, 0.0f, 0.0f,
+        1.0f, 2.0f, 0.0f, 0.0f
+    }
+};
+
+const AlignedMatrix<std::uint16_t> halfFloatValues = {
+    {  // Joint group 0
+        15360, 0, 0, 0,
+        15360, 0, 0, 0,
+        15360, 0, 0, 0,
+        15360, 0, 0, 0,
+        15360, 0, 0, 0,
+        15360, 0, 0, 0,
+        15360, 0, 0, 0,
+        15360, 0, 0, 0,
+        15360, 0, 0, 0,
+        15360, 0, 0, 0,
+        15360, 0, 0, 0,
+        15360, 0, 0, 0,
+        15360, 0, 0, 0,
+    },
+    {  // Joint group 1
+        16384, 16896, 17408, 17664,
+        16384, 16896, 17408, 17664,
+        16384, 16896, 17408, 17664,
+        16384, 16896, 17408, 17664,
+        16384, 16896, 17408, 17664,
+        16384, 16896, 17408, 17664,
+        16384, 16896, 17408, 17664,
+        16384, 16896, 17408, 17664,
+        16384, 16896, 17408, 17664,
+        16384, 16896, 17408, 17664,
+        16384, 16896, 17408, 17664,
+        16384, 16896, 17408, 17664,
+        16384, 16896, 17408, 17664,
+    },
+    {  // Joint group 2
+        17920, 18176, 18432, 18560, 18688, 18816, 18944, 0,
+        17920, 18176, 18432, 18560, 18688, 18816, 18944, 0,
+        17920, 18176, 18432, 18560, 18688, 18816, 18944, 0,
+        17920, 18176, 18432, 18560, 18688, 18816, 18944, 0,
+        17920, 18176, 18432, 18560, 18688, 18816, 18944, 0,
+        17920, 18176, 18432, 18560, 18688, 18816, 18944, 0,
+        17920, 18176, 18432, 18560, 18688, 18816, 18944, 0,
+        17920, 18176, 18432, 18560, 18688, 18816, 18944, 0,
+        17920, 18176, 18432, 18560, 18688, 18816, 18944, 0,
+        17920, 18176, 18432, 18560, 18688, 18816, 18944, 0,
+        17920, 18176, 18432, 18560, 18688, 18816, 18944, 0,
+        17920, 18176, 18432, 18560, 18688, 18816, 18944, 0,
+        17920, 18176, 18432, 18560, 18688, 18816, 18944, 0,
+    },
+    {  // Joint group 3
+        19072, 19200, 19328, 19456, 19520, 19584, 19648, 19712,
+        19072, 19200, 19328, 19456, 19520, 19584, 19648, 19712,
+        19072, 19200, 19328, 19456, 19520, 19584, 19648, 19712,
+        19072, 19200, 19328, 19456, 19520, 19584, 19648, 19712,
+        19072, 19200, 19328, 19456, 19520, 19584, 19648, 19712,
+        19072, 19200, 19328, 19456, 19520, 19584, 19648, 19712,
+        19072, 19200, 19328, 19456, 19520, 19584, 19648, 19712,
+        19072, 19200, 19328, 19456, 19520, 19584, 19648, 19712,
+        19072, 19200, 19328, 19456, 19520, 19584, 19648, 19712,
+        19072, 19200, 19328, 19456, 19520, 19584, 19648, 19712,
+        19072, 19200, 19328, 19456, 19520, 19584, 19648, 19712,
+        19072, 19200, 19328, 19456, 19520, 19584, 19648, 19712,
+        19072, 19200, 19328, 19456, 19520, 19584, 19648, 19712,
+    },
+    {  // Joint group 4
+        19776, 19840, 19904, 19968, 20032, 20096, 20160, 20224,
+        19776, 19840, 19904, 19968, 20032, 20096, 20160, 20224,
+        19776, 19840, 19904, 19968, 20032, 20096, 20160, 20224,
+        19776, 19840, 19904, 19968, 20032, 20096, 20160, 20224,
+        19776, 19840, 19904, 19968, 20032, 20096, 20160, 20224,
+        19776, 19840, 19904, 19968, 20032, 20096, 20160, 20224,
+        19776, 19840, 19904, 19968, 20032, 20096, 20160, 20224,
+        19776, 19840, 19904, 19968, 20032, 20096, 20160, 20224,
+        19776, 19840, 19904, 19968, 20032, 20096, 20160, 20224,
+        19776, 19840, 19904, 19968, 20032, 20096, 20160, 20224,
+        19776, 19840, 19904, 19968, 20032, 20096, 20160, 20224,
+        19776, 19840, 19904, 19968, 20032, 20096, 20160, 20224,
+        19776, 19840, 19904, 19968, 20032, 20096, 20160, 20224,
+        20288, 0, 0, 0,
+        20288, 0, 0, 0,
+        20288, 0, 0, 0,
+        20288, 0, 0, 0,
+        20288, 0, 0, 0,
+        20288, 0, 0, 0,
+        20288, 0, 0, 0,
+        20288, 0, 0, 0,
+        20288, 0, 0, 0,
+        20288, 0, 0, 0,
+        20288, 0, 0, 0,
+        20288, 0, 0, 0,
+        20288, 0, 0, 0,
+    },
+    {  // Joint group 5
+        20352, 20416, 20480, 20512, 20544, 20576, 20608, 20640,
+        20352, 20416, 20480, 20512, 20544, 20576, 20608, 20640,
+        20352, 20416, 20480, 20512, 20544, 20576, 20608, 20640,
+        20352, 20416, 20480, 20512, 20544, 20576, 20608, 20640,
+        20352, 20416, 20480, 20512, 20544, 20576, 20608, 20640,
+        20352, 20416, 20480, 20512, 20544, 20576, 20608, 20640,
+        20352, 20416, 20480, 20512, 20544, 20576, 20608, 20640,
+        20352, 20416, 20480, 20512, 20544, 20576, 20608, 20640,
+        20352, 20416, 20480, 20512, 20544, 20576, 20608, 20640,
+        20352, 20416, 20480, 20512, 20544, 20576, 20608, 20640,
+        20352, 20416, 20480, 20512, 20544, 20576, 20608, 20640,
+        20352, 20416, 20480, 20512, 20544, 20576, 20608, 20640,
+        20352, 20416, 20480, 20512, 20544, 20576, 20608, 20640,
+        20672, 20704, 20736, 20768,
+        20672, 20704, 20736, 20768,
+        20672, 20704, 20736, 20768,
+        20672, 20704, 20736, 20768,
+        20672, 20704, 20736, 20768,
+        20672, 20704, 20736, 20768,
+        20672, 20704, 20736, 20768,
+        20672, 20704, 20736, 20768,
+        20672, 20704, 20736, 20768,
+        20672, 20704, 20736, 20768,
+        20672, 20704, 20736, 20768,
+        20672, 20704, 20736, 20768,
+        20672, 20704, 20736, 20768,
+    },
+    {  // Joint group 6
+        20800, 20832, 20864, 20896, 20928, 20960, 20992, 21024,
+        20800, 20832, 20864, 20896, 20928, 20960, 20992, 21024,
+        20800, 20832, 20864, 20896, 20928, 20960, 20992, 21024,
+        20800, 20832, 20864, 20896, 20928, 20960, 20992, 21024,
+        20800, 20832, 20864, 20896, 20928, 20960, 20992, 21024,
+        20800, 20832, 20864, 20896, 20928, 20960, 20992, 21024,
+        20800, 20832, 20864, 20896, 20928, 20960, 20992, 21024,
+        20800, 20832, 20864, 20896, 20928, 20960, 20992, 21024,
+        20800, 20832, 20864, 20896, 20928, 20960, 20992, 21024,
+        20800, 20832, 20864, 20896, 20928, 20960, 20992, 21024,
+        20800, 20832, 20864, 20896, 20928, 20960, 20992, 21024,
+        20800, 20832, 20864, 20896, 20928, 20960, 20992, 21024,
+        20800, 20832, 20864, 20896, 20928, 20960, 20992, 21024,
+        21056, 21088, 21120, 21152, 21184, 21216, 0, 0,
+        21056, 21088, 21120, 21152, 21184, 21216, 0, 0,
+        21056, 21088, 21120, 21152, 21184, 21216, 0, 0,
+        21056, 21088, 21120, 21152, 21184, 21216, 0, 0,
+        21056, 21088, 21120, 21152, 21184, 21216, 0, 0,
+        21056, 21088, 21120, 21152, 21184, 21216, 0, 0,
+        21056, 21088, 21120, 21152, 21184, 21216, 0, 0,
+        21056, 21088, 21120, 21152, 21184, 21216, 0, 0,
+        21056, 21088, 21120, 21152, 21184, 21216, 0, 0,
+        21056, 21088, 21120, 21152, 21184, 21216, 0, 0,
+        21056, 21088, 21120, 21152, 21184, 21216, 0, 0,
+        21056, 21088, 21120, 21152, 21184, 21216, 0, 0,
+        21056, 21088, 21120, 21152, 21184, 21216, 0, 0,
+    },
+    {  // Joint group 7
+        21248, 21280, 21312, 21344, 21376, 21408, 21440, 21472,
+        21248, 21280, 21312, 21344, 21376, 21408, 21440, 21472,
+        21248, 21280, 21312, 21344, 21376, 21408, 21440, 21472,
+        21248, 21280, 21312, 21344, 21376, 21408, 21440, 21472,
+        21248, 21280, 21312, 21344, 21376, 21408, 21440, 21472,
+        21248, 21280, 21312, 21344, 21376, 21408, 21440, 21472,
+        21248, 21280, 21312, 21344, 21376, 21408, 21440, 21472,
+        21248, 21280, 21312, 21344, 21376, 21408, 21440, 21472,
+        21248, 21280, 21312, 21344, 21376, 21408, 21440, 21472,
+        21248, 21280, 21312, 21344, 21376, 21408, 21440, 21472,
+        21248, 21280, 21312, 21344, 21376, 21408, 21440, 21472,
+        21248, 21280, 21312, 21344, 21376, 21408, 21440, 21472,
+        21248, 21280, 21312, 21344, 21376, 21408, 21440, 21472,
+        21504, 21520, 21536, 21552, 21568, 21584, 21600, 21616,
+        21504, 21520, 21536, 21552, 21568, 21584, 21600, 21616,
+        21504, 21520, 21536, 21552, 21568, 21584, 21600, 21616,
+        21504, 21520, 21536, 21552, 21568, 21584, 21600, 21616,
+        21504, 21520, 21536, 21552, 21568, 21584, 21600, 21616,
+        21504, 21520, 21536, 21552, 21568, 21584, 21600, 21616,
+        21504, 21520, 21536, 21552, 21568, 21584, 21600, 21616,
+        21504, 21520, 21536, 21552, 21568, 21584, 21600, 21616,
+        21504, 21520, 21536, 21552, 21568, 21584, 21600, 21616,
+        21504, 21520, 21536, 21552, 21568, 21584, 21600, 21616,
+        21504, 21520, 21536, 21552, 21568, 21584, 21600, 21616,
+        21504, 21520, 21536, 21552, 21568, 21584, 21600, 21616,
+        21504, 21520, 21536, 21552, 21568, 21584, 21600, 21616,
+    },
+    {  // Joint group 8
+        15360, 16384, 0, 0,
+        15360, 16384, 0, 0,
+        15360, 16384, 0, 0,
+        15360, 16384, 0, 0,
+        15360, 16384, 0, 0,
+    },
+    {  // Joint group 9
+        15360, 16384, 0, 0,
+        15360, 16384, 0, 0,
+        15360, 16384, 0, 0,
+        15360, 16384, 0, 0,
+        15360, 16384, 0, 0,
+    }
+};
+
+const Vector<bpcm::Extent> subMatrices = {
+    {4, 13},  // Joint group 0
+    {4, 13},  // Joint group 1
+    {8, 13},  // Joint group 2
+    {8, 13},  // Joint group 3
+    {12, 13},  // Joint group 4
+    {12, 13},  // Joint group 5
+    {16, 13},  // Joint group 6
+    {16, 13},  // Joint group 7
+    {4, 5},  // Joint group 8
+    {4, 5}  // Joint group 9
+};
+
+const AlignedMatrix<std::uint16_t> inputIndices = {
+    {  // Joint group 0
+        12, 11, 10, 9, 8, 7, 6, 5, 4, 3, 2, 1, 0
+    },
+    {  // Joint group 1
+        12, 11, 10, 9, 8, 7, 6, 5, 4, 3, 2, 1, 0
+    },
+    {  // Joint group 2
+        12, 11, 10, 9, 8, 7, 6, 5, 4, 3, 2, 1, 0
+    },
+    {  // Joint group 3
+        12, 11, 10, 9, 8, 7, 6, 5, 4, 3, 2, 1, 0
+    },
+    {  // Joint group 4
+        12, 11, 10, 9, 8, 7, 6, 5, 4, 3, 2, 1, 0
+    },
+    {  // Joint group 5
+        12, 11, 10, 9, 8, 7, 6, 5, 4, 3, 2, 1, 0
+    },
+    {  // Joint group 6
+        12, 11, 10, 9, 8, 7, 6, 5, 4, 3, 2, 1, 0
+    },
+    {  // Joint group 7
+        12, 11, 10, 9, 8, 7, 6, 5, 4, 3, 2, 1, 0
+    },
+    {  // Joint group 8
+        1, 2, 3, 4, 5
+    },
+    {  // Joint group 9
+        7, 8, 9, 10, 11
+    }
+};
+
+const AlignedMatrix<std::uint16_t> outputIndices = {
+    {  // Joint group 0
+        0, 0, 0, 0
+    },
+    {  // Joint group 1
+        1, 2, 3, 4
+    },
+    {  // Joint group 2
+        5, 6, 7, 8, 9, 10, 11, 0
+    },
+    {  // Joint group 3
+        12, 13, 14, 15, 16, 17, 18, 19
+    },
+    {  // Joint group 4
+        20, 21, 22, 23, 24, 25, 26, 27, 28, 0, 0, 0
+    },
+    {  // Joint group 5
+        29, 30, 31, 32, 33, 34, 35, 36, 37, 38, 39, 40
+    },
+    {  // Joint group 6
+        41, 42, 43, 44, 45, 46, 47, 48, 49, 50, 51, 52, 53, 54, 0, 0
+    },
+    {  // Joint group 7
+        55, 56, 57, 58, 59, 60, 61, 62, 63, 64, 65, 66, 67, 68, 69, 70
+    },
+    {  // Joint group 8
+        71, 72, 0, 0
+    },
+    {  // Joint group 9
+        73, 74, 0, 0
+    }
+};
+
+const Vector<bpcm::JointGroup> jointGroups = {
+    // {valueOffset, inputOffset, outputOffset, lodOffset, valueSize, inputSize, inputSizeAligned4, inputSizeAligned8}
+    {  // Joint group 0
+        0, 0, 0, 0, 52, 13, 12, 8
+    },
+    {  // Joint group 1
+        52, 13, 4, 4, 52, 13, 12, 8
+    },
+    {  // Joint group 2
+        104, 26, 8, 8, 104, 13, 12, 8
+    },
+    {  // Joint group 3
+        208, 39, 16, 12, 104, 13, 12, 8
+    },
+    {  // Joint group 4
+        312, 52, 24, 16, 156, 13, 12, 8
+    },
+    {  // Joint group 5
+        468, 65, 36, 20, 156, 13, 12, 8
+    },
+    {  // Joint group 6
+        624, 78, 48, 24, 208, 13, 12, 8
+    },
+    {  // Joint group 7
+        832, 91, 64, 28, 208, 13, 12, 8
+    },
+    {  // Joint group 8
+        1040, 104, 80, 32, 20, 5, 4, 0
+    },
+    {  // Joint group 9
+        1060, 109, 84, 36, 20, 5, 4, 0
+    }
+};
+
+const Matrix<bpcm::LODRegion> lodRegions = {
+    // {unaligned, aligned to last block-8, aligned to second last block-8}
+    {  // Joint group 0
+        {1, 0, 0},  // LOD-0
+        {1, 0, 0},  // LOD-1
+        {1, 0, 0},  // LOD-2
+        {1, 0, 0},  // LOD-3
+    },
+    {  // Joint group 1
+        {4, 0, 0},  // LOD-0
+        {3, 0, 0},  // LOD-1
+        {2, 0, 0},  // LOD-2
+        {1, 0, 0},  // LOD-3
+    },
+    {  // Joint group 2
+        {7, 8, 0},  // LOD-0
+        {5, 8, 0},  // LOD-1
+        {4, 8, 0},  // LOD-2
+        {1, 8, 0},  // LOD-3
+    },
+    {  // Joint group 3
+        {8, 8, 8},  // LOD-0
+        {6, 8, 0},  // LOD-1
+        {4, 8, 0},  // LOD-2
+        {3, 8, 0},  // LOD-3
+    },
+    {  // Joint group 4
+        {9, 8, 8},  // LOD-0
+        {7, 8, 0},  // LOD-1
+        {4, 8, 0},  // LOD-2
+        {2, 8, 0},  // LOD-3
+    },
+    {  // Joint group 5
+        {12, 8, 8},  // LOD-0
+        {9, 8, 8},  // LOD-1
+        {7, 8, 0},  // LOD-2
+        {3, 8, 0},  // LOD-3
+    },
+    {  // Joint group 6
+        {14, 16, 8},  // LOD-0
+        {13, 16, 8},  // LOD-1
+        {11, 16, 8},  // LOD-2
+        {9, 16, 8},  // LOD-3
+    },
+    {  // Joint group 7
+        {16, 16, 16},  // LOD-0
+        {12, 16, 8},  // LOD-1
+        {8, 8, 8},  // LOD-2
+        {4, 8, 0},  // LOD-3
+    },
+    {  // Joint group 8
+        {2, 0, 0},  // LOD-0
+        {2, 0, 0},  // LOD-1
+        {2, 0, 0},  // LOD-2
+        {2, 0, 0},  // LOD-3
+    },
+    {  // Joint group 9
+        {2, 0, 0},  // LOD-0
+        {2, 0, 0},  // LOD-1
+        {2, 0, 0},  // LOD-2
+        {2, 0, 0}  // LOD-3
+    }
+};
+
+}  // namespace optimized
+
+namespace input {
+
+// Calculation input values
+const rl4::Vector<float> values = {1.0f, 2.0f, 3.0f, 4.0f, 0.0f, 6.0f, 7.0f, 8.0f, 9.0f, 0.0f, 11.0f, 12.0f, 13.0f};
+
+}  // namespace input
+
+namespace output {
+
+// Expected output results for each LOD
+const rl4::Matrix<float> valuesPerLOD = {
+    {
+        // LOD-0
+        76.0f,  // Joint group 0
+        152.0f, 228.0f, 304.0f, 380.0f,  // Joint group 1
+        456.0f, 532.0f, 608.0f, 684.0f, 760.0f, 836.0f, 912.0f,  // Joint group 2
+        988.0f, 1064.0f, 1140.0f, 1216.0f, 1292.0f, 1368.0f, 1444.0f, 1520.0f,  // Joint group 3
+        1596.0f, 1672.0f, 1748.0f, 1824.0f, 1900.0f, 1976.0f, 2052.0f, 2128.0f, 2204.0f,  // Joint group 4
+        2280.0f, 2356.0f, 2432.0f, 2508.0f, 2584.0f, 2660.0f,  // Joint group 5
+        2736.0f, 2812.0f, 2888.0f, 2964.0f, 3040.0f, 3116.0f,  // Joint group 5
+        3192.0f, 3268.0f, 3344.0f, 3420.0f, 3496.0f, 3572.0f, 3648.0f,  // Joint group 6
+        3724.0f, 3800.0f, 3876.0f, 3952.0f, 4028.0f, 4104.0f, 4180.0f,  // Joint group 6
+        4256.0f, 4332.0f, 4408.0f, 4484.0f, 4560.0f, 4636.0f, 4712.0f, 4788.0f,  // Joint group 7
+        4864.0f, 4940.0f, 5016.0f, 5092.0f, 5168.0f, 5244.0f, 5320.0f, 5396.0f,  // Joint group 7
+        15.0f, 30.0f,  // Joint group 8
+        40.0f, 80.0f  // Joint group 9
+    }, {
+        // LOD-1
+        76.0f,  // Joint group 0
+        152.0f, 228.0f, 304.0f, 0.0f,  // Joint group 1
+        456.0f, 532.0f, 608.0f, 684.0f, 760.0f, 0.0f, 0.0f,  // Joint group 2
+        988.0f, 1064.0f, 1140.0f, 1216.0f, 1292.0f, 1368.0f, 0.0f, 0.0f,  // Joint group 3
+        1596.0f, 1672.0f, 1748.0f, 1824.0f, 1900.0f, 1976.0f, 2052.0f, 0.0f, 0.0f,  // Joint group 4
+        2280.0f, 2356.0f, 2432.0f, 2508.0f, 2584.0f, 2660.0f,  // Joint group 5
+        2736.0f, 2812.0f, 2888.0f, 0.0f, 0.0f, 0.0f,  // Joint group 5
+        3192.0f, 3268.0f, 3344.0f, 3420.0f, 3496.0f, 3572.0f, 3648.0f,  // Joint group 6
+        3724.0f, 3800.0f, 3876.0f, 3952.0f, 4028.0f, 4104.0f, 0.0f,  // Joint group 6
+        4256.0f, 4332.0f, 4408.0f, 4484.0f, 4560.0f, 4636.0f, 4712.0f, 4788.0f,  // Joint group 7
+        4864.0f, 4940.0f, 5016.0f, 5092.0f, 0.0f, 0.0f, 0.0f, 0.0f,  // Joint group 7
+        15.0f, 30.0f,  // Joint group 8
+        40.0f, 80.0f  // Joint group 9
+    }, {
+        // LOD-2
+        76.0f,  // Joint group 0
+        152.0f, 228.0f, 0.0f, 0.0f,  // Joint group 1
+        456.0f, 532.0f, 608.0f, 684.0f, 0.0f, 0.0f, 0.0f,  // Joint group 2
+        988.0f, 1064.0f, 1140.0f, 1216.0f, 0.0f, 0.0f, 0.0f, 0.0f,  // Joint group 3
+        1596.0f, 1672.0f, 1748.0f, 1824.0f, 0.0f, 0.0f, 0.0f, 0.0f, 0.0f,  // Joint group 4
+        2280.0f, 2356.0f, 2432.0f, 2508.0f, 2584.0f, 2660.0f,  // Joint group 5
+        2736.0f, 0.0f, 0.0f, 0.0f, 0.0f, 0.0f,  // Joint group 5
+        3192.0f, 3268.0f, 3344.0f, 3420.0f, 3496.0f, 3572.0f, 3648.0f,  // Joint group 6
+        3724.0f, 3800.0f, 3876.0f, 3952.0f, 0.0f, 0.0f, 0.0f,  // Joint group 6
+        4256.0f, 4332.0f, 4408.0f, 4484.0f, 4560.0f, 4636.0f, 4712.0f, 4788.0f,  // Joint group 7
+        0.0f, 0.0f, 0.0f, 0.0f, 0.0f, 0.0f, 0.0f, 0.0f,  // Joint group 7
+        15.0f, 30.0f,  // Joint group 8
+        40.0f, 80.0f  // Joint group 9
+    }, {
+        // LOD-3
+        76.0f,  // Joint group 0
+        152.0f, 0.0f, 0.0f, 0.0f,  // Joint group 1
+        456.0f, 0.0f, 0.0f, 0.0f, 0.0f, 0.0f, 0.0f,  // Joint group 2
+        988.0f, 1064.0f, 1140.0f, 0.0f, 0.0f, 0.0f, 0.0f, 0.0f,  // Joint group 3
+        1596.0f, 1672.0f, 0.0f, 0.0f, 0.0f, 0.0f, 0.0f, 0.0f, 0.0f,  // Joint group 4
+        2280.0f, 2356.0f, 2432.0f, 0.0f, 0.0f, 0.0f,  // Joint group 5
+        0.0f, 0.0f, 0.0f, 0.0f, 0.0f, 0.0f,  // Joint group 5
+        3192.0f, 3268.0f, 3344.0f, 3420.0f, 3496.0f, 3572.0f, 3648.0f,  // Joint group 6
+        3724.0f, 3800.0f, 0.0f, 0.0f, 0.0f, 0.0f, 0.0f,  // Joint group 6
+        4256.0f, 4332.0f, 4408.0f, 4484.0f, 0.0f, 0.0f, 0.0f, 0.0f,  // Joint group 7
+        0.0f, 0.0f, 0.0f, 0.0f, 0.0f, 0.0f, 0.0f, 0.0f,  // Joint group 7
+        15.0f, 30.0f,  // Joint group 8
+        40.0f, 80.0f  // Joint group 9
+    }
+};
+
+}  // namespace output
+
+template<typename TValue>
+struct OptimizedValues;
+
+template<>
+struct OptimizedValues<float> {
+    static const rl4::AlignedMatrix<float>& get() {
+        return optimized::floatValues;
+    }
+
+};
+
+template<>
+struct OptimizedValues<std::uint16_t> {
+    static const rl4::AlignedMatrix<std::uint16_t>& get() {
+        return optimized::halfFloatValues;
+    }
+
+};
+
+template<typename TValue>
+bpcm::Evaluator<TValue> OptimizedStorage<TValue>::create(StrategyPtr strategy, rl4::MemoryResource* memRes) {
+    bpcm::JointStorage<TValue> storage{memRes};
+    const auto& values = OptimizedValues<TValue>::get();
+    for (std::uint16_t i = 0u; i < static_cast<std::uint16_t>(unoptimized::values.size()); ++i) {
+        storage.values.insert(storage.values.end(), values[i].begin(), values[i].end());
+        storage.inputIndices.insert(storage.inputIndices.end(),
+                                    optimized::inputIndices[i].begin(),
+                                    optimized::inputIndices[i].end());
+        storage.outputIndices.insert(storage.outputIndices.end(),
+                                     optimized::outputIndices[i].begin(),
+                                     optimized::outputIndices[i].end());
+        storage.lodRegions.insert(storage.lodRegions.end(),
+                                  optimized::lodRegions[i].begin(),
+                                  optimized::lodRegions[i].end());
+        storage.jointGroups.push_back(optimized::jointGroups[i]);
+    }
+    return bpcm::Evaluator<TValue>{std::move(storage), std::move(strategy), memRes};
+}
+
+template<typename TValue>
+bpcm::Evaluator<TValue> OptimizedStorage<TValue>::create(StrategyPtr strategy,
+                                                         std::uint16_t jointGroupIndex,
+                                                         rl4::MemoryResource* memRes) {
+    const auto& values = OptimizedValues<TValue>::get();
+    bpcm::JointStorage<TValue> storage{memRes};
+    storage.values.assign(values[jointGroupIndex].begin(), values[jointGroupIndex].end());
+    storage.inputIndices.assign(optimized::inputIndices[jointGroupIndex].begin(),
+                                optimized::inputIndices[jointGroupIndex].end());
+    storage.outputIndices.assign(optimized::outputIndices[jointGroupIndex].begin(),
+                                 optimized::outputIndices[jointGroupIndex].end());
+    storage.lodRegions.assign(optimized::lodRegions[jointGroupIndex].begin(),
+                              optimized::lodRegions[jointGroupIndex].end());
+    storage.jointGroups.push_back(bpcm::JointGroup{
+            0u,  // values offset
+            0u,  // inputIndicesOffset
+            0u,  // outputIndicesOffset
+            0u,  // lodsOffset
+            static_cast<std::uint32_t>(storage.values.size()),
+            static_cast<std::uint32_t>(storage.inputIndices.size()),
+            static_cast<std::uint32_t>(storage.inputIndices.size() - (storage.inputIndices.size() % 4ul)),
+            static_cast<std::uint32_t>(storage.inputIndices.size() - (storage.inputIndices.size() % 8ul)),
+        });
+    return bpcm::Evaluator<TValue>{std::move(storage), std::move(strategy), memRes};
+}
+
+template struct OptimizedStorage<float>;
+#ifdef RL_USE_HALF_FLOATS
+    template struct OptimizedStorage<std::uint16_t>;
+#endif
+
+}  // namespace block4
diff --git a/Engine/Plugins/Animation/RigLogic/Source/RigLogicLibTest/Private/rltests/joints/bpcm/JointFixturesBlock4.h b/Engine/Plugins/Animation/RigLogic/Source/RigLogicLibTest/Private/rltests/joints/bpcm/JointFixturesBlock4.h
new file mode 100644
--- /dev/null
+++ b/Engine/Plugins/Animation/RigLogic/Source/RigLogicLibTest/Private/rltests/joints/bpcm/JointFixturesBlock4.h
@@ -0,0 +1,139 @@
+#pragma once
+
+#include "rltests/dna/FakeReader.h"
+#include "rltests/joints/bpcm/StorageValueType.h"
+
+#include "riglogic/TypeDefs.h"
+#include "riglogic/joints/bpcm/Extent.h"
+#include "riglogic/joints/bpcm/JointGroup.h"
+#include "riglogic/joints/bpcm/JointsEvaluator.h"
+#include "riglogic/joints/bpcm/LODRegion.h"
+#include "riglogic/riglogic/RigLogic.h"
+
+#include <cstdint>
+#include <functional>
+#include <memory>
+
+namespace block4 {
+
+using namespace rl4;
+
+namespace unoptimized {
+
+extern const std::uint16_t lodCount;
+extern const bpcm::Extent dimensions;
+extern const Matrix<float> values;
+extern const Matrix<std::uint16_t> inputIndices;
+extern const Matrix<std::uint16_t> outputIndices;
+extern const Matrix<std::uint16_t> lods;
+extern const Vector<bpcm::Extent> subMatrices;
+extern const Vector<float> neutralJointTranslationXs;
+extern const Vector<float> neutralJointTranslationYs;
+extern const Vector<float> neutralJointTranslationZs;
+extern const Vector<float> neutralJointRotationXs;
+extern const Vector<float> neutralJointRotationYs;
+extern const Vector<float> neutralJointRotationZs;
+
+}  // namespace unoptimized
+
+namespace optimized {
+
+extern const bpcm::Extent dimensions;
+extern const AlignedMatrix<float> floatValues;
+extern const AlignedMatrix<std::uint16_t> halfFloatValues;
+extern const AlignedMatrix<std::uint16_t> inputIndices;
+extern const AlignedMatrix<std::uint16_t> outputIndices;
+extern const Vector<bpcm::JointGroup> jointGroups;
+extern const Matrix<bpcm::LODRegion> lodRegions;
+
+}  // namespace optimized
+
+namespace input {
+
+// Calculation input values
+extern const Vector<float> values;
+
+}  // namespace input
+
+namespace output {
+
+// Calculation output values
+extern const Matrix<float> valuesPerLOD;
+
+}  // namespace output
+
+class CanonicalReader : public dna::FakeReader {
+    public:
+        std::uint16_t getLODCount() const override {
+            return unoptimized::lodCount;
+        }
+
+        std::uint16_t getJointCount() const override {
+            return static_cast<std::uint16_t>(unoptimized::neutralJointTranslationXs.size());
+        }
+
+        ConstArrayView<float> getNeutralJointTranslationXs() const override {
+            return ConstArrayView<float>{unoptimized::neutralJointTranslationXs};
+        }
+
+        ConstArrayView<float> getNeutralJointTranslationYs() const override {
+            return ConstArrayView<float>{unoptimized::neutralJointTranslationYs};
+        }
+
+        ConstArrayView<float> getNeutralJointTranslationZs() const override {
+            return ConstArrayView<float>{unoptimized::neutralJointTranslationZs};
+        }
+
+        ConstArrayView<float> getNeutralJointRotationXs() const override {
+            return ConstArrayView<float>{unoptimized::neutralJointRotationXs};
+        }
+
+        ConstArrayView<float> getNeutralJointRotationYs() const override {
+            return ConstArrayView<float>{unoptimized::neutralJointRotationYs};
+        }
+
+        ConstArrayView<float> getNeutralJointRotationZs() const override {
+            return ConstArrayView<float>{unoptimized::neutralJointRotationZs};
+        }
+
+        std::uint16_t getJointRowCount() const override {
+            return unoptimized::dimensions.rows;
+        }
+
+        std::uint16_t getJointColumnCount() const override {
+            return unoptimized::dimensions.cols;
+        }
+
+        std::uint16_t getJointGroupCount() const override {
+            return static_cast<std::uint16_t>(unoptimized::subMatrices.size());
+        }
+
+        ConstArrayView<std::uint16_t> getJointGroupLODs(std::uint16_t jointGroupIndex) const override {
+            return ConstArrayView<std::uint16_t>{unoptimized::lods[jointGroupIndex]};
+        }
+
+        ConstArrayView<std::uint16_t> getJointGroupInputIndices(std::uint16_t jointGroupIndex) const override {
+            return ConstArrayView<std::uint16_t>{unoptimized::inputIndices[jointGroupIndex]};
+        }
+
+        ConstArrayView<std::uint16_t> getJointGroupOutputIndices(std::uint16_t jointGroupIndex) const override {
+            return ConstArrayView<std::uint16_t>{unoptimized::outputIndices[jointGroupIndex]};
+        }
+
+        ConstArrayView<float> getJointGroupValues(std::uint16_t jointGroupIndex) const override {
+            return ConstArrayView<float>{unoptimized::values[jointGroupIndex]};
+        }
+
+};
+
+template<typename TValue>
+struct OptimizedStorage {
+    using StrategyPtr = std::unique_ptr<bpcm::JointCalculationStrategy<TValue>,
+                                        std::function<void (bpcm::JointCalculationStrategy<TValue>*)> >;
+
+    static bpcm::Evaluator<TValue> create(StrategyPtr strategy, rl4::MemoryResource* memRes);
+    static bpcm::Evaluator<TValue> create(StrategyPtr strategy, std::uint16_t jointGroupIndex, rl4::MemoryResource* memRes);
+
+};
+
+}  // namespace block4
