Editgrating RigLogic plugin from Dev-CharacterTech to UE5/Main @ CL 13891876.

diff --git a/Engine/Plugins/Animation/RigLogic/Source/RigLogicLibTest/Private/rltests/joints/bpcm/JointFixturesBlock8.cpp b/Engine/Plugins/Animation/RigLogic/Source/RigLogicLibTest/Private/rltests/joints/bpcm/JointFixturesBlock8.cpp
new file mode 100644
--- /dev/null
+++ b/Engine/Plugins/Animation/RigLogic/Source/RigLogicLibTest/Private/rltests/joints/bpcm/JointFixturesBlock8.cpp
@@ -0,0 +1,931 @@
+#include "rltests/joints/bpcm/JointFixturesBlock8.h"
+
+#include "rltests/Defs.h"
+
+#include "riglogic/TypeDefs.h"
+#include "riglogic/joints/bpcm/JointsEvaluator.h"
+
+namespace block8 {
+
+namespace unoptimized {
+
+const std::uint16_t lodCount = 4u;
+const bpcm::Extent dimensions = {136ul, 13ul};
+
+const Matrix<float> values = {
+    {  // Joint group 0
+        1.0f, 1.0f, 1.0f, 1.0f, 1.0f, 1.0f, 1.0f, 1.0f, 1.0f, 1.0f, 1.0f, 1.0f, 1.0f,
+        2.0f, 2.0f, 2.0f, 2.0f, 2.0f, 2.0f, 2.0f, 2.0f, 2.0f, 2.0f, 2.0f, 2.0f, 2.0f,
+        3.0f, 3.0f, 3.0f, 3.0f, 3.0f, 3.0f, 3.0f, 3.0f, 3.0f, 3.0f, 3.0f, 3.0f, 3.0f,
+        4.0f, 4.0f, 4.0f, 4.0f, 4.0f, 4.0f, 4.0f, 4.0f, 4.0f, 4.0f, 4.0f, 4.0f, 4.0f,
+        5.0f, 5.0f, 5.0f, 5.0f, 5.0f, 5.0f, 5.0f, 5.0f, 5.0f, 5.0f, 5.0f, 5.0f, 5.0f
+    },
+    {  // Joint group 1
+        6.0f, 6.0f, 6.0f, 6.0f, 6.0f, 6.0f, 6.0f, 6.0f, 6.0f, 6.0f, 6.0f, 6.0f, 6.0f,
+        7.0f, 7.0f, 7.0f, 7.0f, 7.0f, 7.0f, 7.0f, 7.0f, 7.0f, 7.0f, 7.0f, 7.0f, 7.0f,
+        8.0f, 8.0f, 8.0f, 8.0f, 8.0f, 8.0f, 8.0f, 8.0f, 8.0f, 8.0f, 8.0f, 8.0f, 8.0f,
+        9.0f, 9.0f, 9.0f, 9.0f, 9.0f, 9.0f, 9.0f, 9.0f, 9.0f, 9.0f, 9.0f, 9.0f, 9.0f,
+        10.0f, 10.0f, 10.0f, 10.0f, 10.0f, 10.0f, 10.0f, 10.0f, 10.0f, 10.0f, 10.0f, 10.0f, 10.0f,
+        11.0f, 11.0f, 11.0f, 11.0f, 11.0f, 11.0f, 11.0f, 11.0f, 11.0f, 11.0f, 11.0f, 11.0f, 11.0f,
+        12.0f, 12.0f, 12.0f, 12.0f, 12.0f, 12.0f, 12.0f, 12.0f, 12.0f, 12.0f, 12.0f, 12.0f, 12.0f,
+        13.0f, 13.0f, 13.0f, 13.0f, 13.0f, 13.0f, 13.0f, 13.0f, 13.0f, 13.0f, 13.0f, 13.0f, 13.0f
+    },
+    {  // Joint group 2
+        14.0f, 14.0f, 14.0f, 14.0f, 14.0f, 14.0f, 14.0f, 14.0f, 14.0f, 14.0f, 14.0f, 14.0f, 14.0f,
+        15.0f, 15.0f, 15.0f, 15.0f, 15.0f, 15.0f, 15.0f, 15.0f, 15.0f, 15.0f, 15.0f, 15.0f, 15.0f,
+        16.0f, 16.0f, 16.0f, 16.0f, 16.0f, 16.0f, 16.0f, 16.0f, 16.0f, 16.0f, 16.0f, 16.0f, 16.0f,
+        17.0f, 17.0f, 17.0f, 17.0f, 17.0f, 17.0f, 17.0f, 17.0f, 17.0f, 17.0f, 17.0f, 17.0f, 17.0f,
+        18.0f, 18.0f, 18.0f, 18.0f, 18.0f, 18.0f, 18.0f, 18.0f, 18.0f, 18.0f, 18.0f, 18.0f, 18.0f,
+        19.0f, 19.0f, 19.0f, 19.0f, 19.0f, 19.0f, 19.0f, 19.0f, 19.0f, 19.0f, 19.0f, 19.0f, 19.0f,
+        20.0f, 20.0f, 20.0f, 20.0f, 20.0f, 20.0f, 20.0f, 20.0f, 20.0f, 20.0f, 20.0f, 20.0f, 20.0f,
+        21.0f, 21.0f, 21.0f, 21.0f, 21.0f, 21.0f, 21.0f, 21.0f, 21.0f, 21.0f, 21.0f, 21.0f, 21.0f,
+        22.0f, 22.0f, 22.0f, 22.0f, 22.0f, 22.0f, 22.0f, 22.0f, 22.0f, 22.0f, 22.0f, 22.0f, 22.0f
+    },
+    {  // Joint group 3
+        23.0f, 23.0f, 23.0f, 23.0f, 23.0f, 23.0f, 23.0f, 23.0f, 23.0f, 23.0f, 23.0f, 23.0f, 23.0f,
+        24.0f, 24.0f, 24.0f, 24.0f, 24.0f, 24.0f, 24.0f, 24.0f, 24.0f, 24.0f, 24.0f, 24.0f, 24.0f,
+        25.0f, 25.0f, 25.0f, 25.0f, 25.0f, 25.0f, 25.0f, 25.0f, 25.0f, 25.0f, 25.0f, 25.0f, 25.0f,
+        26.0f, 26.0f, 26.0f, 26.0f, 26.0f, 26.0f, 26.0f, 26.0f, 26.0f, 26.0f, 26.0f, 26.0f, 26.0f,
+        27.0f, 27.0f, 27.0f, 27.0f, 27.0f, 27.0f, 27.0f, 27.0f, 27.0f, 27.0f, 27.0f, 27.0f, 27.0f,
+        28.0f, 28.0f, 28.0f, 28.0f, 28.0f, 28.0f, 28.0f, 28.0f, 28.0f, 28.0f, 28.0f, 28.0f, 28.0f,
+        29.0f, 29.0f, 29.0f, 29.0f, 29.0f, 29.0f, 29.0f, 29.0f, 29.0f, 29.0f, 29.0f, 29.0f, 29.0f,
+        30.0f, 30.0f, 30.0f, 30.0f, 30.0f, 30.0f, 30.0f, 30.0f, 30.0f, 30.0f, 30.0f, 30.0f, 30.0f,
+        31.0f, 31.0f, 31.0f, 31.0f, 31.0f, 31.0f, 31.0f, 31.0f, 31.0f, 31.0f, 31.0f, 31.0f, 31.0f,
+        32.0f, 32.0f, 32.0f, 32.0f, 32.0f, 32.0f, 32.0f, 32.0f, 32.0f, 32.0f, 32.0f, 32.0f, 32.0f,
+        33.0f, 33.0f, 33.0f, 33.0f, 33.0f, 33.0f, 33.0f, 33.0f, 33.0f, 33.0f, 33.0f, 33.0f, 33.0f,
+        34.0f, 34.0f, 34.0f, 34.0f, 34.0f, 34.0f, 34.0f, 34.0f, 34.0f, 34.0f, 34.0f, 34.0f, 34.0f,
+        35.0f, 35.0f, 35.0f, 35.0f, 35.0f, 35.0f, 35.0f, 35.0f, 35.0f, 35.0f, 35.0f, 35.0f, 35.0f,
+        36.0f, 36.0f, 36.0f, 36.0f, 36.0f, 36.0f, 36.0f, 36.0f, 36.0f, 36.0f, 36.0f, 36.0f, 36.0f,
+        37.0f, 37.0f, 37.0f, 37.0f, 37.0f, 37.0f, 37.0f, 37.0f, 37.0f, 37.0f, 37.0f, 37.0f, 37.0f,
+        38.0f, 38.0f, 38.0f, 38.0f, 38.0f, 38.0f, 38.0f, 38.0f, 38.0f, 38.0f, 38.0f, 38.0f, 38.0f
+    },
+    {  // Joint group 4
+        39.0f, 39.0f, 39.0f, 39.0f, 39.0f, 39.0f, 39.0f, 39.0f, 39.0f, 39.0f, 39.0f, 39.0f, 39.0f,
+        40.0f, 40.0f, 40.0f, 40.0f, 40.0f, 40.0f, 40.0f, 40.0f, 40.0f, 40.0f, 40.0f, 40.0f, 40.0f,
+        41.0f, 41.0f, 41.0f, 41.0f, 41.0f, 41.0f, 41.0f, 41.0f, 41.0f, 41.0f, 41.0f, 41.0f, 41.0f,
+        42.0f, 42.0f, 42.0f, 42.0f, 42.0f, 42.0f, 42.0f, 42.0f, 42.0f, 42.0f, 42.0f, 42.0f, 42.0f,
+        43.0f, 43.0f, 43.0f, 43.0f, 43.0f, 43.0f, 43.0f, 43.0f, 43.0f, 43.0f, 43.0f, 43.0f, 43.0f,
+        44.0f, 44.0f, 44.0f, 44.0f, 44.0f, 44.0f, 44.0f, 44.0f, 44.0f, 44.0f, 44.0f, 44.0f, 44.0f,
+        45.0f, 45.0f, 45.0f, 45.0f, 45.0f, 45.0f, 45.0f, 45.0f, 45.0f, 45.0f, 45.0f, 45.0f, 45.0f,
+        46.0f, 46.0f, 46.0f, 46.0f, 46.0f, 46.0f, 46.0f, 46.0f, 46.0f, 46.0f, 46.0f, 46.0f, 46.0f,
+        47.0f, 47.0f, 47.0f, 47.0f, 47.0f, 47.0f, 47.0f, 47.0f, 47.0f, 47.0f, 47.0f, 47.0f, 47.0f,
+        48.0f, 48.0f, 48.0f, 48.0f, 48.0f, 48.0f, 48.0f, 48.0f, 48.0f, 48.0f, 48.0f, 48.0f, 48.0f,
+        49.0f, 49.0f, 49.0f, 49.0f, 49.0f, 49.0f, 49.0f, 49.0f, 49.0f, 49.0f, 49.0f, 49.0f, 49.0f,
+        50.0f, 50.0f, 50.0f, 50.0f, 50.0f, 50.0f, 50.0f, 50.0f, 50.0f, 50.0f, 50.0f, 50.0f, 50.0f,
+        51.0f, 51.0f, 51.0f, 51.0f, 51.0f, 51.0f, 51.0f, 51.0f, 51.0f, 51.0f, 51.0f, 51.0f, 51.0f,
+        52.0f, 52.0f, 52.0f, 52.0f, 52.0f, 52.0f, 52.0f, 52.0f, 52.0f, 52.0f, 52.0f, 52.0f, 52.0f,
+        53.0f, 53.0f, 53.0f, 53.0f, 53.0f, 53.0f, 53.0f, 53.0f, 53.0f, 53.0f, 53.0f, 53.0f, 53.0f,
+        54.0f, 54.0f, 54.0f, 54.0f, 54.0f, 54.0f, 54.0f, 54.0f, 54.0f, 54.0f, 54.0f, 54.0f, 54.0f,
+        55.0f, 55.0f, 55.0f, 55.0f, 55.0f, 55.0f, 55.0f, 55.0f, 55.0f, 55.0f, 55.0f, 55.0f, 55.0f
+    },
+    {  // Joint group 5
+        56.0f, 56.0f, 56.0f, 56.0f, 56.0f, 56.0f, 56.0f, 56.0f, 56.0f, 56.0f, 56.0f, 56.0f, 56.0f,
+        57.0f, 57.0f, 57.0f, 57.0f, 57.0f, 57.0f, 57.0f, 57.0f, 57.0f, 57.0f, 57.0f, 57.0f, 57.0f,
+        58.0f, 58.0f, 58.0f, 58.0f, 58.0f, 58.0f, 58.0f, 58.0f, 58.0f, 58.0f, 58.0f, 58.0f, 58.0f,
+        59.0f, 59.0f, 59.0f, 59.0f, 59.0f, 59.0f, 59.0f, 59.0f, 59.0f, 59.0f, 59.0f, 59.0f, 59.0f,
+        60.0f, 60.0f, 60.0f, 60.0f, 60.0f, 60.0f, 60.0f, 60.0f, 60.0f, 60.0f, 60.0f, 60.0f, 60.0f,
+        61.0f, 61.0f, 61.0f, 61.0f, 61.0f, 61.0f, 61.0f, 61.0f, 61.0f, 61.0f, 61.0f, 61.0f, 61.0f,
+        62.0f, 62.0f, 62.0f, 62.0f, 62.0f, 62.0f, 62.0f, 62.0f, 62.0f, 62.0f, 62.0f, 62.0f, 62.0f,
+        63.0f, 63.0f, 63.0f, 63.0f, 63.0f, 63.0f, 63.0f, 63.0f, 63.0f, 63.0f, 63.0f, 63.0f, 63.0f,
+        64.0f, 64.0f, 64.0f, 64.0f, 64.0f, 64.0f, 64.0f, 64.0f, 64.0f, 64.0f, 64.0f, 64.0f, 64.0f,
+        65.0f, 65.0f, 65.0f, 65.0f, 65.0f, 65.0f, 65.0f, 65.0f, 65.0f, 65.0f, 65.0f, 65.0f, 65.0f,
+        66.0f, 66.0f, 66.0f, 66.0f, 66.0f, 66.0f, 66.0f, 66.0f, 66.0f, 66.0f, 66.0f, 66.0f, 66.0f,
+        67.0f, 67.0f, 67.0f, 67.0f, 67.0f, 67.0f, 67.0f, 67.0f, 67.0f, 67.0f, 67.0f, 67.0f, 67.0f,
+        68.0f, 68.0f, 68.0f, 68.0f, 68.0f, 68.0f, 68.0f, 68.0f, 68.0f, 68.0f, 68.0f, 68.0f, 68.0f,
+        69.0f, 69.0f, 69.0f, 69.0f, 69.0f, 69.0f, 69.0f, 69.0f, 69.0f, 69.0f, 69.0f, 69.0f, 69.0f,
+        70.0f, 70.0f, 70.0f, 70.0f, 70.0f, 70.0f, 70.0f, 70.0f, 70.0f, 70.0f, 70.0f, 70.0f, 70.0f,
+        71.0f, 71.0f, 71.0f, 71.0f, 71.0f, 71.0f, 71.0f, 71.0f, 71.0f, 71.0f, 71.0f, 71.0f, 71.0f,
+        72.0f, 72.0f, 72.0f, 72.0f, 72.0f, 72.0f, 72.0f, 72.0f, 72.0f, 72.0f, 72.0f, 72.0f, 72.0f,
+        73.0f, 73.0f, 73.0f, 73.0f, 73.0f, 73.0f, 73.0f, 73.0f, 73.0f, 73.0f, 73.0f, 73.0f, 73.0f,
+        74.0f, 74.0f, 74.0f, 74.0f, 74.0f, 74.0f, 74.0f, 74.0f, 74.0f, 74.0f, 74.0f, 74.0f, 74.0f,
+        75.0f, 75.0f, 75.0f, 75.0f, 75.0f, 75.0f, 75.0f, 75.0f, 75.0f, 75.0f, 75.0f, 75.0f, 75.0f,
+        76.0f, 76.0f, 76.0f, 76.0f, 76.0f, 76.0f, 76.0f, 76.0f, 76.0f, 76.0f, 76.0f, 76.0f, 76.0f,
+        77.0f, 77.0f, 77.0f, 77.0f, 77.0f, 77.0f, 77.0f, 77.0f, 77.0f, 77.0f, 77.0f, 77.0f, 77.0f,
+        78.0f, 78.0f, 78.0f, 78.0f, 78.0f, 78.0f, 78.0f, 78.0f, 78.0f, 78.0f, 78.0f, 78.0f, 78.0f,
+        79.0f, 79.0f, 79.0f, 79.0f, 79.0f, 79.0f, 79.0f, 79.0f, 79.0f, 79.0f, 79.0f, 79.0f, 79.0f
+    },
+    {  // Joint group 6
+        80.0f, 80.0f, 80.0f, 80.0f, 80.0f, 80.0f, 80.0f, 80.0f, 80.0f, 80.0f, 80.0f, 80.0f, 80.0f,
+        81.0f, 81.0f, 81.0f, 81.0f, 81.0f, 81.0f, 81.0f, 81.0f, 81.0f, 81.0f, 81.0f, 81.0f, 81.0f,
+        82.0f, 82.0f, 82.0f, 82.0f, 82.0f, 82.0f, 82.0f, 82.0f, 82.0f, 82.0f, 82.0f, 82.0f, 82.0f,
+        83.0f, 83.0f, 83.0f, 83.0f, 83.0f, 83.0f, 83.0f, 83.0f, 83.0f, 83.0f, 83.0f, 83.0f, 83.0f,
+        84.0f, 84.0f, 84.0f, 84.0f, 84.0f, 84.0f, 84.0f, 84.0f, 84.0f, 84.0f, 84.0f, 84.0f, 84.0f,
+        85.0f, 85.0f, 85.0f, 85.0f, 85.0f, 85.0f, 85.0f, 85.0f, 85.0f, 85.0f, 85.0f, 85.0f, 85.0f,
+        86.0f, 86.0f, 86.0f, 86.0f, 86.0f, 86.0f, 86.0f, 86.0f, 86.0f, 86.0f, 86.0f, 86.0f, 86.0f,
+        87.0f, 87.0f, 87.0f, 87.0f, 87.0f, 87.0f, 87.0f, 87.0f, 87.0f, 87.0f, 87.0f, 87.0f, 87.0f,
+        88.0f, 88.0f, 88.0f, 88.0f, 88.0f, 88.0f, 88.0f, 88.0f, 88.0f, 88.0f, 88.0f, 88.0f, 88.0f,
+        89.0f, 89.0f, 89.0f, 89.0f, 89.0f, 89.0f, 89.0f, 89.0f, 89.0f, 89.0f, 89.0f, 89.0f, 89.0f,
+        90.0f, 90.0f, 90.0f, 90.0f, 90.0f, 90.0f, 90.0f, 90.0f, 90.0f, 90.0f, 90.0f, 90.0f, 90.0f,
+        91.0f, 91.0f, 91.0f, 91.0f, 91.0f, 91.0f, 91.0f, 91.0f, 91.0f, 91.0f, 91.0f, 91.0f, 91.0f,
+        92.0f, 92.0f, 92.0f, 92.0f, 92.0f, 92.0f, 92.0f, 92.0f, 92.0f, 92.0f, 92.0f, 92.0f, 92.0f,
+        93.0f, 93.0f, 93.0f, 93.0f, 93.0f, 93.0f, 93.0f, 93.0f, 93.0f, 93.0f, 93.0f, 93.0f, 93.0f,
+        94.0f, 94.0f, 94.0f, 94.0f, 94.0f, 94.0f, 94.0f, 94.0f, 94.0f, 94.0f, 94.0f, 94.0f, 94.0f,
+        95.0f, 95.0f, 95.0f, 95.0f, 95.0f, 95.0f, 95.0f, 95.0f, 95.0f, 95.0f, 95.0f, 95.0f, 95.0f,
+        96.0f, 96.0f, 96.0f, 96.0f, 96.0f, 96.0f, 96.0f, 96.0f, 96.0f, 96.0f, 96.0f, 96.0f, 96.0f,
+        97.0f, 97.0f, 97.0f, 97.0f, 97.0f, 97.0f, 97.0f, 97.0f, 97.0f, 97.0f, 97.0f, 97.0f, 97.0f,
+        98.0f, 98.0f, 98.0f, 98.0f, 98.0f, 98.0f, 98.0f, 98.0f, 98.0f, 98.0f, 98.0f, 98.0f, 98.0f,
+        99.0f, 99.0f, 99.0f, 99.0f, 99.0f, 99.0f, 99.0f, 99.0f, 99.0f, 99.0f, 99.0f, 99.0f, 99.0f,
+        100.0f, 100.0f, 100.0f, 100.0f, 100.0f, 100.0f, 100.0f, 100.0f, 100.0f, 100.0f, 100.0f, 100.0f, 100.0f,
+        101.0f, 101.0f, 101.0f, 101.0f, 101.0f, 101.0f, 101.0f, 101.0f, 101.0f, 101.0f, 101.0f, 101.0f, 101.0f,
+        102.0f, 102.0f, 102.0f, 102.0f, 102.0f, 102.0f, 102.0f, 102.0f, 102.0f, 102.0f, 102.0f, 102.0f, 102.0f,
+        103.0f, 103.0f, 103.0f, 103.0f, 103.0f, 103.0f, 103.0f, 103.0f, 103.0f, 103.0f, 103.0f, 103.0f, 103.0f,
+        104.0f, 104.0f, 104.0f, 104.0f, 104.0f, 104.0f, 104.0f, 104.0f, 104.0f, 104.0f, 104.0f, 104.0f, 104.0f
+    },
+    {  // Joint group 7
+        105.0f, 105.0f, 105.0f, 105.0f, 105.0f, 105.0f, 105.0f, 105.0f, 105.0f, 105.0f, 105.0f, 105.0f, 105.0f,
+        106.0f, 106.0f, 106.0f, 106.0f, 106.0f, 106.0f, 106.0f, 106.0f, 106.0f, 106.0f, 106.0f, 106.0f, 106.0f,
+        107.0f, 107.0f, 107.0f, 107.0f, 107.0f, 107.0f, 107.0f, 107.0f, 107.0f, 107.0f, 107.0f, 107.0f, 107.0f,
+        108.0f, 108.0f, 108.0f, 108.0f, 108.0f, 108.0f, 108.0f, 108.0f, 108.0f, 108.0f, 108.0f, 108.0f, 108.0f,
+        109.0f, 109.0f, 109.0f, 109.0f, 109.0f, 109.0f, 109.0f, 109.0f, 109.0f, 109.0f, 109.0f, 109.0f, 109.0f,
+        110.0f, 110.0f, 110.0f, 110.0f, 110.0f, 110.0f, 110.0f, 110.0f, 110.0f, 110.0f, 110.0f, 110.0f, 110.0f,
+        111.0f, 111.0f, 111.0f, 111.0f, 111.0f, 111.0f, 111.0f, 111.0f, 111.0f, 111.0f, 111.0f, 111.0f, 111.0f,
+        112.0f, 112.0f, 112.0f, 112.0f, 112.0f, 112.0f, 112.0f, 112.0f, 112.0f, 112.0f, 112.0f, 112.0f, 112.0f,
+        113.0f, 113.0f, 113.0f, 113.0f, 113.0f, 113.0f, 113.0f, 113.0f, 113.0f, 113.0f, 113.0f, 113.0f, 113.0f,
+        114.0f, 114.0f, 114.0f, 114.0f, 114.0f, 114.0f, 114.0f, 114.0f, 114.0f, 114.0f, 114.0f, 114.0f, 114.0f,
+        115.0f, 115.0f, 115.0f, 115.0f, 115.0f, 115.0f, 115.0f, 115.0f, 115.0f, 115.0f, 115.0f, 115.0f, 115.0f,
+        116.0f, 116.0f, 116.0f, 116.0f, 116.0f, 116.0f, 116.0f, 116.0f, 116.0f, 116.0f, 116.0f, 116.0f, 116.0f,
+        117.0f, 117.0f, 117.0f, 117.0f, 117.0f, 117.0f, 117.0f, 117.0f, 117.0f, 117.0f, 117.0f, 117.0f, 117.0f,
+        118.0f, 118.0f, 118.0f, 118.0f, 118.0f, 118.0f, 118.0f, 118.0f, 118.0f, 118.0f, 118.0f, 118.0f, 118.0f,
+        119.0f, 119.0f, 119.0f, 119.0f, 119.0f, 119.0f, 119.0f, 119.0f, 119.0f, 119.0f, 119.0f, 119.0f, 119.0f,
+        120.0f, 120.0f, 120.0f, 120.0f, 120.0f, 120.0f, 120.0f, 120.0f, 120.0f, 120.0f, 120.0f, 120.0f, 120.0f,
+        121.0f, 121.0f, 121.0f, 121.0f, 121.0f, 121.0f, 121.0f, 121.0f, 121.0f, 121.0f, 121.0f, 121.0f, 121.0f,
+        122.0f, 122.0f, 122.0f, 122.0f, 122.0f, 122.0f, 122.0f, 122.0f, 122.0f, 122.0f, 122.0f, 122.0f, 122.0f,
+        123.0f, 123.0f, 123.0f, 123.0f, 123.0f, 123.0f, 123.0f, 123.0f, 123.0f, 123.0f, 123.0f, 123.0f, 123.0f,
+        124.0f, 124.0f, 124.0f, 124.0f, 124.0f, 124.0f, 124.0f, 124.0f, 124.0f, 124.0f, 124.0f, 124.0f, 124.0f,
+        125.0f, 125.0f, 125.0f, 125.0f, 125.0f, 125.0f, 125.0f, 125.0f, 125.0f, 125.0f, 125.0f, 125.0f, 125.0f,
+        126.0f, 126.0f, 126.0f, 126.0f, 126.0f, 126.0f, 126.0f, 126.0f, 126.0f, 126.0f, 126.0f, 126.0f, 126.0f,
+        127.0f, 127.0f, 127.0f, 127.0f, 127.0f, 127.0f, 127.0f, 127.0f, 127.0f, 127.0f, 127.0f, 127.0f, 127.0f,
+        128.0f, 128.0f, 128.0f, 128.0f, 128.0f, 128.0f, 128.0f, 128.0f, 128.0f, 128.0f, 128.0f, 128.0f, 128.0f,
+        129.0f, 129.0f, 129.0f, 129.0f, 129.0f, 129.0f, 129.0f, 129.0f, 129.0f, 129.0f, 129.0f, 129.0f, 129.0f,
+        130.0f, 130.0f, 130.0f, 130.0f, 130.0f, 130.0f, 130.0f, 130.0f, 130.0f, 130.0f, 130.0f, 130.0f, 130.0f,
+        131.0f, 131.0f, 131.0f, 131.0f, 131.0f, 131.0f, 131.0f, 131.0f, 131.0f, 131.0f, 131.0f, 131.0f, 131.0f,
+        132.0f, 132.0f, 132.0f, 132.0f, 132.0f, 132.0f, 132.0f, 132.0f, 132.0f, 132.0f, 132.0f, 132.0f, 132.0f,
+        133.0f, 133.0f, 133.0f, 133.0f, 133.0f, 133.0f, 133.0f, 133.0f, 133.0f, 133.0f, 133.0f, 133.0f, 133.0f,
+        134.0f, 134.0f, 134.0f, 134.0f, 134.0f, 134.0f, 134.0f, 134.0f, 134.0f, 134.0f, 134.0f, 134.0f, 134.0f,
+        135.0f, 135.0f, 135.0f, 135.0f, 135.0f, 135.0f, 135.0f, 135.0f, 135.0f, 135.0f, 135.0f, 135.0f, 135.0f,
+        136.0f, 136.0f, 136.0f, 136.0f, 136.0f, 136.0f, 136.0f, 136.0f, 136.0f, 136.0f, 136.0f, 136.0f, 136.0f
+    }
+};
+
+const Matrix<std::uint16_t> inputIndices = {
+    {  // Joint group 0
+        12, 11, 10, 9, 8, 7, 6, 5, 4, 3, 2, 1, 0,
+    },
+    {  // Joint group 1
+        12, 11, 10, 9, 8, 7, 6, 5, 4, 3, 2, 1, 0,
+    },
+    {  // Joint group 2
+        12, 11, 10, 9, 8, 7, 6, 5, 4, 3, 2, 1, 0,
+    },
+    {  // Joint group 3
+        12, 11, 10, 9, 8, 7, 6, 5, 4, 3, 2, 1, 0,
+    },
+    {  // Joint group 4
+        12, 11, 10, 9, 8, 7, 6, 5, 4, 3, 2, 1, 0,
+    },
+    {  // Joint group 5
+        12, 11, 10, 9, 8, 7, 6, 5, 4, 3, 2, 1, 0,
+    },
+    {  // Joint group 6
+        12, 11, 10, 9, 8, 7, 6, 5, 4, 3, 2, 1, 0,
+    },
+    {  // Joint group 7
+        12, 11, 10, 9, 8, 7, 6, 5, 4, 3, 2, 1, 0,
+    },
+};
+
+const Matrix<std::uint16_t> outputIndices = {
+    {  // Joint group 0
+        0, 1, 2, 3, 4
+    },
+    {  // Joint group 1
+        5, 6, 7, 8, 9, 10, 11, 12
+    },
+    {  // Joint group 2
+        13, 14, 15, 16, 17, 18, 19, 20, 21
+    },
+    {  // Joint group 3
+        22, 23, 24, 25, 26, 27, 28, 29, 30, 31, 32, 33, 34, 35, 36, 37
+    },
+    {  // Joint group 4
+        38, 39, 40, 41, 42, 43, 44, 45, 46, 47, 48, 49, 50, 51, 52, 53, 54
+    },
+    {  // Joint group 5
+        55, 56, 57, 58, 59, 60, 61, 62, 63, 64, 65, 66, 67, 68, 69, 70, 71, 72, 73, 74, 75, 76, 77, 78
+    },
+    {  // Joint group 6
+        79, 80, 81, 82, 83, 84, 85, 86, 87, 88, 89, 90, 91, 92, 93, 94, 95, 96, 97, 98, 99, 100, 101, 102, 103
+    },
+    {  // Joint group 7
+        104, 105, 106, 107, 108, 109, 110, 111, 112, 113, 114, 115, 116, 117, 118, 119,
+        120, 121, 122, 123, 124, 125, 126, 127, 128, 129, 130, 131, 132, 133, 134, 135
+    }
+};
+
+const Matrix<std::uint16_t> lods = {
+    {  // Joint group 0
+        5, 4, 2, 1
+    },
+    {  // Joint group 1
+        8, 5, 3, 1
+    },
+    {  // Joint group 2
+        9, 7, 4, 2
+    },
+    {  // Joint group 3
+        16, 13, 8, 4
+    },
+    {  // Joint group 4
+        17, 16, 15, 12
+    },
+    {  // Joint group 5
+        24, 16, 15, 8
+    },
+    {  // Joint group 6
+        25, 17, 12, 7
+    },
+    {  // Joint group 7
+        32, 25, 17, 9
+    }
+};
+
+}  // namespace unoptimized
+
+namespace optimized {
+
+const bpcm::Extent dimensions = {150ul, 13ul};
+
+const AlignedMatrix<float> floatValues = {
+    {  // Joint group 0 - Pad to block-8
+        1.0f, 2.0f, 3.0f, 4.0f, 5.0f, 0.0f, 0.0f, 0.0f,
+        1.0f, 2.0f, 3.0f, 4.0f, 5.0f, 0.0f, 0.0f, 0.0f,
+        1.0f, 2.0f, 3.0f, 4.0f, 5.0f, 0.0f, 0.0f, 0.0f,
+        1.0f, 2.0f, 3.0f, 4.0f, 5.0f, 0.0f, 0.0f, 0.0f,
+        1.0f, 2.0f, 3.0f, 4.0f, 5.0f, 0.0f, 0.0f, 0.0f,
+        1.0f, 2.0f, 3.0f, 4.0f, 5.0f, 0.0f, 0.0f, 0.0f,
+        1.0f, 2.0f, 3.0f, 4.0f, 5.0f, 0.0f, 0.0f, 0.0f,
+        1.0f, 2.0f, 3.0f, 4.0f, 5.0f, 0.0f, 0.0f, 0.0f,
+        1.0f, 2.0f, 3.0f, 4.0f, 5.0f, 0.0f, 0.0f, 0.0f,
+        1.0f, 2.0f, 3.0f, 4.0f, 5.0f, 0.0f, 0.0f, 0.0f,
+        1.0f, 2.0f, 3.0f, 4.0f, 5.0f, 0.0f, 0.0f, 0.0f,
+        1.0f, 2.0f, 3.0f, 4.0f, 5.0f, 0.0f, 0.0f, 0.0f,
+        1.0f, 2.0f, 3.0f, 4.0f, 5.0f, 0.0f, 0.0f, 0.0f
+    },
+    {  // Joint group 1 - No padding => block-8
+        6.0f, 7.0f, 8.0f, 9.0f, 10.0f, 11.0f, 12.0f, 13.0f,
+        6.0f, 7.0f, 8.0f, 9.0f, 10.0f, 11.0f, 12.0f, 13.0f,
+        6.0f, 7.0f, 8.0f, 9.0f, 10.0f, 11.0f, 12.0f, 13.0f,
+        6.0f, 7.0f, 8.0f, 9.0f, 10.0f, 11.0f, 12.0f, 13.0f,
+        6.0f, 7.0f, 8.0f, 9.0f, 10.0f, 11.0f, 12.0f, 13.0f,
+        6.0f, 7.0f, 8.0f, 9.0f, 10.0f, 11.0f, 12.0f, 13.0f,
+        6.0f, 7.0f, 8.0f, 9.0f, 10.0f, 11.0f, 12.0f, 13.0f,
+        6.0f, 7.0f, 8.0f, 9.0f, 10.0f, 11.0f, 12.0f, 13.0f,
+        6.0f, 7.0f, 8.0f, 9.0f, 10.0f, 11.0f, 12.0f, 13.0f,
+        6.0f, 7.0f, 8.0f, 9.0f, 10.0f, 11.0f, 12.0f, 13.0f,
+        6.0f, 7.0f, 8.0f, 9.0f, 10.0f, 11.0f, 12.0f, 13.0f,
+        6.0f, 7.0f, 8.0f, 9.0f, 10.0f, 11.0f, 12.0f, 13.0f,
+        6.0f, 7.0f, 8.0f, 9.0f, 10.0f, 11.0f, 12.0f, 13.0f
+    },
+    {  // Joint group 2 - Pad to block-8 => block-16
+        14.0f, 15.0f, 16.0f, 17.0f, 18.0f, 19.0f, 20.0f, 21.0f, 22.0f, 0.0f, 0.0f, 0.0f, 0.0f, 0.0f, 0.0f, 0.0f,
+        14.0f, 15.0f, 16.0f, 17.0f, 18.0f, 19.0f, 20.0f, 21.0f, 22.0f, 0.0f, 0.0f, 0.0f, 0.0f, 0.0f, 0.0f, 0.0f,
+        14.0f, 15.0f, 16.0f, 17.0f, 18.0f, 19.0f, 20.0f, 21.0f, 22.0f, 0.0f, 0.0f, 0.0f, 0.0f, 0.0f, 0.0f, 0.0f,
+        14.0f, 15.0f, 16.0f, 17.0f, 18.0f, 19.0f, 20.0f, 21.0f, 22.0f, 0.0f, 0.0f, 0.0f, 0.0f, 0.0f, 0.0f, 0.0f,
+        14.0f, 15.0f, 16.0f, 17.0f, 18.0f, 19.0f, 20.0f, 21.0f, 22.0f, 0.0f, 0.0f, 0.0f, 0.0f, 0.0f, 0.0f, 0.0f,
+        14.0f, 15.0f, 16.0f, 17.0f, 18.0f, 19.0f, 20.0f, 21.0f, 22.0f, 0.0f, 0.0f, 0.0f, 0.0f, 0.0f, 0.0f, 0.0f,
+        14.0f, 15.0f, 16.0f, 17.0f, 18.0f, 19.0f, 20.0f, 21.0f, 22.0f, 0.0f, 0.0f, 0.0f, 0.0f, 0.0f, 0.0f, 0.0f,
+        14.0f, 15.0f, 16.0f, 17.0f, 18.0f, 19.0f, 20.0f, 21.0f, 22.0f, 0.0f, 0.0f, 0.0f, 0.0f, 0.0f, 0.0f, 0.0f,
+        14.0f, 15.0f, 16.0f, 17.0f, 18.0f, 19.0f, 20.0f, 21.0f, 22.0f, 0.0f, 0.0f, 0.0f, 0.0f, 0.0f, 0.0f, 0.0f,
+        14.0f, 15.0f, 16.0f, 17.0f, 18.0f, 19.0f, 20.0f, 21.0f, 22.0f, 0.0f, 0.0f, 0.0f, 0.0f, 0.0f, 0.0f, 0.0f,
+        14.0f, 15.0f, 16.0f, 17.0f, 18.0f, 19.0f, 20.0f, 21.0f, 22.0f, 0.0f, 0.0f, 0.0f, 0.0f, 0.0f, 0.0f, 0.0f,
+        14.0f, 15.0f, 16.0f, 17.0f, 18.0f, 19.0f, 20.0f, 21.0f, 22.0f, 0.0f, 0.0f, 0.0f, 0.0f, 0.0f, 0.0f, 0.0f,
+        14.0f, 15.0f, 16.0f, 17.0f, 18.0f, 19.0f, 20.0f, 21.0f, 22.0f, 0.0f, 0.0f, 0.0f, 0.0f, 0.0f, 0.0f, 0.0f
+    },
+    {  // Joint group 3 - No padding => block-16
+        23.0f, 24.0f, 25.0f, 26.0f, 27.0f, 28.0f, 29.0f, 30.0f, 31.0f, 32.0f, 33.0f, 34.0f, 35.0f, 36.0f, 37.0f, 38.0f,
+        23.0f, 24.0f, 25.0f, 26.0f, 27.0f, 28.0f, 29.0f, 30.0f, 31.0f, 32.0f, 33.0f, 34.0f, 35.0f, 36.0f, 37.0f, 38.0f,
+        23.0f, 24.0f, 25.0f, 26.0f, 27.0f, 28.0f, 29.0f, 30.0f, 31.0f, 32.0f, 33.0f, 34.0f, 35.0f, 36.0f, 37.0f, 38.0f,
+        23.0f, 24.0f, 25.0f, 26.0f, 27.0f, 28.0f, 29.0f, 30.0f, 31.0f, 32.0f, 33.0f, 34.0f, 35.0f, 36.0f, 37.0f, 38.0f,
+        23.0f, 24.0f, 25.0f, 26.0f, 27.0f, 28.0f, 29.0f, 30.0f, 31.0f, 32.0f, 33.0f, 34.0f, 35.0f, 36.0f, 37.0f, 38.0f,
+        23.0f, 24.0f, 25.0f, 26.0f, 27.0f, 28.0f, 29.0f, 30.0f, 31.0f, 32.0f, 33.0f, 34.0f, 35.0f, 36.0f, 37.0f, 38.0f,
+        23.0f, 24.0f, 25.0f, 26.0f, 27.0f, 28.0f, 29.0f, 30.0f, 31.0f, 32.0f, 33.0f, 34.0f, 35.0f, 36.0f, 37.0f, 38.0f,
+        23.0f, 24.0f, 25.0f, 26.0f, 27.0f, 28.0f, 29.0f, 30.0f, 31.0f, 32.0f, 33.0f, 34.0f, 35.0f, 36.0f, 37.0f, 38.0f,
+        23.0f, 24.0f, 25.0f, 26.0f, 27.0f, 28.0f, 29.0f, 30.0f, 31.0f, 32.0f, 33.0f, 34.0f, 35.0f, 36.0f, 37.0f, 38.0f,
+        23.0f, 24.0f, 25.0f, 26.0f, 27.0f, 28.0f, 29.0f, 30.0f, 31.0f, 32.0f, 33.0f, 34.0f, 35.0f, 36.0f, 37.0f, 38.0f,
+        23.0f, 24.0f, 25.0f, 26.0f, 27.0f, 28.0f, 29.0f, 30.0f, 31.0f, 32.0f, 33.0f, 34.0f, 35.0f, 36.0f, 37.0f, 38.0f,
+        23.0f, 24.0f, 25.0f, 26.0f, 27.0f, 28.0f, 29.0f, 30.0f, 31.0f, 32.0f, 33.0f, 34.0f, 35.0f, 36.0f, 37.0f, 38.0f,
+        23.0f, 24.0f, 25.0f, 26.0f, 27.0f, 28.0f, 29.0f, 30.0f, 31.0f, 32.0f, 33.0f, 34.0f, 35.0f, 36.0f, 37.0f, 38.0f
+    },
+    {  // Joint group 4 - Pad to block-8 => block-16 + block-8
+        39.0f, 40.0f, 41.0f, 42.0f, 43.0f, 44.0f, 45.0f, 46.0f, 47.0f, 48.0f, 49.0f, 50.0f, 51.0f, 52.0f, 53.0f, 54.0f,
+        39.0f, 40.0f, 41.0f, 42.0f, 43.0f, 44.0f, 45.0f, 46.0f, 47.0f, 48.0f, 49.0f, 50.0f, 51.0f, 52.0f, 53.0f, 54.0f,
+        39.0f, 40.0f, 41.0f, 42.0f, 43.0f, 44.0f, 45.0f, 46.0f, 47.0f, 48.0f, 49.0f, 50.0f, 51.0f, 52.0f, 53.0f, 54.0f,
+        39.0f, 40.0f, 41.0f, 42.0f, 43.0f, 44.0f, 45.0f, 46.0f, 47.0f, 48.0f, 49.0f, 50.0f, 51.0f, 52.0f, 53.0f, 54.0f,
+        39.0f, 40.0f, 41.0f, 42.0f, 43.0f, 44.0f, 45.0f, 46.0f, 47.0f, 48.0f, 49.0f, 50.0f, 51.0f, 52.0f, 53.0f, 54.0f,
+        39.0f, 40.0f, 41.0f, 42.0f, 43.0f, 44.0f, 45.0f, 46.0f, 47.0f, 48.0f, 49.0f, 50.0f, 51.0f, 52.0f, 53.0f, 54.0f,
+        39.0f, 40.0f, 41.0f, 42.0f, 43.0f, 44.0f, 45.0f, 46.0f, 47.0f, 48.0f, 49.0f, 50.0f, 51.0f, 52.0f, 53.0f, 54.0f,
+        39.0f, 40.0f, 41.0f, 42.0f, 43.0f, 44.0f, 45.0f, 46.0f, 47.0f, 48.0f, 49.0f, 50.0f, 51.0f, 52.0f, 53.0f, 54.0f,
+        39.0f, 40.0f, 41.0f, 42.0f, 43.0f, 44.0f, 45.0f, 46.0f, 47.0f, 48.0f, 49.0f, 50.0f, 51.0f, 52.0f, 53.0f, 54.0f,
+        39.0f, 40.0f, 41.0f, 42.0f, 43.0f, 44.0f, 45.0f, 46.0f, 47.0f, 48.0f, 49.0f, 50.0f, 51.0f, 52.0f, 53.0f, 54.0f,
+        39.0f, 40.0f, 41.0f, 42.0f, 43.0f, 44.0f, 45.0f, 46.0f, 47.0f, 48.0f, 49.0f, 50.0f, 51.0f, 52.0f, 53.0f, 54.0f,
+        39.0f, 40.0f, 41.0f, 42.0f, 43.0f, 44.0f, 45.0f, 46.0f, 47.0f, 48.0f, 49.0f, 50.0f, 51.0f, 52.0f, 53.0f, 54.0f,
+        39.0f, 40.0f, 41.0f, 42.0f, 43.0f, 44.0f, 45.0f, 46.0f, 47.0f, 48.0f, 49.0f, 50.0f, 51.0f, 52.0f, 53.0f, 54.0f,
+        55.0f, 0.0f, 0.0f, 0.0f, 0.0f, 0.0f, 0.0f, 0.0f,
+        55.0f, 0.0f, 0.0f, 0.0f, 0.0f, 0.0f, 0.0f, 0.0f,
+        55.0f, 0.0f, 0.0f, 0.0f, 0.0f, 0.0f, 0.0f, 0.0f,
+        55.0f, 0.0f, 0.0f, 0.0f, 0.0f, 0.0f, 0.0f, 0.0f,
+        55.0f, 0.0f, 0.0f, 0.0f, 0.0f, 0.0f, 0.0f, 0.0f,
+        55.0f, 0.0f, 0.0f, 0.0f, 0.0f, 0.0f, 0.0f, 0.0f,
+        55.0f, 0.0f, 0.0f, 0.0f, 0.0f, 0.0f, 0.0f, 0.0f,
+        55.0f, 0.0f, 0.0f, 0.0f, 0.0f, 0.0f, 0.0f, 0.0f,
+        55.0f, 0.0f, 0.0f, 0.0f, 0.0f, 0.0f, 0.0f, 0.0f,
+        55.0f, 0.0f, 0.0f, 0.0f, 0.0f, 0.0f, 0.0f, 0.0f,
+        55.0f, 0.0f, 0.0f, 0.0f, 0.0f, 0.0f, 0.0f, 0.0f,
+        55.0f, 0.0f, 0.0f, 0.0f, 0.0f, 0.0f, 0.0f, 0.0f,
+        55.0f, 0.0f, 0.0f, 0.0f, 0.0f, 0.0f, 0.0f, 0.0f
+    },
+    {  // Joint group 5 - No padding -> block-16 + block-8
+        56.0f, 57.0f, 58.0f, 59.0f, 60.0f, 61.0f, 62.0f, 63.0f, 64.0f, 65.0f, 66.0f, 67.0f, 68.0f, 69.0f, 70.0f, 71.0f,
+        56.0f, 57.0f, 58.0f, 59.0f, 60.0f, 61.0f, 62.0f, 63.0f, 64.0f, 65.0f, 66.0f, 67.0f, 68.0f, 69.0f, 70.0f, 71.0f,
+        56.0f, 57.0f, 58.0f, 59.0f, 60.0f, 61.0f, 62.0f, 63.0f, 64.0f, 65.0f, 66.0f, 67.0f, 68.0f, 69.0f, 70.0f, 71.0f,
+        56.0f, 57.0f, 58.0f, 59.0f, 60.0f, 61.0f, 62.0f, 63.0f, 64.0f, 65.0f, 66.0f, 67.0f, 68.0f, 69.0f, 70.0f, 71.0f,
+        56.0f, 57.0f, 58.0f, 59.0f, 60.0f, 61.0f, 62.0f, 63.0f, 64.0f, 65.0f, 66.0f, 67.0f, 68.0f, 69.0f, 70.0f, 71.0f,
+        56.0f, 57.0f, 58.0f, 59.0f, 60.0f, 61.0f, 62.0f, 63.0f, 64.0f, 65.0f, 66.0f, 67.0f, 68.0f, 69.0f, 70.0f, 71.0f,
+        56.0f, 57.0f, 58.0f, 59.0f, 60.0f, 61.0f, 62.0f, 63.0f, 64.0f, 65.0f, 66.0f, 67.0f, 68.0f, 69.0f, 70.0f, 71.0f,
+        56.0f, 57.0f, 58.0f, 59.0f, 60.0f, 61.0f, 62.0f, 63.0f, 64.0f, 65.0f, 66.0f, 67.0f, 68.0f, 69.0f, 70.0f, 71.0f,
+        56.0f, 57.0f, 58.0f, 59.0f, 60.0f, 61.0f, 62.0f, 63.0f, 64.0f, 65.0f, 66.0f, 67.0f, 68.0f, 69.0f, 70.0f, 71.0f,
+        56.0f, 57.0f, 58.0f, 59.0f, 60.0f, 61.0f, 62.0f, 63.0f, 64.0f, 65.0f, 66.0f, 67.0f, 68.0f, 69.0f, 70.0f, 71.0f,
+        56.0f, 57.0f, 58.0f, 59.0f, 60.0f, 61.0f, 62.0f, 63.0f, 64.0f, 65.0f, 66.0f, 67.0f, 68.0f, 69.0f, 70.0f, 71.0f,
+        56.0f, 57.0f, 58.0f, 59.0f, 60.0f, 61.0f, 62.0f, 63.0f, 64.0f, 65.0f, 66.0f, 67.0f, 68.0f, 69.0f, 70.0f, 71.0f,
+        56.0f, 57.0f, 58.0f, 59.0f, 60.0f, 61.0f, 62.0f, 63.0f, 64.0f, 65.0f, 66.0f, 67.0f, 68.0f, 69.0f, 70.0f, 71.0f,
+        72.0f, 73.0f, 74.0f, 75.0f, 76.0f, 77.0f, 78.0f, 79.0f,
+        72.0f, 73.0f, 74.0f, 75.0f, 76.0f, 77.0f, 78.0f, 79.0f,
+        72.0f, 73.0f, 74.0f, 75.0f, 76.0f, 77.0f, 78.0f, 79.0f,
+        72.0f, 73.0f, 74.0f, 75.0f, 76.0f, 77.0f, 78.0f, 79.0f,
+        72.0f, 73.0f, 74.0f, 75.0f, 76.0f, 77.0f, 78.0f, 79.0f,
+        72.0f, 73.0f, 74.0f, 75.0f, 76.0f, 77.0f, 78.0f, 79.0f,
+        72.0f, 73.0f, 74.0f, 75.0f, 76.0f, 77.0f, 78.0f, 79.0f,
+        72.0f, 73.0f, 74.0f, 75.0f, 76.0f, 77.0f, 78.0f, 79.0f,
+        72.0f, 73.0f, 74.0f, 75.0f, 76.0f, 77.0f, 78.0f, 79.0f,
+        72.0f, 73.0f, 74.0f, 75.0f, 76.0f, 77.0f, 78.0f, 79.0f,
+        72.0f, 73.0f, 74.0f, 75.0f, 76.0f, 77.0f, 78.0f, 79.0f,
+        72.0f, 73.0f, 74.0f, 75.0f, 76.0f, 77.0f, 78.0f, 79.0f,
+        72.0f, 73.0f, 74.0f, 75.0f, 76.0f, 77.0f, 78.0f, 79.0f
+    },
+    {  // Joint group 6 - Pad to block-16 -> block-16 + block-16
+        80.0f, 81.0f, 82.0f, 83.0f, 84.0f, 85.0f, 86.0f, 87.0f, 88.0f, 89.0f, 90.0f, 91.0f, 92.0f, 93.0f, 94.0f, 95.0f,
+        80.0f, 81.0f, 82.0f, 83.0f, 84.0f, 85.0f, 86.0f, 87.0f, 88.0f, 89.0f, 90.0f, 91.0f, 92.0f, 93.0f, 94.0f, 95.0f,
+        80.0f, 81.0f, 82.0f, 83.0f, 84.0f, 85.0f, 86.0f, 87.0f, 88.0f, 89.0f, 90.0f, 91.0f, 92.0f, 93.0f, 94.0f, 95.0f,
+        80.0f, 81.0f, 82.0f, 83.0f, 84.0f, 85.0f, 86.0f, 87.0f, 88.0f, 89.0f, 90.0f, 91.0f, 92.0f, 93.0f, 94.0f, 95.0f,
+        80.0f, 81.0f, 82.0f, 83.0f, 84.0f, 85.0f, 86.0f, 87.0f, 88.0f, 89.0f, 90.0f, 91.0f, 92.0f, 93.0f, 94.0f, 95.0f,
+        80.0f, 81.0f, 82.0f, 83.0f, 84.0f, 85.0f, 86.0f, 87.0f, 88.0f, 89.0f, 90.0f, 91.0f, 92.0f, 93.0f, 94.0f, 95.0f,
+        80.0f, 81.0f, 82.0f, 83.0f, 84.0f, 85.0f, 86.0f, 87.0f, 88.0f, 89.0f, 90.0f, 91.0f, 92.0f, 93.0f, 94.0f, 95.0f,
+        80.0f, 81.0f, 82.0f, 83.0f, 84.0f, 85.0f, 86.0f, 87.0f, 88.0f, 89.0f, 90.0f, 91.0f, 92.0f, 93.0f, 94.0f, 95.0f,
+        80.0f, 81.0f, 82.0f, 83.0f, 84.0f, 85.0f, 86.0f, 87.0f, 88.0f, 89.0f, 90.0f, 91.0f, 92.0f, 93.0f, 94.0f, 95.0f,
+        80.0f, 81.0f, 82.0f, 83.0f, 84.0f, 85.0f, 86.0f, 87.0f, 88.0f, 89.0f, 90.0f, 91.0f, 92.0f, 93.0f, 94.0f, 95.0f,
+        80.0f, 81.0f, 82.0f, 83.0f, 84.0f, 85.0f, 86.0f, 87.0f, 88.0f, 89.0f, 90.0f, 91.0f, 92.0f, 93.0f, 94.0f, 95.0f,
+        80.0f, 81.0f, 82.0f, 83.0f, 84.0f, 85.0f, 86.0f, 87.0f, 88.0f, 89.0f, 90.0f, 91.0f, 92.0f, 93.0f, 94.0f, 95.0f,
+        80.0f, 81.0f, 82.0f, 83.0f, 84.0f, 85.0f, 86.0f, 87.0f, 88.0f, 89.0f, 90.0f, 91.0f, 92.0f, 93.0f, 94.0f, 95.0f,
+        96.0f, 97.0f, 98.0f, 99.0f, 100.0f, 101.0f, 102.0f, 103.0f, 104.0f, 0.0f, 0.0f, 0.0f, 0.0f, 0.0f, 0.0f, 0.0f,
+        96.0f, 97.0f, 98.0f, 99.0f, 100.0f, 101.0f, 102.0f, 103.0f, 104.0f, 0.0f, 0.0f, 0.0f, 0.0f, 0.0f, 0.0f, 0.0f,
+        96.0f, 97.0f, 98.0f, 99.0f, 100.0f, 101.0f, 102.0f, 103.0f, 104.0f, 0.0f, 0.0f, 0.0f, 0.0f, 0.0f, 0.0f, 0.0f,
+        96.0f, 97.0f, 98.0f, 99.0f, 100.0f, 101.0f, 102.0f, 103.0f, 104.0f, 0.0f, 0.0f, 0.0f, 0.0f, 0.0f, 0.0f, 0.0f,
+        96.0f, 97.0f, 98.0f, 99.0f, 100.0f, 101.0f, 102.0f, 103.0f, 104.0f, 0.0f, 0.0f, 0.0f, 0.0f, 0.0f, 0.0f, 0.0f,
+        96.0f, 97.0f, 98.0f, 99.0f, 100.0f, 101.0f, 102.0f, 103.0f, 104.0f, 0.0f, 0.0f, 0.0f, 0.0f, 0.0f, 0.0f, 0.0f,
+        96.0f, 97.0f, 98.0f, 99.0f, 100.0f, 101.0f, 102.0f, 103.0f, 104.0f, 0.0f, 0.0f, 0.0f, 0.0f, 0.0f, 0.0f, 0.0f,
+        96.0f, 97.0f, 98.0f, 99.0f, 100.0f, 101.0f, 102.0f, 103.0f, 104.0f, 0.0f, 0.0f, 0.0f, 0.0f, 0.0f, 0.0f, 0.0f,
+        96.0f, 97.0f, 98.0f, 99.0f, 100.0f, 101.0f, 102.0f, 103.0f, 104.0f, 0.0f, 0.0f, 0.0f, 0.0f, 0.0f, 0.0f, 0.0f,
+        96.0f, 97.0f, 98.0f, 99.0f, 100.0f, 101.0f, 102.0f, 103.0f, 104.0f, 0.0f, 0.0f, 0.0f, 0.0f, 0.0f, 0.0f, 0.0f,
+        96.0f, 97.0f, 98.0f, 99.0f, 100.0f, 101.0f, 102.0f, 103.0f, 104.0f, 0.0f, 0.0f, 0.0f, 0.0f, 0.0f, 0.0f, 0.0f,
+        96.0f, 97.0f, 98.0f, 99.0f, 100.0f, 101.0f, 102.0f, 103.0f, 104.0f, 0.0f, 0.0f, 0.0f, 0.0f, 0.0f, 0.0f, 0.0f,
+        96.0f, 97.0f, 98.0f, 99.0f, 100.0f, 101.0f, 102.0f, 103.0f, 104.0f, 0.0f, 0.0f, 0.0f, 0.0f, 0.0f, 0.0f, 0.0f
+    },
+    {  // Joint group 7 - No padding -> block-16 + block-16
+        105.0f, 106.0f, 107.0f, 108.0f, 109.0f, 110.0f, 111.0f, 112.0f, 113.0f, 114.0f, 115.0f, 116.0f, 117.0f, 118.0f, 119.0f,
+        120.0f,
+        105.0f, 106.0f, 107.0f, 108.0f, 109.0f, 110.0f, 111.0f, 112.0f, 113.0f, 114.0f, 115.0f, 116.0f, 117.0f, 118.0f, 119.0f,
+        120.0f,
+        105.0f, 106.0f, 107.0f, 108.0f, 109.0f, 110.0f, 111.0f, 112.0f, 113.0f, 114.0f, 115.0f, 116.0f, 117.0f, 118.0f, 119.0f,
+        120.0f,
+        105.0f, 106.0f, 107.0f, 108.0f, 109.0f, 110.0f, 111.0f, 112.0f, 113.0f, 114.0f, 115.0f, 116.0f, 117.0f, 118.0f, 119.0f,
+        120.0f,
+        105.0f, 106.0f, 107.0f, 108.0f, 109.0f, 110.0f, 111.0f, 112.0f, 113.0f, 114.0f, 115.0f, 116.0f, 117.0f, 118.0f, 119.0f,
+        120.0f,
+        105.0f, 106.0f, 107.0f, 108.0f, 109.0f, 110.0f, 111.0f, 112.0f, 113.0f, 114.0f, 115.0f, 116.0f, 117.0f, 118.0f, 119.0f,
+        120.0f,
+        105.0f, 106.0f, 107.0f, 108.0f, 109.0f, 110.0f, 111.0f, 112.0f, 113.0f, 114.0f, 115.0f, 116.0f, 117.0f, 118.0f, 119.0f,
+        120.0f,
+        105.0f, 106.0f, 107.0f, 108.0f, 109.0f, 110.0f, 111.0f, 112.0f, 113.0f, 114.0f, 115.0f, 116.0f, 117.0f, 118.0f, 119.0f,
+        120.0f,
+        105.0f, 106.0f, 107.0f, 108.0f, 109.0f, 110.0f, 111.0f, 112.0f, 113.0f, 114.0f, 115.0f, 116.0f, 117.0f, 118.0f, 119.0f,
+        120.0f,
+        105.0f, 106.0f, 107.0f, 108.0f, 109.0f, 110.0f, 111.0f, 112.0f, 113.0f, 114.0f, 115.0f, 116.0f, 117.0f, 118.0f, 119.0f,
+        120.0f,
+        105.0f, 106.0f, 107.0f, 108.0f, 109.0f, 110.0f, 111.0f, 112.0f, 113.0f, 114.0f, 115.0f, 116.0f, 117.0f, 118.0f, 119.0f,
+        120.0f,
+        105.0f, 106.0f, 107.0f, 108.0f, 109.0f, 110.0f, 111.0f, 112.0f, 113.0f, 114.0f, 115.0f, 116.0f, 117.0f, 118.0f, 119.0f,
+        120.0f,
+        105.0f, 106.0f, 107.0f, 108.0f, 109.0f, 110.0f, 111.0f, 112.0f, 113.0f, 114.0f, 115.0f, 116.0f, 117.0f, 118.0f, 119.0f,
+        120.0f,
+        121.0f, 122.0f, 123.0f, 124.0f, 125.0f, 126.0f, 127.0f, 128.0f, 129.0f, 130.0f, 131.0f, 132.0f, 133.0f, 134.0f, 135.0f,
+        136.0f,
+        121.0f, 122.0f, 123.0f, 124.0f, 125.0f, 126.0f, 127.0f, 128.0f, 129.0f, 130.0f, 131.0f, 132.0f, 133.0f, 134.0f, 135.0f,
+        136.0f,
+        121.0f, 122.0f, 123.0f, 124.0f, 125.0f, 126.0f, 127.0f, 128.0f, 129.0f, 130.0f, 131.0f, 132.0f, 133.0f, 134.0f, 135.0f,
+        136.0f,
+        121.0f, 122.0f, 123.0f, 124.0f, 125.0f, 126.0f, 127.0f, 128.0f, 129.0f, 130.0f, 131.0f, 132.0f, 133.0f, 134.0f, 135.0f,
+        136.0f,
+        121.0f, 122.0f, 123.0f, 124.0f, 125.0f, 126.0f, 127.0f, 128.0f, 129.0f, 130.0f, 131.0f, 132.0f, 133.0f, 134.0f, 135.0f,
+        136.0f,
+        121.0f, 122.0f, 123.0f, 124.0f, 125.0f, 126.0f, 127.0f, 128.0f, 129.0f, 130.0f, 131.0f, 132.0f, 133.0f, 134.0f, 135.0f,
+        136.0f,
+        121.0f, 122.0f, 123.0f, 124.0f, 125.0f, 126.0f, 127.0f, 128.0f, 129.0f, 130.0f, 131.0f, 132.0f, 133.0f, 134.0f, 135.0f,
+        136.0f,
+        121.0f, 122.0f, 123.0f, 124.0f, 125.0f, 126.0f, 127.0f, 128.0f, 129.0f, 130.0f, 131.0f, 132.0f, 133.0f, 134.0f, 135.0f,
+        136.0f,
+        121.0f, 122.0f, 123.0f, 124.0f, 125.0f, 126.0f, 127.0f, 128.0f, 129.0f, 130.0f, 131.0f, 132.0f, 133.0f, 134.0f, 135.0f,
+        136.0f,
+        121.0f, 122.0f, 123.0f, 124.0f, 125.0f, 126.0f, 127.0f, 128.0f, 129.0f, 130.0f, 131.0f, 132.0f, 133.0f, 134.0f, 135.0f,
+        136.0f,
+        121.0f, 122.0f, 123.0f, 124.0f, 125.0f, 126.0f, 127.0f, 128.0f, 129.0f, 130.0f, 131.0f, 132.0f, 133.0f, 134.0f, 135.0f,
+        136.0f,
+        121.0f, 122.0f, 123.0f, 124.0f, 125.0f, 126.0f, 127.0f, 128.0f, 129.0f, 130.0f, 131.0f, 132.0f, 133.0f, 134.0f, 135.0f,
+        136.0f,
+        121.0f, 122.0f, 123.0f, 124.0f, 125.0f, 126.0f, 127.0f, 128.0f, 129.0f, 130.0f, 131.0f, 132.0f, 133.0f, 134.0f, 135.0f,
+        136.0f
+    }
+};
+
+const AlignedMatrix<std::uint16_t> halfFloatValues = {
+    {  // Joint group 0 - Pad to block-8
+        15360, 16384, 16896, 17408, 17664, 0, 0, 0,
+        15360, 16384, 16896, 17408, 17664, 0, 0, 0,
+        15360, 16384, 16896, 17408, 17664, 0, 0, 0,
+        15360, 16384, 16896, 17408, 17664, 0, 0, 0,
+        15360, 16384, 16896, 17408, 17664, 0, 0, 0,
+        15360, 16384, 16896, 17408, 17664, 0, 0, 0,
+        15360, 16384, 16896, 17408, 17664, 0, 0, 0,
+        15360, 16384, 16896, 17408, 17664, 0, 0, 0,
+        15360, 16384, 16896, 17408, 17664, 0, 0, 0,
+        15360, 16384, 16896, 17408, 17664, 0, 0, 0,
+        15360, 16384, 16896, 17408, 17664, 0, 0, 0,
+        15360, 16384, 16896, 17408, 17664, 0, 0, 0,
+        15360, 16384, 16896, 17408, 17664, 0, 0, 0,
+    },
+    {  // Joint group 1 - No padding => block-8
+        17920, 18176, 18432, 18560, 18688, 18816, 18944, 19072,
+        17920, 18176, 18432, 18560, 18688, 18816, 18944, 19072,
+        17920, 18176, 18432, 18560, 18688, 18816, 18944, 19072,
+        17920, 18176, 18432, 18560, 18688, 18816, 18944, 19072,
+        17920, 18176, 18432, 18560, 18688, 18816, 18944, 19072,
+        17920, 18176, 18432, 18560, 18688, 18816, 18944, 19072,
+        17920, 18176, 18432, 18560, 18688, 18816, 18944, 19072,
+        17920, 18176, 18432, 18560, 18688, 18816, 18944, 19072,
+        17920, 18176, 18432, 18560, 18688, 18816, 18944, 19072,
+        17920, 18176, 18432, 18560, 18688, 18816, 18944, 19072,
+        17920, 18176, 18432, 18560, 18688, 18816, 18944, 19072,
+        17920, 18176, 18432, 18560, 18688, 18816, 18944, 19072,
+        17920, 18176, 18432, 18560, 18688, 18816, 18944, 19072,
+    },
+    {  // Joint group 2 - Pad to block-8 => block-16
+        19200, 19328, 19456, 19520, 19584, 19648, 19712, 19776, 19840, 0, 0, 0, 0, 0, 0, 0,
+        19200, 19328, 19456, 19520, 19584, 19648, 19712, 19776, 19840, 0, 0, 0, 0, 0, 0, 0,
+        19200, 19328, 19456, 19520, 19584, 19648, 19712, 19776, 19840, 0, 0, 0, 0, 0, 0, 0,
+        19200, 19328, 19456, 19520, 19584, 19648, 19712, 19776, 19840, 0, 0, 0, 0, 0, 0, 0,
+        19200, 19328, 19456, 19520, 19584, 19648, 19712, 19776, 19840, 0, 0, 0, 0, 0, 0, 0,
+        19200, 19328, 19456, 19520, 19584, 19648, 19712, 19776, 19840, 0, 0, 0, 0, 0, 0, 0,
+        19200, 19328, 19456, 19520, 19584, 19648, 19712, 19776, 19840, 0, 0, 0, 0, 0, 0, 0,
+        19200, 19328, 19456, 19520, 19584, 19648, 19712, 19776, 19840, 0, 0, 0, 0, 0, 0, 0,
+        19200, 19328, 19456, 19520, 19584, 19648, 19712, 19776, 19840, 0, 0, 0, 0, 0, 0, 0,
+        19200, 19328, 19456, 19520, 19584, 19648, 19712, 19776, 19840, 0, 0, 0, 0, 0, 0, 0,
+        19200, 19328, 19456, 19520, 19584, 19648, 19712, 19776, 19840, 0, 0, 0, 0, 0, 0, 0,
+        19200, 19328, 19456, 19520, 19584, 19648, 19712, 19776, 19840, 0, 0, 0, 0, 0, 0, 0,
+        19200, 19328, 19456, 19520, 19584, 19648, 19712, 19776, 19840, 0, 0, 0, 0, 0, 0, 0,
+    },
+    {  // Joint group 3 - No padding => block-16
+        19904, 19968, 20032, 20096, 20160, 20224, 20288, 20352, 20416, 20480, 20512, 20544, 20576, 20608, 20640, 20672,
+        19904, 19968, 20032, 20096, 20160, 20224, 20288, 20352, 20416, 20480, 20512, 20544, 20576, 20608, 20640, 20672,
+        19904, 19968, 20032, 20096, 20160, 20224, 20288, 20352, 20416, 20480, 20512, 20544, 20576, 20608, 20640, 20672,
+        19904, 19968, 20032, 20096, 20160, 20224, 20288, 20352, 20416, 20480, 20512, 20544, 20576, 20608, 20640, 20672,
+        19904, 19968, 20032, 20096, 20160, 20224, 20288, 20352, 20416, 20480, 20512, 20544, 20576, 20608, 20640, 20672,
+        19904, 19968, 20032, 20096, 20160, 20224, 20288, 20352, 20416, 20480, 20512, 20544, 20576, 20608, 20640, 20672,
+        19904, 19968, 20032, 20096, 20160, 20224, 20288, 20352, 20416, 20480, 20512, 20544, 20576, 20608, 20640, 20672,
+        19904, 19968, 20032, 20096, 20160, 20224, 20288, 20352, 20416, 20480, 20512, 20544, 20576, 20608, 20640, 20672,
+        19904, 19968, 20032, 20096, 20160, 20224, 20288, 20352, 20416, 20480, 20512, 20544, 20576, 20608, 20640, 20672,
+        19904, 19968, 20032, 20096, 20160, 20224, 20288, 20352, 20416, 20480, 20512, 20544, 20576, 20608, 20640, 20672,
+        19904, 19968, 20032, 20096, 20160, 20224, 20288, 20352, 20416, 20480, 20512, 20544, 20576, 20608, 20640, 20672,
+        19904, 19968, 20032, 20096, 20160, 20224, 20288, 20352, 20416, 20480, 20512, 20544, 20576, 20608, 20640, 20672,
+        19904, 19968, 20032, 20096, 20160, 20224, 20288, 20352, 20416, 20480, 20512, 20544, 20576, 20608, 20640, 20672,
+    },
+    {  // Joint group 4 - Pad to block-8 => block-16 + block-8
+        20704, 20736, 20768, 20800, 20832, 20864, 20896, 20928, 20960, 20992, 21024, 21056, 21088, 21120, 21152, 21184,
+        20704, 20736, 20768, 20800, 20832, 20864, 20896, 20928, 20960, 20992, 21024, 21056, 21088, 21120, 21152, 21184,
+        20704, 20736, 20768, 20800, 20832, 20864, 20896, 20928, 20960, 20992, 21024, 21056, 21088, 21120, 21152, 21184,
+        20704, 20736, 20768, 20800, 20832, 20864, 20896, 20928, 20960, 20992, 21024, 21056, 21088, 21120, 21152, 21184,
+        20704, 20736, 20768, 20800, 20832, 20864, 20896, 20928, 20960, 20992, 21024, 21056, 21088, 21120, 21152, 21184,
+        20704, 20736, 20768, 20800, 20832, 20864, 20896, 20928, 20960, 20992, 21024, 21056, 21088, 21120, 21152, 21184,
+        20704, 20736, 20768, 20800, 20832, 20864, 20896, 20928, 20960, 20992, 21024, 21056, 21088, 21120, 21152, 21184,
+        20704, 20736, 20768, 20800, 20832, 20864, 20896, 20928, 20960, 20992, 21024, 21056, 21088, 21120, 21152, 21184,
+        20704, 20736, 20768, 20800, 20832, 20864, 20896, 20928, 20960, 20992, 21024, 21056, 21088, 21120, 21152, 21184,
+        20704, 20736, 20768, 20800, 20832, 20864, 20896, 20928, 20960, 20992, 21024, 21056, 21088, 21120, 21152, 21184,
+        20704, 20736, 20768, 20800, 20832, 20864, 20896, 20928, 20960, 20992, 21024, 21056, 21088, 21120, 21152, 21184,
+        20704, 20736, 20768, 20800, 20832, 20864, 20896, 20928, 20960, 20992, 21024, 21056, 21088, 21120, 21152, 21184,
+        20704, 20736, 20768, 20800, 20832, 20864, 20896, 20928, 20960, 20992, 21024, 21056, 21088, 21120, 21152, 21184,
+        21216, 0, 0, 0, 0, 0, 0, 0,
+        21216, 0, 0, 0, 0, 0, 0, 0,
+        21216, 0, 0, 0, 0, 0, 0, 0,
+        21216, 0, 0, 0, 0, 0, 0, 0,
+        21216, 0, 0, 0, 0, 0, 0, 0,
+        21216, 0, 0, 0, 0, 0, 0, 0,
+        21216, 0, 0, 0, 0, 0, 0, 0,
+        21216, 0, 0, 0, 0, 0, 0, 0,
+        21216, 0, 0, 0, 0, 0, 0, 0,
+        21216, 0, 0, 0, 0, 0, 0, 0,
+        21216, 0, 0, 0, 0, 0, 0, 0,
+        21216, 0, 0, 0, 0, 0, 0, 0,
+        21216, 0, 0, 0, 0, 0, 0, 0,
+    },
+    {  // Joint group 5 - No padding -> block-16 + block-8
+        21248, 21280, 21312, 21344, 21376, 21408, 21440, 21472, 21504, 21520, 21536, 21552, 21568, 21584, 21600, 21616,
+        21248, 21280, 21312, 21344, 21376, 21408, 21440, 21472, 21504, 21520, 21536, 21552, 21568, 21584, 21600, 21616,
+        21248, 21280, 21312, 21344, 21376, 21408, 21440, 21472, 21504, 21520, 21536, 21552, 21568, 21584, 21600, 21616,
+        21248, 21280, 21312, 21344, 21376, 21408, 21440, 21472, 21504, 21520, 21536, 21552, 21568, 21584, 21600, 21616,
+        21248, 21280, 21312, 21344, 21376, 21408, 21440, 21472, 21504, 21520, 21536, 21552, 21568, 21584, 21600, 21616,
+        21248, 21280, 21312, 21344, 21376, 21408, 21440, 21472, 21504, 21520, 21536, 21552, 21568, 21584, 21600, 21616,
+        21248, 21280, 21312, 21344, 21376, 21408, 21440, 21472, 21504, 21520, 21536, 21552, 21568, 21584, 21600, 21616,
+        21248, 21280, 21312, 21344, 21376, 21408, 21440, 21472, 21504, 21520, 21536, 21552, 21568, 21584, 21600, 21616,
+        21248, 21280, 21312, 21344, 21376, 21408, 21440, 21472, 21504, 21520, 21536, 21552, 21568, 21584, 21600, 21616,
+        21248, 21280, 21312, 21344, 21376, 21408, 21440, 21472, 21504, 21520, 21536, 21552, 21568, 21584, 21600, 21616,
+        21248, 21280, 21312, 21344, 21376, 21408, 21440, 21472, 21504, 21520, 21536, 21552, 21568, 21584, 21600, 21616,
+        21248, 21280, 21312, 21344, 21376, 21408, 21440, 21472, 21504, 21520, 21536, 21552, 21568, 21584, 21600, 21616,
+        21248, 21280, 21312, 21344, 21376, 21408, 21440, 21472, 21504, 21520, 21536, 21552, 21568, 21584, 21600, 21616,
+        21632, 21648, 21664, 21680, 21696, 21712, 21728, 21744,
+        21632, 21648, 21664, 21680, 21696, 21712, 21728, 21744,
+        21632, 21648, 21664, 21680, 21696, 21712, 21728, 21744,
+        21632, 21648, 21664, 21680, 21696, 21712, 21728, 21744,
+        21632, 21648, 21664, 21680, 21696, 21712, 21728, 21744,
+        21632, 21648, 21664, 21680, 21696, 21712, 21728, 21744,
+        21632, 21648, 21664, 21680, 21696, 21712, 21728, 21744,
+        21632, 21648, 21664, 21680, 21696, 21712, 21728, 21744,
+        21632, 21648, 21664, 21680, 21696, 21712, 21728, 21744,
+        21632, 21648, 21664, 21680, 21696, 21712, 21728, 21744,
+        21632, 21648, 21664, 21680, 21696, 21712, 21728, 21744,
+        21632, 21648, 21664, 21680, 21696, 21712, 21728, 21744,
+        21632, 21648, 21664, 21680, 21696, 21712, 21728, 21744,
+    },
+    {  // Joint group 6 - Pad to block-16 -> block-16 + block-16
+        21760, 21776, 21792, 21808, 21824, 21840, 21856, 21872, 21888, 21904, 21920, 21936, 21952, 21968, 21984, 22000,
+        21760, 21776, 21792, 21808, 21824, 21840, 21856, 21872, 21888, 21904, 21920, 21936, 21952, 21968, 21984, 22000,
+        21760, 21776, 21792, 21808, 21824, 21840, 21856, 21872, 21888, 21904, 21920, 21936, 21952, 21968, 21984, 22000,
+        21760, 21776, 21792, 21808, 21824, 21840, 21856, 21872, 21888, 21904, 21920, 21936, 21952, 21968, 21984, 22000,
+        21760, 21776, 21792, 21808, 21824, 21840, 21856, 21872, 21888, 21904, 21920, 21936, 21952, 21968, 21984, 22000,
+        21760, 21776, 21792, 21808, 21824, 21840, 21856, 21872, 21888, 21904, 21920, 21936, 21952, 21968, 21984, 22000,
+        21760, 21776, 21792, 21808, 21824, 21840, 21856, 21872, 21888, 21904, 21920, 21936, 21952, 21968, 21984, 22000,
+        21760, 21776, 21792, 21808, 21824, 21840, 21856, 21872, 21888, 21904, 21920, 21936, 21952, 21968, 21984, 22000,
+        21760, 21776, 21792, 21808, 21824, 21840, 21856, 21872, 21888, 21904, 21920, 21936, 21952, 21968, 21984, 22000,
+        21760, 21776, 21792, 21808, 21824, 21840, 21856, 21872, 21888, 21904, 21920, 21936, 21952, 21968, 21984, 22000,
+        21760, 21776, 21792, 21808, 21824, 21840, 21856, 21872, 21888, 21904, 21920, 21936, 21952, 21968, 21984, 22000,
+        21760, 21776, 21792, 21808, 21824, 21840, 21856, 21872, 21888, 21904, 21920, 21936, 21952, 21968, 21984, 22000,
+        21760, 21776, 21792, 21808, 21824, 21840, 21856, 21872, 21888, 21904, 21920, 21936, 21952, 21968, 21984, 22000,
+        22016, 22032, 22048, 22064, 22080, 22096, 22112, 22128, 22144, 0, 0, 0, 0, 0, 0, 0,
+        22016, 22032, 22048, 22064, 22080, 22096, 22112, 22128, 22144, 0, 0, 0, 0, 0, 0, 0,
+        22016, 22032, 22048, 22064, 22080, 22096, 22112, 22128, 22144, 0, 0, 0, 0, 0, 0, 0,
+        22016, 22032, 22048, 22064, 22080, 22096, 22112, 22128, 22144, 0, 0, 0, 0, 0, 0, 0,
+        22016, 22032, 22048, 22064, 22080, 22096, 22112, 22128, 22144, 0, 0, 0, 0, 0, 0, 0,
+        22016, 22032, 22048, 22064, 22080, 22096, 22112, 22128, 22144, 0, 0, 0, 0, 0, 0, 0,
+        22016, 22032, 22048, 22064, 22080, 22096, 22112, 22128, 22144, 0, 0, 0, 0, 0, 0, 0,
+        22016, 22032, 22048, 22064, 22080, 22096, 22112, 22128, 22144, 0, 0, 0, 0, 0, 0, 0,
+        22016, 22032, 22048, 22064, 22080, 22096, 22112, 22128, 22144, 0, 0, 0, 0, 0, 0, 0,
+        22016, 22032, 22048, 22064, 22080, 22096, 22112, 22128, 22144, 0, 0, 0, 0, 0, 0, 0,
+        22016, 22032, 22048, 22064, 22080, 22096, 22112, 22128, 22144, 0, 0, 0, 0, 0, 0, 0,
+        22016, 22032, 22048, 22064, 22080, 22096, 22112, 22128, 22144, 0, 0, 0, 0, 0, 0, 0,
+        22016, 22032, 22048, 22064, 22080, 22096, 22112, 22128, 22144, 0, 0, 0, 0, 0, 0, 0,
+    },
+    {  // Joint group 7 - No padding -> block-16 + block-16
+        22160, 22176, 22192, 22208, 22224, 22240, 22256, 22272, 22288, 22304, 22320, 22336, 22352, 22368, 22384, 22400,
+        22160, 22176, 22192, 22208, 22224, 22240, 22256, 22272, 22288, 22304, 22320, 22336, 22352, 22368, 22384, 22400,
+        22160, 22176, 22192, 22208, 22224, 22240, 22256, 22272, 22288, 22304, 22320, 22336, 22352, 22368, 22384, 22400,
+        22160, 22176, 22192, 22208, 22224, 22240, 22256, 22272, 22288, 22304, 22320, 22336, 22352, 22368, 22384, 22400,
+        22160, 22176, 22192, 22208, 22224, 22240, 22256, 22272, 22288, 22304, 22320, 22336, 22352, 22368, 22384, 22400,
+        22160, 22176, 22192, 22208, 22224, 22240, 22256, 22272, 22288, 22304, 22320, 22336, 22352, 22368, 22384, 22400,
+        22160, 22176, 22192, 22208, 22224, 22240, 22256, 22272, 22288, 22304, 22320, 22336, 22352, 22368, 22384, 22400,
+        22160, 22176, 22192, 22208, 22224, 22240, 22256, 22272, 22288, 22304, 22320, 22336, 22352, 22368, 22384, 22400,
+        22160, 22176, 22192, 22208, 22224, 22240, 22256, 22272, 22288, 22304, 22320, 22336, 22352, 22368, 22384, 22400,
+        22160, 22176, 22192, 22208, 22224, 22240, 22256, 22272, 22288, 22304, 22320, 22336, 22352, 22368, 22384, 22400,
+        22160, 22176, 22192, 22208, 22224, 22240, 22256, 22272, 22288, 22304, 22320, 22336, 22352, 22368, 22384, 22400,
+        22160, 22176, 22192, 22208, 22224, 22240, 22256, 22272, 22288, 22304, 22320, 22336, 22352, 22368, 22384, 22400,
+        22160, 22176, 22192, 22208, 22224, 22240, 22256, 22272, 22288, 22304, 22320, 22336, 22352, 22368, 22384, 22400,
+        22416, 22432, 22448, 22464, 22480, 22496, 22512, 22528, 22536, 22544, 22552, 22560, 22568, 22576, 22584, 22592,
+        22416, 22432, 22448, 22464, 22480, 22496, 22512, 22528, 22536, 22544, 22552, 22560, 22568, 22576, 22584, 22592,
+        22416, 22432, 22448, 22464, 22480, 22496, 22512, 22528, 22536, 22544, 22552, 22560, 22568, 22576, 22584, 22592,
+        22416, 22432, 22448, 22464, 22480, 22496, 22512, 22528, 22536, 22544, 22552, 22560, 22568, 22576, 22584, 22592,
+        22416, 22432, 22448, 22464, 22480, 22496, 22512, 22528, 22536, 22544, 22552, 22560, 22568, 22576, 22584, 22592,
+        22416, 22432, 22448, 22464, 22480, 22496, 22512, 22528, 22536, 22544, 22552, 22560, 22568, 22576, 22584, 22592,
+        22416, 22432, 22448, 22464, 22480, 22496, 22512, 22528, 22536, 22544, 22552, 22560, 22568, 22576, 22584, 22592,
+        22416, 22432, 22448, 22464, 22480, 22496, 22512, 22528, 22536, 22544, 22552, 22560, 22568, 22576, 22584, 22592,
+        22416, 22432, 22448, 22464, 22480, 22496, 22512, 22528, 22536, 22544, 22552, 22560, 22568, 22576, 22584, 22592,
+        22416, 22432, 22448, 22464, 22480, 22496, 22512, 22528, 22536, 22544, 22552, 22560, 22568, 22576, 22584, 22592,
+        22416, 22432, 22448, 22464, 22480, 22496, 22512, 22528, 22536, 22544, 22552, 22560, 22568, 22576, 22584, 22592,
+        22416, 22432, 22448, 22464, 22480, 22496, 22512, 22528, 22536, 22544, 22552, 22560, 22568, 22576, 22584, 22592,
+        22416, 22432, 22448, 22464, 22480, 22496, 22512, 22528, 22536, 22544, 22552, 22560, 22568, 22576, 22584, 22592,
+    }
+};
+
+const AlignedMatrix<std::uint16_t> inputIndices = {
+    {  // Joint group 0
+        12, 11, 10, 9, 8, 7, 6, 5, 4, 3, 2, 1, 0,
+    },
+    {  // Joint group 1
+        12, 11, 10, 9, 8, 7, 6, 5, 4, 3, 2, 1, 0,
+    },
+    {  // Joint group 2
+        12, 11, 10, 9, 8, 7, 6, 5, 4, 3, 2, 1, 0,
+    },
+    {  // Joint group 3
+        12, 11, 10, 9, 8, 7, 6, 5, 4, 3, 2, 1, 0,
+    },
+    {  // Joint group 4
+        12, 11, 10, 9, 8, 7, 6, 5, 4, 3, 2, 1, 0,
+    },
+    {  // Joint group 5
+        12, 11, 10, 9, 8, 7, 6, 5, 4, 3, 2, 1, 0,
+    },
+    {  // Joint group 6
+        12, 11, 10, 9, 8, 7, 6, 5, 4, 3, 2, 1, 0,
+    },
+    {  // Joint group 7
+        12, 11, 10, 9, 8, 7, 6, 5, 4, 3, 2, 1, 0,
+    }
+};
+
+const AlignedMatrix<std::uint16_t> outputIndices = {
+    {  // Joint group 0
+        0, 1, 2, 3, 4, 0, 0, 0
+    },
+    {  // Joint group 1
+        5, 6, 7, 8, 9, 10, 11, 12
+    },
+    {  // Joint group 2
+        13, 14, 15, 16, 17, 18, 19, 20, 21, 0, 0, 0, 0, 0, 0, 0
+    },
+    {  // Joint group 3
+        22, 23, 24, 25, 26, 27, 28, 29, 30, 31, 32, 33, 34, 35, 36, 37
+    },
+    {  // Joint group 4
+        38, 39, 40, 41, 42, 43, 44, 45, 46, 47, 48, 49, 50, 51, 52, 53, 54, 0, 0, 0, 0, 0, 0, 0
+    },
+    {  // Joint group 5
+        55, 56, 57, 58, 59, 60, 61, 62, 63, 64, 65, 66, 67, 68, 69, 70, 71, 72, 73, 74, 75, 76, 77, 78
+    },
+    {  // Joint group 6
+        79, 80, 81, 82, 83, 84, 85, 86, 87, 88, 89, 90, 91, 92, 93, 94,
+        95, 96, 97, 98, 99, 100, 101, 102, 103, 0, 0, 0, 0, 0, 0, 0
+    },
+    {  // Joint group 7
+        104, 105, 106, 107, 108, 109, 110, 111, 112, 113, 114, 115, 116, 117, 118, 119,
+        120, 121, 122, 123, 124, 125, 126, 127, 128, 129, 130, 131, 132, 133, 134, 135
+    }
+};
+
+const Vector<bpcm::JointGroup> jointGroups = {
+    // {valueOffset, inputOffset, outputOffset, lodOffset, valueSize, inputSize, inputSizeAligned4, inputSizeAligned8}
+    {  // Joint group 0
+        0, 0, 0, 0, 104, 13, 12, 8
+    },
+    {  // Joint group 1
+        104, 13, 8, 4, 104, 13, 12, 8
+    },
+    {  // Joint group 2
+        208, 26, 16, 8, 208, 13, 12, 8
+    },
+    {  // Joint group 3
+        416, 39, 32, 12, 208, 13, 12, 8
+    },
+    {  // Joint group 4
+        624, 52, 48, 16, 312, 13, 12, 8
+    },
+    {  // Joint group 5
+        936, 65, 72, 20, 312, 13, 12, 8
+    },
+    {  // Joint group 6
+        1248, 78, 96, 24, 416, 13, 12, 8
+    },
+    {  // Joint group 7
+        1664, 91, 128, 28, 416, 13, 12, 8
+    }
+};
+
+const Matrix<bpcm::LODRegion> lodRegions = {
+    // {size, sizeAlignedToLastFullBlock, sizeAlignedToSecondLastFullBlock}
+    {  // Joint group 0
+        {5, 0, 0},  // LOD-0
+        {4, 0, 0},  // LOD-1
+        {2, 0, 0},  // LOD-2
+        {1, 0, 0},  // LOD-3
+    },
+    {  // Joint group 1
+        {8, 0, 0},  // LOD-0
+        {5, 0, 0},  // LOD-1
+        {3, 0, 0},  // LOD-2
+        {1, 0, 0},  // LOD-3
+    },
+    {  // Joint group 2
+        {9, 16, 0},  // LOD-0
+        {7, 16, 0},  // LOD-1
+        {4, 16, 0},  // LOD-2
+        {2, 16, 0},  // LOD-3
+    },
+    {  // Joint group 3
+        {16, 16, 16},  // LOD-0
+        {13, 16, 0},  // LOD-1
+        {8, 16, 0},  // LOD-2
+        {4, 16, 0},  // LOD-3
+    },
+    {  // Joint group 4
+        {17, 16, 16},  // LOD-0
+        {16, 16, 16},  // LOD-1
+        {15, 16, 0},  // LOD-2
+        {12, 16, 0},  // LOD-3
+    },
+    {  // Joint group 5
+        {24, 16, 16},  // LOD-0
+        {16, 16, 16},  // LOD-1
+        {15, 16, 0},  // LOD-2
+        {8, 16, 0},  // LOD-3
+    },
+    {  // Joint group 6
+        {25, 32, 16},  // LOD-0
+        {17, 32, 16},  // LOD-1
+        {12, 16, 0},  // LOD-2
+        {7, 16, 0},  // LOD-3
+    },
+    {  // Joint group 7
+        {32, 32, 32},  // LOD-0
+        {25, 32, 16},  // LOD-1
+        {17, 32, 16},  // LOD-2
+        {9, 16, 0},  // LOD-3
+    }
+};
+
+}  // namespace optimized
+
+namespace input {
+
+// Calculation input values
+const rl4::Vector<float> values = {1.0f, 2.0f, 3.0f, 4.0f, 0.0f, 6.0f, 7.0f, 8.0f, 9.0f, 0.0f, 11.0f, 12.0f, 13.0f};
+
+}  // namespace input
+
+namespace output {
+
+// Expected output results for each LOD
+const rl4::Matrix<float> valuesPerLOD = {
+    {
+        // LOD-0
+        76.0f, 152.0f, 228.0f, 304.0f, 380.0f,  // Joint group 0
+        456.0f, 532.0f, 608.0f, 684.0f, 760.0f, 836.0f, 912.0f, 988.0f,  // Joint group 1
+        1064.0f, 1140.0f, 1216.0f, 1292.0f, 1368.0f, 1444.0f, 1520.0f, 1596.0f, 1672.0f,  // Joint group 2
+        1748.0f, 1824.0f, 1900.0f, 1976.0f, 2052.0f, 2128.0f, 2204.0f, 2280.0f,  // Joint group 3
+        2356.0f, 2432.0f, 2508.0f, 2584.0f, 2660.0f, 2736.0f, 2812.0f, 2888.0f,  // Joint group 3
+        2964.0f, 3040.0f, 3116.0f, 3192.0f, 3268.0f, 3344.0f, 3420.0f, 3496.0f,  // Joint group 4
+        3572.0f, 3648.0f, 3724.0f, 3800.0f, 3876.0f, 3952.0f, 4028.0f, 4104.0f, 4180.0f,  // Joint group 4
+        4256.0f, 4332.0f, 4408.0f, 4484.0f, 4560.0f, 4636.0f, 4712.0f, 4788.0f,  // Joint group 5
+        4864.0f, 4940.0f, 5016.0f, 5092.0f, 5168.0f, 5244.0f, 5320.0f, 5396.0f,  // Joint group 5
+        5472.0f, 5548.0f, 5624.0f, 5700.0f, 5776.0f, 5852.0f, 5928.0f, 6004.0f,  // Joint group 5
+        6080.0f, 6156.0f, 6232.0f, 6308.0f, 6384.0f, 6460.0f, 6536.0f, 6612.0f,  // Joint group 6
+        6688.0f, 6764.0f, 6840.0f, 6916.0f, 6992.0f, 7068.0f, 7144.0f, 7220.0f,  // Joint group 6
+        7296.0f, 7372.0f, 7448.0f, 7524.0f, 7600.0f, 7676.0f, 7752.0f, 7828.0f, 7904.0f,  // Joint group 6
+        7980.0f, 8056.0f, 8132.0f, 8208.0f, 8284.0f, 8360.0f, 8436.0f, 8512.0f,  // Joint group 7
+        8588.0f, 8664.0f, 8740.0f, 8816.0f, 8892.0f, 8968.0f, 9044.0f, 9120.0f,  // Joint group 7
+        9196.0f, 9272.0f, 9348.0f, 9424.0f, 9500.0f, 9576.0f, 9652.0f, 9728.0f,  // Joint group 7
+        9804.0f, 9880.0f, 9956.0f, 10032.0f, 10108.0f, 10184.0f, 10260.0f, 10336.0f  // Joint group 7
+    }, {
+        // LOD-1
+        76.0f, 152.0f, 228.0f, 304.0f, 0.0f,  // Joint group 0
+        456.0f, 532.0f, 608.0f, 684.0f, 760.0f, 0.0f, 0.0f, 0.0f,  // Joint group 1
+        1064.0f, 1140.0f, 1216.0f, 1292.0f, 1368.0f, 1444.0f, 1520.0f, 0.0f, 0.0f,  // Joint group 2
+        1748.0f, 1824.0f, 1900.0f, 1976.0f, 2052.0f, 2128.0f, 2204.0f, 2280.0f,  // Joint group 3
+        2356.0f, 2432.0f, 2508.0f, 2584.0f, 2660.0f, 0.0f, 0.0f, 0.0f,  // Joint group 3
+        2964.0f, 3040.0f, 3116.0f, 3192.0f, 3268.0f, 3344.0f, 3420.0f, 3496.0f,  // Joint group 4
+        3572.0f, 3648.0f, 3724.0f, 3800.0f, 3876.0f, 3952.0f, 4028.0f, 4104.0f, 0.0f,  // Joint group 4
+        4256.0f, 4332.0f, 4408.0f, 4484.0f, 4560.0f, 4636.0f, 4712.0f, 4788.0f,  // Joint group 5
+        4864.0f, 4940.0f, 5016.0f, 5092.0f, 5168.0f, 5244.0f, 5320.0f, 5396.0f,  // Joint group 5
+        0.0f, 0.0f, 0.0f, 0.0f, 0.0f, 0.0f, 0.0f, 0.0f,  // Joint group 5
+        6080.0f, 6156.0f, 6232.0f, 6308.0f, 6384.0f, 6460.0f, 6536.0f, 6612.0f,  // Joint group 6
+        6688.0f, 6764.0f, 6840.0f, 6916.0f, 6992.0f, 7068.0f, 7144.0f, 7220.0f,  // Joint group 6
+        7296.0f, 0.0f, 0.0f, 0.0f, 0.0f, 0.0f, 0.0f, 0.0f, 0.0f,  // Joint group 6
+        7980.0f, 8056.0f, 8132.0f, 8208.0f, 8284.0f, 8360.0f, 8436.0f, 8512.0f,  // Joint group 7
+        8588.0f, 8664.0f, 8740.0f, 8816.0f, 8892.0f, 8968.0f, 9044.0f, 9120.0f,  // Joint group 7
+        9196.0f, 9272.0f, 9348.0f, 9424.0f, 9500.0f, 9576.0f, 9652.0f, 9728.0f,  // Joint group 7
+        9804.0f, 0.0f, 0.0f, 0.0f, 0.0f, 0.0f, 0.0f, 0.0f  // Joint group 7
+    }, {
+        // LOD-2
+        76.0f, 152.0f, 0.0f, 0.0f, 0.0f,  // Joint group 0
+        456.0f, 532.0f, 608.0f, 0.0f, 0.0f, 0.0f, 0.0f, 0.0f,  // Joint group 1
+        1064.0f, 1140.0f, 1216.0f, 1292.0f, 0.0f, 0.0f, 0.0f, 0.0f, 0.0f,  // Joint group 2
+        1748.0f, 1824.0f, 1900.0f, 1976.0f, 2052.0f, 2128.0f, 2204.0f, 2280.0f,  // Joint group 3
+        0.0f, 0.0f, 0.0f, 0.0f, 0.0f, 0.0f, 0.0f, 0.0f,  // Joint group 3
+        2964.0f, 3040.0f, 3116.0f, 3192.0f, 3268.0f, 3344.0f, 3420.0f, 3496.0f,  // Joint group 4
+        3572.0f, 3648.0f, 3724.0f, 3800.0f, 3876.0f, 3952.0f, 4028.0f, 0.0f, 0.0f,  // Joint group 4
+        4256.0f, 4332.0f, 4408.0f, 4484.0f, 4560.0f, 4636.0f, 4712.0f, 4788.0f,  // Joint group 5
+        4864.0f, 4940.0f, 5016.0f, 5092.0f, 5168.0f, 5244.0f, 5320.0f, 0.0f,  // Joint group 5
+        0.0f, 0.0f, 0.0f, 0.0f, 0.0f, 0.0f, 0.0f, 0.0f,  // Joint group 5
+        6080.0f, 6156.0f, 6232.0f, 6308.0f, 6384.0f, 6460.0f, 6536.0f, 6612.0f,  // Joint group 6
+        6688.0f, 6764.0f, 6840.0f, 6916.0f, 0.0f, 0.0f, 0.0f, 0.0f,  // Joint group 6
+        0.0f, 0.0f, 0.0f, 0.0f, 0.0f, 0.0f, 0.0f, 0.0f, 0.0f,  // Joint group 6
+        7980.0f, 8056.0f, 8132.0f, 8208.0f, 8284.0f, 8360.0f, 8436.0f, 8512.0f,  // Joint group 7
+        8588.0f, 8664.0f, 8740.0f, 8816.0f, 8892.0f, 8968.0f, 9044.0f, 9120.0f,  // Joint group 7
+        9196.0f, 0.0f, 0.0f, 0.0f, 0.0f, 0.0f, 0.0f, 0.0f,  // Joint group 7
+        0.0f, 0.0f, 0.0f, 0.0f, 0.0f, 0.0f, 0.0f, 0.0f  // Joint group 7
+    }, {
+        // LOD-3
+        76.0f, 0.0f, 0.0f, 0.0f, 0.0f,  // Joint group 0
+        456.0f, 0.0f, 0.0f, 0.0f, 0.0f, 0.0f, 0.0f, 0.0f,  // Joint group 1
+        1064.0f, 1140.0f, 0.0f, 0.0f, 0.0f, 0.0f, 0.0f, 0.0f, 0.0f,  // Joint group 2
+        1748.0f, 1824.0f, 1900.0f, 1976.0f, 0.0f, 0.0f, 0.0f, 0.0f,  // Joint group 3
+        0.0f, 0.0f, 0.0f, 0.0f, 0.0f, 0.0f, 0.0f, 0.0f,  // Joint group 3
+        2964.0f, 3040.0f, 3116.0f, 3192.0f, 3268.0f, 3344.0f, 3420.0f, 3496.0f,  // Joint group 4
+        3572.0f, 3648.0f, 3724.0f, 3800.0f, 0.0f, 0.0f, 0.0f, 0.0f, 0.0f,  // Joint group 4
+        4256.0f, 4332.0f, 4408.0f, 4484.0f, 4560.0f, 4636.0f, 4712.0f, 4788.0f,  // Joint group 5
+        0.0f, 0.0f, 0.0f, 0.0f, 0.0f, 0.0f, 0.0f, 0.0f,  // Joint group 5
+        0.0f, 0.0f, 0.0f, 0.0f, 0.0f, 0.0f, 0.0f, 0.0f,  // Joint group 5
+        6080.0f, 6156.0f, 6232.0f, 6308.0f, 6384.0f, 6460.0f, 6536.0f, 0.0f,  // Joint group 6
+        0.0f, 0.0f, 0.0f, 0.0f, 0.0f, 0.0f, 0.0f, 0.0f,  // Joint group 6
+        0.0f, 0.0f, 0.0f, 0.0f, 0.0f, 0.0f, 0.0f, 0.0f, 0.0f,  // Joint group 6
+        7980.0f, 8056.0f, 8132.0f, 8208.0f, 8284.0f, 8360.0f, 8436.0f, 8512.0f,  // Joint group 7
+        8588.0f, 0.0f, 0.0f, 0.0f, 0.0f, 0.0f, 0.0f, 0.0f,  // Joint group 7
+        0.0f, 0.0f, 0.0f, 0.0f, 0.0f, 0.0f, 0.0f, 0.0f,  // Joint group 7
+        0.0f, 0.0f, 0.0f, 0.0f, 0.0f, 0.0f, 0.0f, 0.0f  // Joint group 7
+    }
+};
+
+}  // namespace output
+
+template<typename TValue>
+struct OptimizedValues;
+
+template<>
+struct OptimizedValues<std::uint16_t> {
+    static const rl4::AlignedMatrix<std::uint16_t>& get() {
+        return optimized::halfFloatValues;
+    }
+
+};
+
+template<>
+struct OptimizedValues<float> {
+    static const rl4::AlignedMatrix<float>& get() {
+        return optimized::floatValues;
+    }
+
+};
+
+template<typename TValue>
+bpcm::Evaluator<TValue> OptimizedStorage<TValue>::create(StrategyPtr strategy, rl4::MemoryResource* memRes) {
+    bpcm::JointStorage<TValue> storage{memRes};
+    const auto& values = OptimizedValues<TValue>::get();
+    for (std::uint16_t i = 0u; i < static_cast<std::uint16_t>(unoptimized::values.size()); ++i) {
+        storage.values.insert(storage.values.end(), values[i].begin(), values[i].end());
+        storage.inputIndices.insert(storage.inputIndices.end(),
+                                    optimized::inputIndices[i].begin(),
+                                    optimized::inputIndices[i].end());
+        storage.outputIndices.insert(storage.outputIndices.end(),
+                                     optimized::outputIndices[i].begin(),
+                                     optimized::outputIndices[i].end());
+        storage.lodRegions.insert(storage.lodRegions.end(),
+                                  optimized::lodRegions[i].begin(),
+                                  optimized::lodRegions[i].end());
+        storage.jointGroups.push_back(optimized::jointGroups[i]);
+    }
+    return bpcm::Evaluator<TValue>{std::move(storage), std::move(strategy), memRes};
+}
+
+template<typename TValue>
+bpcm::Evaluator<TValue> OptimizedStorage<TValue>::create(StrategyPtr strategy,
+                                                         std::uint16_t jointGroupIndex,
+                                                         rl4::MemoryResource* memRes) {
+    bpcm::JointStorage<TValue> storage{memRes};
+    const auto& values = OptimizedValues<TValue>::get();
+    storage.values.assign(values[jointGroupIndex].begin(), values[jointGroupIndex].end());
+    storage.inputIndices.assign(optimized::inputIndices[jointGroupIndex].begin(),
+                                optimized::inputIndices[jointGroupIndex].end());
+    storage.outputIndices.assign(optimized::outputIndices[jointGroupIndex].begin(),
+                                 optimized::outputIndices[jointGroupIndex].end());
+    storage.lodRegions.assign(optimized::lodRegions[jointGroupIndex].begin(),
+                              optimized::lodRegions[jointGroupIndex].end());
+    storage.jointGroups.push_back(bpcm::JointGroup{
+            0u,  // values offset
+            0u,  // inputIndicesOffset
+            0u,  // outputIndicesOffset
+            0u,  // lodsOffset
+            static_cast<std::uint32_t>(storage.values.size()),
+            static_cast<std::uint32_t>(storage.inputIndices.size()),
+            static_cast<std::uint32_t>(storage.inputIndices.size() - (storage.inputIndices.size() % 4ul)),
+            static_cast<std::uint32_t>(storage.inputIndices.size() - (storage.inputIndices.size() % 8ul)),
+        });
+    return bpcm::Evaluator<TValue>{std::move(storage), std::move(strategy), memRes};
+}
+
+template struct OptimizedStorage<StorageValueType>;
+
+}  // namespace block8
diff --git a/Engine/Plugins/Animation/RigLogic/Source/RigLogicLibTest/Private/rltests/joints/bpcm/JointFixturesBlock8.h b/Engine/Plugins/Animation/RigLogic/Source/RigLogicLibTest/Private/rltests/joints/bpcm/JointFixturesBlock8.h
new file mode 100644
--- /dev/null
+++ b/Engine/Plugins/Animation/RigLogic/Source/RigLogicLibTest/Private/rltests/joints/bpcm/JointFixturesBlock8.h
@@ -0,0 +1,100 @@
+#pragma once
+
+#include "rltests/Defs.h"
+#include "rltests/dna/FakeReader.h"
+#include "rltests/joints/bpcm/StorageValueType.h"
+
+#include "riglogic/TypeDefs.h"
+#include "riglogic/joints/bpcm/Extent.h"
+#include "riglogic/joints/bpcm/JointGroup.h"
+#include "riglogic/joints/bpcm/JointsEvaluator.h"
+#include "riglogic/joints/bpcm/LODRegion.h"
+
+namespace block8 {
+
+using namespace rl4;
+
+namespace unoptimized {
+
+extern const std::uint16_t lodCount;
+extern const bpcm::Extent dimensions;
+extern const Matrix<float> values;
+extern const Matrix<std::uint16_t> inputIndices;
+extern const Matrix<std::uint16_t> outputIndices;
+extern const Matrix<std::uint16_t> lods;
+
+}  // namespace unoptimized
+
+namespace optimized {
+
+extern const bpcm::Extent dimensions;
+extern const AlignedMatrix<float> floatValues;
+extern const AlignedMatrix<std::uint16_t> halfFloatValues;
+extern const AlignedMatrix<std::uint16_t> inputIndices;
+extern const AlignedMatrix<std::uint16_t> outputIndices;
+extern const Vector<bpcm::JointGroup> jointGroups;
+extern const Matrix<bpcm::LODRegion> lodRegions;
+
+}  // namespace optimized
+
+namespace input {
+
+// Calculation input values
+extern const Vector<float> values;
+
+}  // namespace input
+
+namespace output {
+
+// Calculation output values
+extern const Matrix<float> valuesPerLOD;
+
+}  // namespace output
+
+class CanonicalReader : public dna::FakeReader {
+    public:
+        std::uint16_t getLODCount() const override {
+            return unoptimized::lodCount;
+        }
+
+        std::uint16_t getJointRowCount() const override {
+            return unoptimized::dimensions.rows;
+        }
+
+        std::uint16_t getJointColumnCount() const override {
+            return unoptimized::dimensions.cols;
+        }
+
+        std::uint16_t getJointGroupCount() const override {
+            return static_cast<std::uint16_t>(unoptimized::values.size());
+        }
+
+        ConstArrayView<std::uint16_t> getJointGroupLODs(std::uint16_t jointGroupIndex) const override {
+            return ConstArrayView<std::uint16_t>{unoptimized::lods[jointGroupIndex]};
+        }
+
+        ConstArrayView<std::uint16_t> getJointGroupInputIndices(std::uint16_t jointGroupIndex) const override {
+            return ConstArrayView<std::uint16_t>{unoptimized::inputIndices[jointGroupIndex]};
+        }
+
+        ConstArrayView<std::uint16_t> getJointGroupOutputIndices(std::uint16_t jointGroupIndex) const override {
+            return ConstArrayView<std::uint16_t>{unoptimized::outputIndices[jointGroupIndex]};
+        }
+
+        ConstArrayView<float> getJointGroupValues(std::uint16_t jointGroupIndex) const override {
+            return ConstArrayView<float>{unoptimized::values[jointGroupIndex]};
+        }
+
+};
+
+template<typename TValue>
+struct OptimizedStorage {
+    using StrategyPtr = std::unique_ptr<bpcm::JointCalculationStrategy<TValue>,
+                                        std::function<void (bpcm::JointCalculationStrategy<TValue>*)> >;
+
+    static bpcm::Evaluator<TValue> create(StrategyPtr strategy, rl4::MemoryResource* memRes);
+    static bpcm::Evaluator<TValue> create(StrategyPtr strategy, std::uint16_t jointGroupIndex, rl4::MemoryResource* memRes);
+
+};
+
+}  // namespace block8
diff --git a/Engine/Plugins/Animation/RigLogic/Source/RigLogicLibTest/Private/rltests/joints/bpcm/TestJointCalculationStrategyAVX.cpp b/Engine/Plugins/Animation/RigLogic/Source/RigLogicLibTest/Private/rltests/joints/bpcm/TestJointCalculationStrategyAVX.cpp
new file mode 100644
--- /dev/null
+++ b/Engine/Plugins/Animation/RigLogic/Source/RigLogicLibTest/Private/rltests/joints/bpcm/TestJointCalculationStrategyAVX.cpp
@@ -0,0 +1,106 @@
+#ifdef RL_BUILD_WITH_AVX
+
+#include "rltests/Defs.h"
+#include "rltests/joints/bpcm/Helpers.h"
+#include "rltests/joints/bpcm/JointFixturesBlock8.h"
+
+#include "riglogic/joints/bpcm/JointsEvaluator.h"
+#include "riglogic/joints/bpcm/strategies/AVX.h"
+#include "riglogic/types/Aliases.h"
+
+#include <pma/resources/AlignedMemoryResource.h>
+#include <pma/utils/ManagedInstance.h>
+
+class AVXJointCalculationStrategyTest : public ::testing::TestWithParam<StrategyTestParams> {
+    protected:
+        void SetUp() {
+            strategy = pma::UniqueInstance<rl4::bpcm::AVXJointCalculationStrategy<StorageValueType>,
+                                            rl4::bpcm::JointCalculationStrategy<StorageValueType> >::with(&memRes).create();
+        }
+
+        template<typename TArray>
+        void execute(const rl4::bpcm::Evaluator<StorageValueType>& joints, const TArray& expected, OutputScope scope) {
+            rl4::ConstArrayView<float> expectedView{expected[scope.lod].data() + scope.offset, scope.size};
+            rl4::AlignedVector<float> outputs{expected.front().size(), {}, & memRes};
+            rl4::ConstArrayView<float> inputView{block8::input::values};
+            rl4::ConstArrayView<float> outputView{outputs.data() + scope.offset, scope.size};
+            joints.calculate(inputView, rl4::ArrayView<float>{outputs}, scope.lod);
+            ASSERT_EQ(outputView, expectedView);
+        }
+
+    protected:
+        pma::AlignedMemoryResource memRes;
+        block8::OptimizedStorage<StorageValueType>::StrategyPtr strategy;
+
+};
+
+TEST_P(AVXJointCalculationStrategyTest, Block8Padded) {
+    const auto params = GetParam();
+    const OutputScope scope{params.lod, 0ul, 5ul};
+    auto joints = block8::OptimizedStorage<StorageValueType>::create(std::move(strategy), 0u, &memRes);
+    execute(joints, block8::output::valuesPerLOD, scope);
+}
+
+TEST_P(AVXJointCalculationStrategyTest, Block8Exact) {
+    const auto params = GetParam();
+    const OutputScope scope{params.lod, 5ul, 8ul};
+    auto joints = block8::OptimizedStorage<StorageValueType>::create(std::move(strategy), 1u, &memRes);
+    execute(joints, block8::output::valuesPerLOD, scope);
+}
+
+TEST_P(AVXJointCalculationStrategyTest, Block16Padded) {
+    const auto params = GetParam();
+    const OutputScope scope{params.lod, 13ul, 9ul};
+    auto joints = block8::OptimizedStorage<StorageValueType>::create(std::move(strategy), 2u, &memRes);
+    execute(joints, block8::output::valuesPerLOD, scope);
+}
+
+TEST_P(AVXJointCalculationStrategyTest, Block16Exact) {
+    const auto params = GetParam();
+    const OutputScope scope{params.lod, 22ul, 16ul};
+    auto joints = block8::OptimizedStorage<StorageValueType>::create(std::move(strategy), 3u, &memRes);
+    execute(joints, block8::output::valuesPerLOD, scope);
+}
+
+TEST_P(AVXJointCalculationStrategyTest, Block24Padded) {
+    const auto params = GetParam();
+    const OutputScope scope{params.lod, 38ul, 17ul};
+    auto joints = block8::OptimizedStorage<StorageValueType>::create(std::move(strategy), 4u, &memRes);
+    execute(joints, block8::output::valuesPerLOD, scope);
+}
+
+TEST_P(AVXJointCalculationStrategyTest, Block24Exact) {
+    const auto params = GetParam();
+    const OutputScope scope{params.lod, 55ul, 24ul};
+    auto joints = block8::OptimizedStorage<StorageValueType>::create(std::move(strategy), 5u, &memRes);
+    execute(joints, block8::output::valuesPerLOD, scope);
+}
+
+TEST_P(AVXJointCalculationStrategyTest, Block32Padded) {
+    const auto params = GetParam();
+    const OutputScope scope{params.lod, 79ul, 25ul};
+    auto joints = block8::OptimizedStorage<StorageValueType>::create(std::move(strategy), 6u, &memRes);
+    execute(joints, block8::output::valuesPerLOD, scope);
+}
+
+TEST_P(AVXJointCalculationStrategyTest, Block32Exact) {
+    const auto params = GetParam();
+    const OutputScope scope{params.lod, 104ul, 32ul};
+    auto joints = block8::OptimizedStorage<StorageValueType>::create(std::move(strategy), 7u, &memRes);
+    execute(joints, block8::output::valuesPerLOD, scope);
+}
+
+TEST_P(AVXJointCalculationStrategyTest, MultipleBlocks) {
+    const auto params = GetParam();
+    const OutputScope scope{params.lod, 0ul, block8::output::valuesPerLOD.front().size()};
+    auto joints = block8::OptimizedStorage<StorageValueType>::create(std::move(strategy), &memRes);
+    execute(joints, block8::output::valuesPerLOD, scope);
+}
+
+INSTANTIATE_TEST_SUITE_P(AVXJointCalculationStrategyTest, AVXJointCalculationStrategyTest, ::testing::Values(
+                             StrategyTestParams{0u},
+                             StrategyTestParams{1u},
+                             StrategyTestParams{2u},
+                             StrategyTestParams{3u}));
+
+#endif  // RL_BUILD_WITH_AVX
diff --git a/Engine/Plugins/Animation/RigLogic/Source/RigLogicLibTest/Private/rltests/joints/bpcm/TestJointCalculationStrategySSE.cpp b/Engine/Plugins/Animation/RigLogic/Source/RigLogicLibTest/Private/rltests/joints/bpcm/TestJointCalculationStrategySSE.cpp
new file mode 100644
--- /dev/null
+++ b/Engine/Plugins/Animation/RigLogic/Source/RigLogicLibTest/Private/rltests/joints/bpcm/TestJointCalculationStrategySSE.cpp
@@ -0,0 +1,120 @@
+#ifdef RL_BUILD_WITH_SSE
+
+#include "rltests/Defs.h"
+#include "rltests/joints/bpcm/Helpers.h"
+#include "rltests/joints/bpcm/JointFixturesBlock4.h"
+
+#include "riglogic/joints/bpcm/JointsEvaluator.h"
+#include "riglogic/joints/bpcm/strategies/SSE.h"
+#include "riglogic/types/Aliases.h"
+
+#include <pma/resources/AlignedMemoryResource.h>
+#include <pma/utils/ManagedInstance.h>
+
+class SSEJointCalculationStrategyTest : public ::testing::TestWithParam<StrategyTestParams> {
+    protected:
+        void SetUp() {
+            strategy = pma::UniqueInstance<rl4::bpcm::SSEJointCalculationStrategy<StorageValueType>,
+                                           rl4::bpcm::JointCalculationStrategy<StorageValueType> >::with(&memRes).create();
+        }
+
+        template<typename TArray>
+        void execute(const rl4::bpcm::Evaluator<StorageValueType>& joints, const TArray& expected, OutputScope scope) {
+            rl4::ConstArrayView<float> expectedView{expected[scope.lod].data() + scope.offset, scope.size};
+            rl4::AlignedVector<float> outputs{expected.front().size(), {}, & memRes};
+            rl4::ConstArrayView<float> inputView{block4::input::values};
+            rl4::ConstArrayView<float> outputView{outputs.data() + scope.offset, scope.size};
+            joints.calculate(inputView, rl4::ArrayView<float>{outputs}, scope.lod);
+            ASSERT_EQ(outputView, expectedView);
+        }
+
+    protected:
+        pma::AlignedMemoryResource memRes;
+        block4::OptimizedStorage<StorageValueType>::StrategyPtr strategy;
+};
+
+TEST_P(SSEJointCalculationStrategyTest, Block4Padded) {
+    const auto params = GetParam();
+    const OutputScope scope{params.lod, 0ul, 1ul};
+    auto joints = block4::OptimizedStorage<StorageValueType>::create(std::move(strategy), 0u, &memRes);
+    execute(joints, block4::output::valuesPerLOD, scope);
+}
+
+TEST_P(SSEJointCalculationStrategyTest, Block4Exact) {
+    const auto params = GetParam();
+    const OutputScope scope{params.lod, 1ul, 4ul};
+    auto joints = block4::OptimizedStorage<StorageValueType>::create(std::move(strategy), 1u, &memRes);
+    execute(joints, block4::output::valuesPerLOD, scope);
+}
+
+TEST_P(SSEJointCalculationStrategyTest, Block8Padded) {
+    const auto params = GetParam();
+    const OutputScope scope{params.lod, 5ul, 7ul};
+    auto joints = block4::OptimizedStorage<StorageValueType>::create(std::move(strategy), 2u, &memRes);
+    execute(joints, block4::output::valuesPerLOD, scope);
+}
+
+TEST_P(SSEJointCalculationStrategyTest, Block8Exact) {
+    const auto params = GetParam();
+    const OutputScope scope{params.lod, 12ul, 8ul};
+    auto joints = block4::OptimizedStorage<StorageValueType>::create(std::move(strategy), 3u, &memRes);
+    execute(joints, block4::output::valuesPerLOD, scope);
+}
+
+TEST_P(SSEJointCalculationStrategyTest, Block12Padded) {
+    const auto params = GetParam();
+    const OutputScope scope{params.lod, 20ul, 9ul};
+    auto joints = block4::OptimizedStorage<StorageValueType>::create(std::move(strategy), 4u, &memRes);
+    execute(joints, block4::output::valuesPerLOD, scope);
+}
+
+TEST_P(SSEJointCalculationStrategyTest, Block12Exact) {
+    const auto params = GetParam();
+    const OutputScope scope{params.lod, 29ul, 12ul};
+    auto joints = block4::OptimizedStorage<StorageValueType>::create(std::move(strategy), 5u, &memRes);
+    execute(joints, block4::output::valuesPerLOD, scope);
+}
+
+TEST_P(SSEJointCalculationStrategyTest, Block16Padded) {
+    const auto params = GetParam();
+    const OutputScope scope{params.lod, 41ul, 14ul};
+    auto joints = block4::OptimizedStorage<StorageValueType>::create(std::move(strategy), 6u, &memRes);
+    execute(joints, block4::output::valuesPerLOD, scope);
+}
+
+TEST_P(SSEJointCalculationStrategyTest, Block16Exact) {
+    const auto params = GetParam();
+    const OutputScope scope{params.lod, 55ul, 16ul};
+    auto joints = block4::OptimizedStorage<StorageValueType>::create(std::move(strategy), 7u, &memRes);
+    execute(joints, block4::output::valuesPerLOD, scope);
+}
+
+TEST_P(SSEJointCalculationStrategyTest, MultipleBlocks) {
+    const auto params = GetParam();
+    const OutputScope scope{params.lod, 0ul, block4::output::valuesPerLOD.front().size()};
+    auto joints = block4::OptimizedStorage<StorageValueType>::create(std::move(strategy), &memRes);
+    execute(joints, block4::output::valuesPerLOD, scope);
+}
+
+TEST_P(SSEJointCalculationStrategyTest, InputRegionA) {
+    const auto params = GetParam();
+    const OutputScope scope{params.lod, 71ul, 2ul};
+    auto joints = block4::OptimizedStorage<StorageValueType>::create(std::move(strategy), 8u, &memRes);
+    execute(joints, block4::output::valuesPerLOD, scope);
+}
+
+TEST_P(SSEJointCalculationStrategyTest, InputRegionB) {
+    const auto params = GetParam();
+    const OutputScope scope{params.lod, 73ul, 2ul};
+    auto joints = block4::OptimizedStorage<StorageValueType>::create(std::move(strategy), 9u, &memRes);
+    execute(joints, block4::output::valuesPerLOD, scope);
+}
+
+INSTANTIATE_TEST_SUITE_P(SSEJointCalculationStrategyTest, SSEJointCalculationStrategyTest, ::testing::Values(
+                             StrategyTestParams{0u},
+                             StrategyTestParams{1u},
+                             StrategyTestParams{2u},
+                             StrategyTestParams{3u}
+                             ));
+
+#endif  // RL_BUILD_WITH_SSE
diff --git a/Engine/Plugins/Animation/RigLogic/Source/RigLogicLibTest/Private/rltests/joints/bpcm/TestJointCalculationStrategyScalar.cpp b/Engine/Plugins/Animation/RigLogic/Source/RigLogicLibTest/Private/rltests/joints/bpcm/TestJointCalculationStrategyScalar.cpp
new file mode 100644
--- /dev/null
+++ b/Engine/Plugins/Animation/RigLogic/Source/RigLogicLibTest/Private/rltests/joints/bpcm/TestJointCalculationStrategyScalar.cpp
@@ -0,0 +1,116 @@
+#include "rltests/Defs.h"
+#include "rltests/joints/bpcm/Helpers.h"
+#include "rltests/joints/bpcm/JointFixturesBlock4.h"
+
+#include "riglogic/joints/bpcm/JointsEvaluator.h"
+#include "riglogic/joints/bpcm/strategies/Scalar.h"
+#include "riglogic/types/Aliases.h"
+
+#include <pma/resources/AlignedMemoryResource.h>
+#include <pma/utils/ManagedInstance.h>
+
+class ScalarJointCalculationStrategyTest : public ::testing::TestWithParam<StrategyTestParams> {
+    protected:
+        void SetUp() {
+            strategy = pma::UniqueInstance<rl4::bpcm::ScalarJointCalculationStrategy,
+                                           rl4::bpcm::JointCalculationStrategy<float> >::with(&memRes).create();
+        }
+
+        template<typename TArray>
+        void execute(const rl4::bpcm::Evaluator<float>& joints, const TArray& expected, OutputScope scope) {
+            rl4::ConstArrayView<float> expectedView{expected[scope.lod].data() + scope.offset, scope.size};
+            rl4::AlignedVector<float> outputs{expected.front().size(), {}, & memRes};
+            rl4::ConstArrayView<float> inputView{block4::input::values};
+            rl4::ConstArrayView<float> outputView{outputs.data() + scope.offset, scope.size};
+            joints.calculate(inputView, rl4::ArrayView<float>{outputs}, scope.lod);
+            ASSERT_EQ(outputView, expectedView);
+        }
+
+    protected:
+        pma::AlignedMemoryResource memRes;
+        block4::OptimizedStorage<float>::StrategyPtr strategy;
+};
+
+TEST_P(ScalarJointCalculationStrategyTest, Block4Padded) {
+    const auto params = GetParam();
+    const OutputScope scope{params.lod, 0ul, 1ul};
+    auto joints = block4::OptimizedStorage<float>::create(std::move(strategy), 0u, &memRes);
+    execute(joints, block4::output::valuesPerLOD, scope);
+}
+
+TEST_P(ScalarJointCalculationStrategyTest, Block4Exact) {
+    const auto params = GetParam();
+    const OutputScope scope{params.lod, 1ul, 4ul};
+    auto joints = block4::OptimizedStorage<float>::create(std::move(strategy), 1u, &memRes);
+    execute(joints, block4::output::valuesPerLOD, scope);
+}
+
+TEST_P(ScalarJointCalculationStrategyTest, Block8Padded) {
+    const auto params = GetParam();
+    const OutputScope scope{params.lod, 5ul, 7ul};
+    auto joints = block4::OptimizedStorage<float>::create(std::move(strategy), 2u, &memRes);
+    execute(joints, block4::output::valuesPerLOD, scope);
+}
+
+TEST_P(ScalarJointCalculationStrategyTest, Block8Exact) {
+    const auto params = GetParam();
+    const OutputScope scope{params.lod, 12ul, 8ul};
+    auto joints = block4::OptimizedStorage<float>::create(std::move(strategy), 3u, &memRes);
+    execute(joints, block4::output::valuesPerLOD, scope);
+}
+
+TEST_P(ScalarJointCalculationStrategyTest, Block12Padded) {
+    const auto params = GetParam();
+    const OutputScope scope{params.lod, 20ul, 9ul};
+    auto joints = block4::OptimizedStorage<float>::create(std::move(strategy), 4u, &memRes);
+    execute(joints, block4::output::valuesPerLOD, scope);
+}
+
+TEST_P(ScalarJointCalculationStrategyTest, Block12Exact) {
+    const auto params = GetParam();
+    const OutputScope scope{params.lod, 29ul, 12ul};
+    auto joints = block4::OptimizedStorage<float>::create(std::move(strategy), 5u, &memRes);
+    execute(joints, block4::output::valuesPerLOD, scope);
+}
+
+TEST_P(ScalarJointCalculationStrategyTest, Block16Padded) {
+    const auto params = GetParam();
+    const OutputScope scope{params.lod, 41ul, 14ul};
+    auto joints = block4::OptimizedStorage<float>::create(std::move(strategy), 6u, &memRes);
+    execute(joints, block4::output::valuesPerLOD, scope);
+}
+
+TEST_P(ScalarJointCalculationStrategyTest, Block16Exact) {
+    const auto params = GetParam();
+    const OutputScope scope{params.lod, 55ul, 16ul};
+    auto joints = block4::OptimizedStorage<float>::create(std::move(strategy), 7u, &memRes);
+    execute(joints, block4::output::valuesPerLOD, scope);
+}
+
+TEST_P(ScalarJointCalculationStrategyTest, MultipleBlocks) {
+    const auto params = GetParam();
+    const OutputScope scope{params.lod, 0ul, block4::output::valuesPerLOD.front().size()};
+    auto joints = block4::OptimizedStorage<float>::create(std::move(strategy), &memRes);
+    execute(joints, block4::output::valuesPerLOD, scope);
+}
+
+TEST_P(ScalarJointCalculationStrategyTest, InputRegionA) {
+    const auto params = GetParam();
+    const OutputScope scope{params.lod, 71ul, 2ul};
+    auto joints = block4::OptimizedStorage<float>::create(std::move(strategy), 8u, &memRes);
+    execute(joints, block4::output::valuesPerLOD, scope);
+}
+
+TEST_P(ScalarJointCalculationStrategyTest, InputRegionB) {
+    const auto params = GetParam();
+    const OutputScope scope{params.lod, 73ul, 2ul};
+    auto joints = block4::OptimizedStorage<float>::create(std::move(strategy), 9u, &memRes);
+    execute(joints, block4::output::valuesPerLOD, scope);
+}
+
+INSTANTIATE_TEST_SUITE_P(ScalarJointCalculationStrategyTest, ScalarJointCalculationStrategyTest, ::testing::Values(
+                             StrategyTestParams{0u},
+                             StrategyTestParams{1u},
+                             StrategyTestParams{2u},
+                             StrategyTestParams{3u}
+                             ));
diff --git a/Engine/Plugins/Animation/RigLogic/Source/RigLogicLibTest/Private/rltests/joints/bpcm/TestJointStorageBuilderAVX.cpp b/Engine/Plugins/Animation/RigLogic/Source/RigLogicLibTest/Private/rltests/joints/bpcm/TestJointStorageBuilderAVX.cpp
new file mode 100644
--- /dev/null
+++ b/Engine/Plugins/Animation/RigLogic/Source/RigLogicLibTest/Private/rltests/joints/bpcm/TestJointStorageBuilderAVX.cpp
@@ -0,0 +1,48 @@
+#ifdef RL_BUILD_WITH_AVX
+
+#include "rltests/Defs.h"
+#include "rltests/joints/bpcm/Assertions.h"
+#include "rltests/joints/bpcm/JointFixturesBlock8.h"
+
+#include "riglogic/joints/JointsBuilder.h"
+#include "riglogic/joints/bpcm/JointsEvaluator.h"
+#include "riglogic/joints/bpcm/strategies/AVX.h"
+#include "riglogic/riglogic/RigLogic.h"
+
+#include <pma/resources/AlignedMemoryResource.h>
+#include <pma/utils/ManagedInstance.h>
+
+template<typename TValue>
+class AVXJointStorageBuilderTest : public ::testing::Test {
+    protected:
+        void buildStorage() {
+            rl4::Configuration config{rl4::RigLogic::CalculationType::AVX};
+            auto builder = rl4::JointsBuilder::create(config, &memRes);
+            builder->allocateStorage(&reader);
+            builder->fillStorage(&reader);
+            auto joints = builder->build();
+            auto jointsImpl = static_cast<rl4::bpcm::Evaluator<TValue>*>(joints.get());
+
+            auto strategy = pma::UniqueInstance<rl4::bpcm::AVXJointCalculationStrategy<TValue>,
+                                                 rl4::bpcm::JointCalculationStrategy<TValue> >::with(&memRes).create();
+            auto expected = block8::OptimizedStorage<TValue>::create(std::move(strategy), &memRes);
+
+            rl4::bpcm::Evaluator<TValue>::Accessor::assertRawDataEqual(*jointsImpl, expected);
+            rl4::bpcm::Evaluator<TValue>::Accessor::assertJointGroupsEqual(*jointsImpl, expected);
+            rl4::bpcm::Evaluator<TValue>::Accessor::assertLODsEqual(*jointsImpl, expected);
+        }
+
+    protected:
+        pma::AlignedMemoryResource memRes;
+        block8::CanonicalReader reader;
+
+};
+
+using StorageValueTypeList = ::testing::Types<StorageValueType>;
+TYPED_TEST_SUITE(AVXJointStorageBuilderTest, StorageValueTypeList);
+
+TYPED_TEST(AVXJointStorageBuilderTest, LayoutOptimization) {
+    this->buildStorage();
+}
+
+#endif  // RL_BUILD_WITH_AVX
diff --git a/Engine/Plugins/Animation/RigLogic/Source/RigLogicLibTest/Private/rltests/joints/bpcm/TestJointStorageBuilderSSE.cpp b/Engine/Plugins/Animation/RigLogic/Source/RigLogicLibTest/Private/rltests/joints/bpcm/TestJointStorageBuilderSSE.cpp
new file mode 100644
--- /dev/null
+++ b/Engine/Plugins/Animation/RigLogic/Source/RigLogicLibTest/Private/rltests/joints/bpcm/TestJointStorageBuilderSSE.cpp
@@ -0,0 +1,48 @@
+#ifdef RL_BUILD_WITH_SSE
+
+#include "rltests/Defs.h"
+#include "rltests/joints/bpcm/Assertions.h"
+#include "rltests/joints/bpcm/JointFixturesBlock4.h"
+
+#include "riglogic/joints/JointsBuilder.h"
+#include "riglogic/joints/bpcm/JointsEvaluator.h"
+#include "riglogic/joints/bpcm/strategies/SSE.h"
+#include "riglogic/riglogic/RigLogic.h"
+
+#include <pma/resources/AlignedMemoryResource.h>
+#include <pma/utils/ManagedInstance.h>
+
+template<typename TValue>
+class SSEJointStorageBuilderTest : public ::testing::Test {
+    protected:
+        void buildStorage() {
+            rl4::Configuration config{rl4::RigLogic::CalculationType::SSE};
+            auto builder = rl4::JointsBuilder::create(config, &memRes);
+            builder->allocateStorage(&reader);
+            builder->fillStorage(&reader);
+            auto joints = builder->build();
+            auto jointsImpl = static_cast<rl4::bpcm::Evaluator<TValue>*>(joints.get());
+
+            auto strategy = pma::UniqueInstance<rl4::bpcm::SSEJointCalculationStrategy<TValue>,
+                                                rl4::bpcm::JointCalculationStrategy<TValue> >::with(&memRes).create();
+            auto expected = block4::OptimizedStorage<TValue>::create(std::move(strategy), &memRes);
+
+            rl4::bpcm::Evaluator<TValue>::Accessor::assertRawDataEqual(*jointsImpl, expected);
+            rl4::bpcm::Evaluator<TValue>::Accessor::assertJointGroupsEqual(*jointsImpl, expected);
+            rl4::bpcm::Evaluator<TValue>::Accessor::assertLODsEqual(*jointsImpl, expected);
+        }
+
+    protected:
+        pma::AlignedMemoryResource memRes;
+        block4::CanonicalReader reader;
+
+};
+
+using StorageValueTypeList = ::testing::Types<StorageValueType>;
+TYPED_TEST_SUITE(SSEJointStorageBuilderTest, StorageValueTypeList);
+
+TYPED_TEST(SSEJointStorageBuilderTest, LayoutOptimization) {
+    this->buildStorage();
+}
+
+#endif  // RL_BUILD_WITH_SSE
diff --git a/Engine/Plugins/Animation/RigLogic/Source/RigLogicLibTest/Private/rltests/joints/bpcm/TestJointStorageBuilderScalar.cpp b/Engine/Plugins/Animation/RigLogic/Source/RigLogicLibTest/Private/rltests/joints/bpcm/TestJointStorageBuilderScalar.cpp
new file mode 100644
--- /dev/null
+++ b/Engine/Plugins/Animation/RigLogic/Source/RigLogicLibTest/Private/rltests/joints/bpcm/TestJointStorageBuilderScalar.cpp
@@ -0,0 +1,41 @@
+#include "rltests/Defs.h"
+#include "rltests/joints/bpcm/Assertions.h"
+#include "rltests/joints/bpcm/JointFixturesBlock4.h"
+
+#include "riglogic/joints/JointsBuilder.h"
+#include "riglogic/joints/JointsFactory.h"
+#include "riglogic/joints/bpcm/JointsEvaluator.h"
+#include "riglogic/joints/bpcm/strategies/Scalar.h"
+#include "riglogic/riglogic/RigLogic.h"
+
+#include <pma/resources/AlignedMemoryResource.h>
+#include <pma/utils/ManagedInstance.h>
+
+class ScalarJointStorageBuilderTest : public ::testing::Test {
+    protected:
+        void buildStorage() {
+            rl4::Configuration config{rl4::RigLogic::CalculationType::Scalar};
+            auto builder = rl4::JointsBuilder::create(config, &memRes);
+            builder->allocateStorage(&reader);
+            builder->fillStorage(&reader);
+            auto joints = builder->build();
+            auto jointsImpl = static_cast<rl4::bpcm::Evaluator<float>*>(joints.get());
+
+            auto strategy = pma::UniqueInstance<rl4::bpcm::ScalarJointCalculationStrategy,
+                                                rl4::bpcm::JointCalculationStrategy<float> >::with(&memRes).create();
+            auto expected = block4::OptimizedStorage<float>::create(std::move(strategy), &memRes);
+
+            rl4::bpcm::Evaluator<float>::Accessor::assertRawDataEqual(*jointsImpl, expected);
+            rl4::bpcm::Evaluator<float>::Accessor::assertJointGroupsEqual(*jointsImpl, expected);
+            rl4::bpcm::Evaluator<float>::Accessor::assertLODsEqual(*jointsImpl, expected);
+        }
+
+    protected:
+        pma::AlignedMemoryResource memRes;
+        block4::CanonicalReader reader;
+
+};
+
+TEST_F(ScalarJointStorageBuilderTest, LayoutOptimization) {
+    this->buildStorage();
+}
diff --git a/Engine/Plugins/Animation/RigLogic/Source/RigLogicLibTest/Private/rltests/joints/bpcm/TestLODRegion.cpp b/Engine/Plugins/Animation/RigLogic/Source/RigLogicLibTest/Private/rltests/joints/bpcm/TestLODRegion.cpp
new file mode 100644
--- /dev/null
+++ b/Engine/Plugins/Animation/RigLogic/Source/RigLogicLibTest/Private/rltests/joints/bpcm/TestLODRegion.cpp
@@ -0,0 +1,106 @@
+#include "rltests/Defs.h"
+
+#include "riglogic/joints/bpcm/LODRegion.h"
+
+#include <cstdint>
+
+TEST(LODRegionTest, ConstructFromFixtures) {
+    rl4::bpcm::LODRegion lod{1, 2, 3};
+
+    ASSERT_EQ(lod.size, 1);
+    ASSERT_EQ(lod.sizeAlignedToLastFullBlock, 2);
+    ASSERT_EQ(lod.sizeAlignedToSecondLastFullBlock, 3);
+}
+
+struct LODTestSetup {
+    std::uint32_t rowCount;
+    std::uint32_t lodEndRow;
+    std::uint32_t blockHeight;
+    std::uint32_t padTo;
+    std::uint32_t lodEndRowAlignedToLastFullBlock;
+    std::uint32_t lodEndRowAlignedToSecondLastFullBlock;
+};
+
+class LODRegionTest : public ::testing::TestWithParam<LODTestSetup> {
+};
+
+TEST_P(LODRegionTest, ConstructLODRegionFromVaryingParameters) {
+    auto params = GetParam();
+    rl4::bpcm::LODRegion lod{params.lodEndRow, params.rowCount, params.blockHeight, params.padTo};
+    ASSERT_EQ(lod.size, params.lodEndRow);
+    ASSERT_EQ(lod.sizeAlignedToLastFullBlock, params.lodEndRowAlignedToLastFullBlock);
+    ASSERT_EQ(lod.sizeAlignedToSecondLastFullBlock, params.lodEndRowAlignedToSecondLastFullBlock);
+}
+
+INSTANTIATE_TEST_SUITE_P(
+    LODTestSuite,
+    LODRegionTest,
+    ::testing::Values(
+        // Matrices passed to LODRegion have already been padded,
+        // so the rowcount must be an integral multiple of 4
+        // {{rows, cols}, lodEndRow, lodEndRowAlignedToLastFullBlock, lodEndRowAlignedToSecondLastFullBlock}
+        // SSE test cases
+        LODTestSetup{4, 2, 8, 4, 0, 0},  // Mask-off last 2 rows from 1st block-4 (handles block-4 loop)
+        LODTestSetup{4, 4, 8, 4, 0, 0},  // No mask-off (handles block-4 loop)
+        LODTestSetup{8, 2, 8, 4, 8, 0},  // Mask-off last 6 rows from 1st block-8 (handles block-8 loop)
+        LODTestSetup{8, 4, 8, 4, 8, 0},  // Mask-off last 4 rows from 1st block-8 (handles block-8 loop)
+        LODTestSetup{8, 6, 8, 4, 8, 0},  // Mask-off last 2 rows from 1st block-8 (handles block-8 loop)
+        LODTestSetup{8, 8, 8, 4, 8, 8},  // No mask-off (handles block-8 loop without mask-off logic)
+        LODTestSetup{12, 2, 8, 4, 8, 0},  // Mask-off last 6 rows from 1st block-8 (handles block-8 loop)
+        LODTestSetup{12, 4, 8, 4, 8, 0},  // Mask-off last 4 rows from 1st block-8 (handles block-8 loop)
+        LODTestSetup{12, 6, 8, 4, 8, 0},  // Mask-off last 2 rows from 1st block-8 (handles block-8 loop)
+        LODTestSetup{12, 8, 8, 4, 8, 8},  // No mask-off (handles block-8 loop without mask-off logic)
+        LODTestSetup{12, 10, 8, 4, 8, 8},  // Mask-off last 2 rows from block-4 (handles block-8 without mask-off, then block-4
+                                           // loop)
+        LODTestSetup{12, 12, 8, 4, 8, 8},  // No mask-off (handles block-8 without mask-off, then block-4 loop)
+        LODTestSetup{32, 2, 8, 4, 8, 0},  // Mask-off last 6 rows from 1st block-8 (handles block-8 loop)
+        LODTestSetup{32, 4, 8, 4, 8, 0},  // Mask-off last 4 rows from 1st block-8 (handles block-8 loop)
+        LODTestSetup{32, 6, 8, 4, 8, 0},  // Mask-off last 2 rows from 1st block-8 (handles block-8 loop)
+        LODTestSetup{32, 8, 8, 4, 8, 8},  // No mask-off (handles block-8 loop without mask-off logic)
+        LODTestSetup{32, 9, 8, 4, 16, 8},  // Mask-off last 7 rows from 2nd block-8 (handles block-8 loop without mask-off, then
+                                           // block-8 loop)
+        LODTestSetup{32, 12, 8, 4, 16, 8},  // Mask-off last 4 rows from 2nd block-8 (handles block-8 loop without mask-off, then
+                                            // block-8 loop)
+        LODTestSetup{32, 16, 8, 4, 16, 16},  // No mask-off (handles block-8 loop without mask-off logic)
+        // AVX test cases (pads to 8 so (rowCount % 4 == 0) is not possible)
+        LODTestSetup{8, 2, 16, 8, 0, 0},
+        LODTestSetup{8, 4, 16, 8, 0, 0},
+        LODTestSetup{8, 6, 16, 8, 0, 0},
+        LODTestSetup{8, 8, 16, 8, 0, 0},
+        LODTestSetup{16, 2, 16, 8, 16, 0},
+        LODTestSetup{16, 4, 16, 8, 16, 0},
+        LODTestSetup{16, 6, 16, 8, 16, 0},
+        LODTestSetup{16, 8, 16, 8, 16, 0},
+        LODTestSetup{16, 10, 16, 8, 16, 0},
+        LODTestSetup{16, 12, 16, 8, 16, 0},
+        LODTestSetup{16, 14, 16, 8, 16, 0},
+        LODTestSetup{16, 16, 16, 8, 16, 16},
+        LODTestSetup{24, 2, 16, 8, 16, 0},
+        LODTestSetup{24, 4, 16, 8, 16, 0},
+        LODTestSetup{24, 6, 16, 8, 16, 0},
+        LODTestSetup{24, 8, 16, 8, 16, 0},
+        LODTestSetup{24, 10, 16, 8, 16, 0},
+        LODTestSetup{24, 12, 16, 8, 16, 0},
+        LODTestSetup{24, 14, 16, 8, 16, 0},
+        LODTestSetup{24, 16, 16, 8, 16, 16},
+        LODTestSetup{24, 18, 16, 8, 16, 16},
+        LODTestSetup{24, 20, 16, 8, 16, 16},
+        LODTestSetup{24, 22, 16, 8, 16, 16},
+        LODTestSetup{24, 24, 16, 8, 16, 16},
+        LODTestSetup{32, 2, 16, 8, 16, 0},
+        LODTestSetup{32, 4, 16, 8, 16, 0},
+        LODTestSetup{32, 6, 16, 8, 16, 0},
+        LODTestSetup{32, 8, 16, 8, 16, 0},
+        LODTestSetup{32, 10, 16, 8, 16, 0},
+        LODTestSetup{32, 12, 16, 8, 16, 0},
+        LODTestSetup{32, 14, 16, 8, 16, 0},
+        LODTestSetup{32, 16, 16, 8, 16, 16},
+        LODTestSetup{32, 18, 16, 8, 32, 16},
+        LODTestSetup{32, 20, 16, 8, 32, 16},
+        LODTestSetup{32, 22, 16, 8, 32, 16},
+        LODTestSetup{32, 24, 16, 8, 32, 16},
+        LODTestSetup{32, 26, 16, 8, 32, 16},
+        LODTestSetup{32, 28, 16, 8, 32, 16},
+        LODTestSetup{32, 30, 16, 8, 32, 16},
+        LODTestSetup{32, 32, 16, 8, 32, 32}
+        ));
diff --git a/Engine/Plugins/Animation/RigLogic/Source/RigLogicLibTest/Private/rltests/psdmatrix/TestPSDMatrix.cpp b/Engine/Plugins/Animation/RigLogic/Source/RigLogicLibTest/Private/rltests/psdmatrix/TestPSDMatrix.cpp
new file mode 100644
--- /dev/null
+++ b/Engine/Plugins/Animation/RigLogic/Source/RigLogicLibTest/Private/rltests/psdmatrix/TestPSDMatrix.cpp
@@ -0,0 +1,51 @@
+#include "rltests/Defs.h"
+
+#include "riglogic/psdmatrix/PSDMatrix.h"
+
+#include <pma/resources/AlignedMemoryResource.h>
+#include <pma/utils/ManagedInstance.h>
+
+TEST(PSDMatrixTest, OutputsAreClamped) {
+    pma::AlignedMemoryResource amr;
+
+    rl4::Vector<std::uint16_t> rows{1u};
+    rl4::Vector<std::uint16_t> cols{0u};
+    rl4::Vector<float> values{100.0f};
+    rl4::PSDMatrix psds{1ul, std::move(rows), std::move(cols), std::move(values)};
+
+    float inputs[] = {0.1f, 0.0f};
+    const float expected[] = {1.0f};
+
+    psds.calculate(inputs, 1ul);
+    ASSERT_ELEMENTS_EQ((inputs + 1ul), expected, 1);
+}
+
+TEST(PSDMatrixTest, OutputsKeepExistingProduct) {
+    pma::AlignedMemoryResource amr;
+
+    rl4::Vector<std::uint16_t> rows{2u, 2u};
+    rl4::Vector<std::uint16_t> cols{0u, 1u};
+    rl4::Vector<float> values{4.0f, 10.0f};
+    rl4::PSDMatrix psds{1ul, std::move(rows), std::move(cols), std::move(values)};
+
+    float inputs[] = {0.1f, 0.2f, 0.0f};
+    const float expected[] = {0.8f};
+
+    psds.calculate(inputs, 2ul);
+    ASSERT_ELEMENTS_EQ((inputs + 2ul), expected, 1);
+}
+
+TEST(PSDMatrixTest, RowsSpecifyDestinationIndex) {
+    pma::AlignedMemoryResource amr;
+
+    rl4::Vector<std::uint16_t> rows{2u, 3u};
+    rl4::Vector<std::uint16_t> cols{0u, 1u};
+    rl4::Vector<float> values{4.0f, 3.0f};
+    rl4::PSDMatrix psds{2ul, std::move(rows), std::move(cols), std::move(values)};
+
+    float inputs[] = {0.1f, 0.2f, 0.0f, 0.0f};
+    const float expected[] = {0.4f, 0.6f};
+
+    psds.calculate(inputs, 2ul);
+    ASSERT_ELEMENTS_EQ((inputs + 2ul), expected, 2);
+}
