Editgrating RigLogic plugin from Dev-CharacterTech to UE5/Main @ CL 13891876.

diff --git a/Engine/Plugins/Animation/RigLogic/Source/RigLogicLibTest/Private/rltests/riglogic/TestRigLogic.cpp b/Engine/Plugins/Animation/RigLogic/Source/RigLogicLibTest/Private/rltests/riglogic/TestRigLogic.cpp
new file mode 100644
--- /dev/null
+++ b/Engine/Plugins/Animation/RigLogic/Source/RigLogicLibTest/Private/rltests/riglogic/TestRigLogic.cpp
@@ -0,0 +1,95 @@
+#include "rltests/Defs.h"
+#include "rltests/dna/DNAFixtures.h"
+#include "rltests/dna/FakeStream.h"
+
+#include "riglogic/RigLogic.h"
+
+#include <pma/resources/AlignedMemoryResource.h>
+
+class RigLogicTest : public ::testing::Test {
+    protected:
+        void SetUp() override {
+            const auto bytes = rltests::raw::getBytes();
+            stream.write(reinterpret_cast<const char*>(bytes.data()), bytes.size());
+            stream.seek(0);
+
+            reader = dna::StreamReader::create(&stream);
+            reader->read();
+
+            rigLogic = rl4::RigLogic::create(reader);
+            rigInstance = rl4::RigInstance::create(rigLogic, &memRes);
+        }
+
+        void TearDown() override {
+            rl4::RigInstance::destroy(rigInstance);
+            rl4::RigLogic::destroy(rigLogic);
+            dna::StreamReader::destroy(reader);
+        }
+
+    protected:
+        pma::AlignedMemoryResource memRes;
+        rltests::FakeStream stream;
+        dna::StreamReader* reader;
+        rl4::RigLogic* rigLogic;
+        rl4::RigInstance* rigInstance;
+};
+
+TEST_F(RigLogicTest, EvaluateRigInstance) {
+    // Try both approaches
+    float guiControls[] = {1.0f, 2.0f, 3.0f, 4.0f, 5.0f, 6.0f, 7.0f, 8.0f, 9.0f};
+    rigInstance->setGUIControlValues(static_cast<const float*>(guiControls));
+    rigLogic->mapGUIToRawControls(rigInstance);
+    // Regardless that this overwrites the computed gui to raw values
+    float rawControls[] = {1.0f, 2.0f, 3.0f, 4.0f, 5.0f, 6.0f, 7.0f, 8.0f, 9.0f};
+    rigInstance->setRawControlValues(static_cast<const float*>(rawControls));
+
+    rigLogic->calculate(rigInstance);
+
+    ASSERT_EQ(rigInstance->getRawJointOutputs().size(), reader->getJointRowCount());
+    ASSERT_EQ(rigInstance->getBlendShapeOutputs().size(), reader->getBlendShapeChannelCount());
+    ASSERT_EQ(rigInstance->getAnimatedMapOutputs().size(), reader->getAnimatedMapCount());
+}
+
+TEST_F(RigLogicTest, AccessJointVariableAttributeIndices) {
+    for (std::uint16_t lod = 0u; lod < rigLogic->getLODCount(); ++lod) {
+        ASSERT_EQ(rigLogic->getJointVariableAttributeIndices(lod),
+                  (rl4::ConstArrayView<std::uint16_t>{rltests::decoded::jointVariableIndices[0ul][lod]}));
+    }
+}
+
+TEST_F(RigLogicTest, DumpStateThenRestore) {
+    float guiControls[] = {1.0f, 2.0f, 3.0f, 4.0f, 5.0f, 6.0f, 7.0f, 8.0f, 9.0f};
+
+    rltests::FakeStream dumpedState;
+    rigLogic->dump(&dumpedState);
+    dumpedState.seek(0);
+    auto cloneRigLogic = rl4::RigLogic::restore(&dumpedState);
+    auto cloneRigInstance = rl4::RigInstance::create(cloneRigLogic);
+
+    for (std::uint16_t lod = 0u; lod < rigLogic->getLODCount(); ++lod) {
+        rigInstance->setLOD(lod);
+        rigInstance->setGUIControlValues(static_cast<const float*>(guiControls));
+        rigLogic->mapGUIToRawControls(rigInstance);
+        rigLogic->calculate(rigInstance);
+
+        cloneRigInstance->setLOD(lod);
+        cloneRigInstance->setGUIControlValues(static_cast<const float*>(guiControls));
+        cloneRigLogic->mapGUIToRawControls(cloneRigInstance);
+        cloneRigLogic->calculate(cloneRigInstance);
+
+        auto origJointOutputs = rigInstance->getRawJointOutputs();
+        auto origBlendShapeOutputs = rigInstance->getBlendShapeOutputs();
+        auto origAnimatedMapOutputs = rigInstance->getAnimatedMapOutputs();
+
+        auto cloneJointOutputs = cloneRigInstance->getRawJointOutputs();
+        auto cloneBlendShapeOutputs = cloneRigInstance->getBlendShapeOutputs();
+        auto cloneAnimatedMapOutputs = cloneRigInstance->getAnimatedMapOutputs();
+
+        ASSERT_EQ(origJointOutputs, cloneJointOutputs);
+        ASSERT_EQ(origBlendShapeOutputs, cloneBlendShapeOutputs);
+        ASSERT_EQ(origAnimatedMapOutputs, cloneAnimatedMapOutputs);
+    }
+
+    rl4::RigInstance::destroy(cloneRigInstance);
+    rl4::RigLogic::destroy(cloneRigLogic);
+}
diff --git a/Engine/Plugins/Animation/RigLogic/Source/RigLogicLibTest/Private/rltests/transformation/TestTransformation.cpp b/Engine/Plugins/Animation/RigLogic/Source/RigLogicLibTest/Private/rltests/transformation/TestTransformation.cpp
new file mode 100644
--- /dev/null
+++ b/Engine/Plugins/Animation/RigLogic/Source/RigLogicLibTest/Private/rltests/transformation/TestTransformation.cpp
@@ -0,0 +1,72 @@
+#include "rltests/Defs.h"
+
+#include "riglogic/transformation/Transformation.h"
+
+#include <cassert>
+
+TEST(TransformationTest, AccessByIndex) {
+    std::vector<float> vec{0.0f, 1.0f, 2.0f, 3.0f, 4.0f, 5.0f, 6.0f, 7.0f, 8.0f,
+                           9.0f, 10.0f, 11.0f, 12.0f, 13.0f, 14.0f, 15.0f, 16.0f, 17.0f};
+
+    assert(vec.size() % 9ul == 0ul);
+
+    rl4::TransformationArrayView transformationView = {vec.data(), vec.size()};
+
+    auto translation0 = transformationView[0].getTranslation();
+    ASSERT_FLOAT_EQ(translation0.x, 0.0f);
+    ASSERT_FLOAT_EQ(translation0.y, 1.0f);
+    ASSERT_FLOAT_EQ(translation0.z, 2.0f);
+
+    auto rotation0 = transformationView[0].getRotation();
+    ASSERT_FLOAT_EQ(rotation0.x, 3.0f);
+    ASSERT_FLOAT_EQ(rotation0.y, 4.0f);
+    ASSERT_FLOAT_EQ(rotation0.z, 5.0f);
+
+    auto scale0 = transformationView[0].getScale();
+    ASSERT_FLOAT_EQ(scale0.x, 6.0f);
+    ASSERT_FLOAT_EQ(scale0.y, 7.0f);
+    ASSERT_FLOAT_EQ(scale0.z, 8.0f);
+
+    auto translation1 = transformationView[1].getTranslation();
+    ASSERT_FLOAT_EQ(translation1.x, 9.0f);
+    ASSERT_FLOAT_EQ(translation1.y, 10.0f);
+    ASSERT_FLOAT_EQ(translation1.z, 11.0f);
+
+    auto rotation1 = transformationView[1].getRotation();
+    ASSERT_FLOAT_EQ(rotation1.x, 12.0f);
+    ASSERT_FLOAT_EQ(rotation1.y, 13.0f);
+    ASSERT_FLOAT_EQ(rotation1.z, 14.0f);
+
+    auto scale1 = transformationView[1].getScale();
+    ASSERT_FLOAT_EQ(scale1.x, 15.0f);
+    ASSERT_FLOAT_EQ(scale1.y, 16.0f);
+    ASSERT_FLOAT_EQ(scale1.z, 17.0f);
+}
+
+TEST(TransformationTest, TransformationIterator) {
+    std::vector<float> vec{0.0f, 1.0f, 2.0f, 3.0f, 4.0f, 5.0f, 6.0f, 7.0f, 8.0f,
+                           9.0f, 10.0f, 11.0f, 12.0f, 13.0f, 14.0f, 15.0f, 16.0f, 17.0f};
+
+    assert(vec.size() % 9ul == 0ul);
+
+    rl4::TransformationArrayView transformationView = {vec.data(), vec.size()};
+
+    std::size_t index = 0ul;
+    for (auto transform : transformationView) {
+        auto translation = transform.getTranslation();
+        ASSERT_FLOAT_EQ(translation.x, vec[index]);
+        ASSERT_FLOAT_EQ(translation.y, vec[index + 1ul]);
+        ASSERT_FLOAT_EQ(translation.z, vec[index + 2ul]);
+
+        auto rotation = transform.getRotation();
+        ASSERT_FLOAT_EQ(rotation.x, vec[index + 3ul]);
+        ASSERT_FLOAT_EQ(rotation.y, vec[index + 4ul]);
+        ASSERT_FLOAT_EQ(rotation.z, vec[index + 5ul]);
+
+        auto scale = transform.getScale();
+        ASSERT_FLOAT_EQ(scale.x, vec[index + 6ul]);
+        ASSERT_FLOAT_EQ(scale.y, vec[index + 7ul]);
+        ASSERT_FLOAT_EQ(scale.z, vec[index + 8ul]);
+        index += 9ul;
+    }
+}
diff --git a/Engine/Plugins/Animation/RigLogic/Source/RigLogicLibTest/Private/tersetests/archives/binary/FakeStream.h b/Engine/Plugins/Animation/RigLogic/Source/RigLogicLibTest/Private/tersetests/archives/binary/FakeStream.h
new file mode 100644
--- /dev/null
+++ b/Engine/Plugins/Animation/RigLogic/Source/RigLogicLibTest/Private/tersetests/archives/binary/FakeStream.h
@@ -0,0 +1,61 @@
+#pragma once
+
+#ifdef _MSC_VER
+    #pragma warning(push)
+    #pragma warning(disable : 4365 4987)
+#endif
+#include <cassert>
+#include <cstring>
+#include <cstddef>
+#include <vector>
+#ifdef _MSC_VER
+    #pragma warning(pop)
+#endif
+
+namespace tersetests {
+
+struct FakeStream {
+
+    FakeStream() : position{} {
+    }
+
+    void open() {
+        position = 0ul;
+    }
+
+    void close() {
+        position = 0ul;
+    }
+
+    std::size_t tell() {
+        return position;
+    }
+
+    void seek(std::size_t position_) {
+        position = position_;
+    }
+
+    void read(char* buffer, std::size_t size) {
+        assert(position + size <= data.size());
+        std::memcpy(buffer, &data[position], size);
+        position += size;
+    }
+
+    void write(const char* buffer, std::size_t size) {
+        std::size_t available = data.size() - position;
+        if (available < size) {
+            data.resize(data.size() + (size - available));
+        }
+        std::memcpy(&data[position], buffer, size);
+        position += size;
+    }
+
+    std::size_t size() {
+        return data.size();
+    }
+
+    std::vector<char> data;
+    std::size_t position;
+};
+
+}  // namespace tersetests
diff --git a/Engine/Plugins/Animation/RigLogic/Source/RigLogicLibTest/Private/tersetests/archives/binary/TestInputArchive.cpp b/Engine/Plugins/Animation/RigLogic/Source/RigLogicLibTest/Private/tersetests/archives/binary/TestInputArchive.cpp
new file mode 100644
--- /dev/null
+++ b/Engine/Plugins/Animation/RigLogic/Source/RigLogicLibTest/Private/tersetests/archives/binary/TestInputArchive.cpp
@@ -0,0 +1,460 @@
+#include "tersetests/Defs.h"
+#include "tersetests/archives/binary/Common.h"
+#include "tersetests/archives/binary/FakeStream.h"
+
+#include "terse/utils/ArchiveOffset.h"
+#include "terse/archives/binary/InputArchive.h"
+
+#include <pma/TypeDefs.h>
+#include <pma/resources/AlignedMemoryResource.h>
+
+#ifdef _MSC_VER
+    #pragma warning(push)
+    #pragma warning(disable : 4365 4987)
+#endif
+#include <cstdint>
+#include <string>
+#include <vector>
+#ifdef _MSC_VER
+    #pragma warning(pop)
+#endif
+
+TEST(BinaryInputArchiveTest, StreamExtractionOperator) {
+    tersetests::FakeStream stream;
+    unsigned char bytes[] = {0x01, 0x02};
+    stream.write(reinterpret_cast<char*>(bytes), 2ul);
+    stream.seek(0ul);
+
+    terse::BinaryInputArchive<tersetests::FakeStream> archive(&stream);
+    char first;
+    char second;
+    archive >> first >> second;
+
+    ASSERT_EQ(first, 0x01);
+    ASSERT_EQ(second, 0x02);
+}
+
+TEST(BinaryInputArchiveTest, ArchiveOffset) {
+    tersetests::FakeStream stream;
+    unsigned char bytes[] = {
+        0x00, 0x00, 0x00, 0x19,  // layerStart offset
+        0x00, 0x00, 0x00, 0x19,  // first offset
+        0x00, 0x00, 0x00, 0x1d,  // second offset
+        0x00, 0x00, 0x04, 0xd2,  // a = 1234
+        0x78,  // b = 'x'
+        0x00, 0x00, 0x00, 0x04,  // c size = 4
+        0x61, 0x62, 0x63, 0x64,  // c = 'abcd'
+        0x00, 0x00, 0x01, 0x00,  // firstInt = 256
+        0x00, 0x00, 0x04, 0x00  // secondInt = 1024
+
+    };
+    stream.write(reinterpret_cast<char*>(bytes), 33ul);
+    stream.seek(0ul);
+
+    terse::BinaryInputArchive<tersetests::FakeStream> archive(&stream);
+    OffsetUtilizer dest;
+    archive >> dest;
+
+    ASSERT_EQ(dest.layerStart.position, 0ul);
+    ASSERT_EQ(dest.layerStart.proxy, &dest.layerMarker);
+    ASSERT_EQ(dest.layerStart.value, 25ul);
+
+    ASSERT_EQ(dest.first.position, 4ul);
+    ASSERT_EQ(dest.first.proxy, &dest.firstMarker);
+    ASSERT_EQ(dest.first.value, 25ul);
+
+    ASSERT_EQ(dest.second.position, 8ul);
+    ASSERT_EQ(dest.second.proxy, &dest.secondMarker);
+    ASSERT_EQ(dest.second.value, 29ul);
+}
+
+struct IgnoredSection {
+    std::uint32_t a;
+    std::string b;
+    std::uint16_t c;
+
+    IgnoredSection() : a{}, b{}, c{} {
+    }
+
+    template<class Archive>
+    void serialize(Archive&  /*unused*/) {
+        // This layer controls it's own serialization and decides
+        // during runtime for some reason not to deserialize it's data,
+        // although it is present in the stream
+    }
+
+};
+
+struct SectionSkipper {
+    terse::ArchiveOffset<std::uint32_t> afterValueOffset;
+    IgnoredSection ignored;
+    terse::ArchiveOffset<std::uint32_t>::Proxy afterValueMarker;
+    std::uint32_t afterValue;
+
+    SectionSkipper() : afterValueMarker{afterValueOffset} {
+    }
+
+    template<class Archive>
+    void serialize(Archive& archive) {
+        archive(afterValueOffset, ignored, afterValueMarker, afterValue);
+    }
+
+};
+
+TEST(BinaryInputArchiveTest, ArchiveOffsetRepositionsStream) {
+    tersetests::FakeStream stream;
+    unsigned char bytes[] = {
+        0x00, 0x00, 0x00, 0x12,  // afterValue offset
+        0x00, 0x00, 0x04, 0xd2,  // a = 1234
+        0x00, 0x00, 0x00, 0x04,  // b size = 4
+        0x61, 0x62, 0x63, 0x64,  // b = 'abcd'
+        0x01, 0x00,  // c = 256
+        0x00, 0x00, 0x04, 0x00  // afterValue = 1024
+
+    };
+    stream.write(reinterpret_cast<char*>(bytes), 22ul);
+    stream.seek(0ul);
+
+    terse::BinaryInputArchive<tersetests::FakeStream> archive(&stream);
+    SectionSkipper dest;
+    archive >> dest;
+
+    // Since `IgnoredSection` decided not to load it's data,
+    // we rely on `afterValueMarker` to reposition the stream
+    // so that it correctly points to the start of `afterValue`,
+    // after the aborted reading of `ignored`. If it works,
+    // `afterValue` should load it's expected value
+    ASSERT_EQ(dest.afterValue, 1024u);
+    // Sanity check to make sure ignored layer was really ignored
+    ASSERT_EQ(dest.ignored.a, 0u);
+    ASSERT_EQ(dest.ignored.b, "");
+    ASSERT_EQ(dest.ignored.c, 0u);
+}
+
+TEST(BinaryInputArchiveTest, PrimitiveTypeDeserialization) {
+    tersetests::FakeStream stream;
+    unsigned char bytes[] = {0x00, 0x00, 0x04, 0xd2};  // 1234
+    stream.write(reinterpret_cast<char*>(bytes), sizeof(std::int32_t));
+    stream.seek(0ul);
+
+    terse::BinaryInputArchive<tersetests::FakeStream> archive(&stream);
+    std::int32_t dest;
+    archive(dest);
+
+    ASSERT_EQ(dest, 1234);
+}
+
+TEST(BinaryInputArchiveTest, ComplexTypeDeserialization) {
+    tersetests::FakeStream stream;
+    unsigned char bytes[] = {
+        0x7f,  // 127
+        0xff,  // 255
+        0x7f, 0xff,  // 32767
+        0xff, 0xff,  // 65535
+        0x7f, 0xff, 0xff, 0xff,  // 2147483647
+        0xff, 0xff, 0xff, 0xff  // 4294967295
+    };
+    stream.write(reinterpret_cast<char*>(bytes), 14ul);
+    stream.seek(0ul);
+
+    terse::BinaryInputArchive<tersetests::FakeStream> archive(&stream);
+    ComplexType dest;
+    archive(dest);
+
+    ASSERT_EQ(dest.a, 127);
+    ASSERT_EQ(dest.b, 255u);
+    ASSERT_EQ(dest.c, 32767);
+    ASSERT_EQ(dest.d, 65535u);
+    ASSERT_EQ(dest.e, 2147483647);
+    ASSERT_EQ(dest.f, 4294967295u);
+}
+
+TEST(BinaryInputArchiveTest, StringDeserialization) {
+    tersetests::FakeStream stream;
+    unsigned char size[] = {0x00, 0x00, 0x00, 0x08};  // 8
+    unsigned char bytes[] = {
+        0x61,  // 'a'
+        0x62,  // 'b'
+        0x63,  // 'c'
+        0x64,  // 'd'
+        0x65,  // 'e'
+        0x66,  // 'f'
+        0x67,  // 'g'
+        0x68  // 'h'
+    };
+    stream.write(reinterpret_cast<char*>(size), 4ul);
+    stream.write(reinterpret_cast<char*>(bytes), 8ul);
+    stream.seek(0ul);
+
+    terse::BinaryInputArchive<tersetests::FakeStream> archive(&stream);
+    std::string dest;
+    archive(dest);
+
+    ASSERT_EQ(dest, "abcdefgh");
+}
+
+TEST(BinaryInputArchiveTest, VectorOfIntegerTypeDeserialization) {
+    tersetests::FakeStream stream;
+    unsigned char size[] = {0x00, 0x00, 0x00, 0x04};  // 4
+    unsigned char bytes[] = {
+        0x00, 0x00, 0x00, 0x10,  // 16
+        0x00, 0x00, 0x01, 0x00,  // 256
+        0x00, 0x00, 0x10, 0x00,  // 4096
+        0x00, 0x01, 0x00, 0x00  // 65536
+    };
+    stream.write(reinterpret_cast<char*>(size), 4ul);
+    stream.write(reinterpret_cast<char*>(bytes), 16ul);
+    stream.seek(0ul);
+
+    terse::BinaryInputArchive<tersetests::FakeStream> archive(&stream);
+    std::vector<std::int32_t> dest;
+    archive(dest);
+
+    std::vector<std::int32_t> expected{16, 256, 4096, 65536};
+    ASSERT_EQ(dest, expected);
+}
+
+TEST(BinaryInputArchiveTest, VectorOfFloatTypeDeserialization) {
+    tersetests::FakeStream stream;
+    unsigned char size[] = {0x00, 0x00, 0x00, 0x04};  // 4
+    unsigned char bytes[] = {
+        0x3f, 0x99, 0x99, 0x9a,  // 1.2f
+        0x40, 0x48, 0xf5, 0xc3,  // 3.14f
+        0x3d, 0x4c, 0xcc, 0xcd,  // 0.05f
+        0x3f, 0x3d, 0x70, 0xa4  // 0.74f
+    };
+    stream.write(reinterpret_cast<char*>(size), 4ul);
+    stream.write(reinterpret_cast<char*>(bytes), 16ul);
+    stream.seek(0ul);
+
+    terse::BinaryInputArchive<tersetests::FakeStream> archive(&stream);
+    std::vector<float> dest;
+    archive(dest);
+
+    std::vector<float> expected{1.2f, 3.14f, 0.05f, 0.74f};
+    ASSERT_EQ(dest, expected);
+}
+
+TEST(BinaryInputArchiveTest, VectorOfComplexTypeDeserialization) {
+    tersetests::FakeStream stream;
+    unsigned char size[] = {0x00, 0x00, 0x00, 0x02};  // 2
+    unsigned char bytes[] = {
+        0x7f,  // 127
+        0xff,  // 255
+        0x7f, 0xff,  // 32767
+        0xff, 0xff,  // 65535
+        0x7f, 0xff, 0xff, 0xff,  // 2147483647
+        0xff, 0xff, 0xff, 0xff,  // 4294967295
+        0x7f,  // 127
+        0xff,  // 255
+        0x7f, 0xff,  // 32767
+        0xff, 0xff,  // 65535
+        0x7f, 0xff, 0xff, 0xff,  // 2147483647
+        0xff, 0xff, 0xff, 0xff  // 4294967295
+    };
+    stream.write(reinterpret_cast<char*>(size), 4ul);
+    stream.write(reinterpret_cast<char*>(bytes), 28ul);
+    stream.seek(0ul);
+
+    terse::BinaryInputArchive<tersetests::FakeStream> archive(&stream);
+    std::vector<ComplexType> dest;
+    archive(dest);
+
+    ASSERT_EQ(dest.size(), 2ul);
+    for (const auto& ct : dest) {
+        ASSERT_EQ(ct.a, 127);
+        ASSERT_EQ(ct.b, 255u);
+        ASSERT_EQ(ct.c, 32767);
+        ASSERT_EQ(ct.d, 65535u);
+        ASSERT_EQ(ct.e, 2147483647);
+        ASSERT_EQ(ct.f, 4294967295u);
+    }
+}
+
+struct INeedAMemoryResource {
+    INeedAMemoryResource(pma::MemoryResource* memRes) : pMemRes{memRes}, number{} {
+    }
+
+    template<class Archive>
+    void serialize(Archive& archive) {
+        archive(number);
+    }
+
+    pma::MemoryResource* pMemRes;
+
+    int number;
+};
+
+TEST(BinaryInputArchiveTest, VectorOfTypeConstructibleFromMemoryResourceDeserialization) {
+    tersetests::FakeStream stream;
+    unsigned char size[] = {0x00, 0x00, 0x00, 0x01};  // 1
+    unsigned char bytes[] = {0x00, 0x00, 0x04, 0xd2};  // 1234
+    stream.write(reinterpret_cast<char*>(size), 4ul);
+    stream.write(reinterpret_cast<char*>(bytes), sizeof(std::int32_t));
+    stream.seek(0ul);
+
+    terse::BinaryInputArchive<tersetests::FakeStream> archive(&stream);
+    pma::AlignedMemoryResource memRes;
+    std::vector<INeedAMemoryResource,
+                pma::PolyAllocator<INeedAMemoryResource> > dest{pma::PolyAllocator<INeedAMemoryResource>{&memRes}};
+    archive(dest);
+
+    ASSERT_EQ(dest.size(), 1ul);
+    for (const auto& inamr : dest) {
+        ASSERT_EQ(inamr.number, 1234);
+        ASSERT_EQ(inamr.pMemRes, &memRes);
+    }
+}
+
+TEST(BinaryInputArchiveTest, VectorOfStringDeserialization) {
+    tersetests::FakeStream stream;
+    unsigned char vectorSize[] = {0x00, 0x00, 0x00, 0x02};  // 2
+    unsigned char stringSize[] = {0x00, 0x00, 0x00, 0x08};  // 8
+    unsigned char stringBytes[] = {
+        0x61,  // 'a'
+        0x62,  // 'b'
+        0x63,  // 'c'
+        0x64,  // 'd'
+        0x65,  // 'e'
+        0x66,  // 'f'
+        0x67,  // 'g'
+        0x68  // 'h'
+    };
+    stream.write(reinterpret_cast<char*>(vectorSize), 4ul);
+    // Write first string
+    stream.write(reinterpret_cast<char*>(stringSize), 4ul);
+    stream.write(reinterpret_cast<char*>(stringBytes), 8ul);
+    // Write second string
+    stream.write(reinterpret_cast<char*>(stringSize), 4ul);
+    stream.write(reinterpret_cast<char*>(stringBytes), 8ul);
+
+    stream.seek(0ul);
+
+    terse::BinaryInputArchive<tersetests::FakeStream> archive(&stream);
+    std::vector<std::string> dest;
+    archive(dest);
+
+    ASSERT_EQ(dest.size(), 2ul);
+    for (const auto& s : dest) {
+        ASSERT_EQ(s, "abcdefgh");
+    }
+}
+
+TEST(BinaryInputArchiveTest, VectorOfPairOfStringsDeserialization) {
+    tersetests::FakeStream stream;
+    unsigned char vectorSize[] = {0x00, 0x00, 0x00, 0x02};  // 2
+    unsigned char stringSize[] = {0x00, 0x00, 0x00, 0x04};  // 4
+    unsigned char keyBytes[] = {0x61, 0x62, 0x63, 0x64};  // "abcd"
+    unsigned char valueBytes[] = {0x65, 0x66, 0x67, 0x68};  // "efgh"
+    stream.write(reinterpret_cast<char*>(vectorSize), 4ul);
+    // Write key
+    stream.write(reinterpret_cast<char*>(stringSize), 4ul);
+    stream.write(reinterpret_cast<char*>(keyBytes), 4ul);
+    // Write value
+    stream.write(reinterpret_cast<char*>(stringSize), 4ul);
+    stream.write(reinterpret_cast<char*>(valueBytes), 4ul);
+    // Write key
+    stream.write(reinterpret_cast<char*>(stringSize), 4ul);
+    stream.write(reinterpret_cast<char*>(keyBytes), 4ul);
+    // Write value
+    stream.write(reinterpret_cast<char*>(stringSize), 4ul);
+    stream.write(reinterpret_cast<char*>(valueBytes), 4ul);
+
+    stream.seek(0ul);
+
+    terse::BinaryInputArchive<tersetests::FakeStream> archive(&stream);
+    std::vector<std::pair<std::string, std::string> > dest;
+    archive(dest);
+
+    ASSERT_EQ(dest.size(), 2ul);
+    for (const auto& s : dest) {
+        ASSERT_EQ(s.first, "abcd");
+        ASSERT_EQ(s.second, "efgh");
+    }
+}
+
+TEST(BinaryInputArchiveTest, VectorOfVectorOfPrimitiveTypeDeserialization) {
+    tersetests::FakeStream stream;
+    unsigned char outerVectorSize[] = {0x00, 0x00, 0x00, 0x02};  // 2
+    unsigned char innerVectorSize[] = {0x00, 0x00, 0x00, 0x04};  // 4
+    unsigned char innerVectorBytes[] = {
+        0x00, 0x04,  // 4
+        0x00, 0x10,  // 16
+        0x01, 0x00,  // 256
+        0x10, 0x00  // 4096
+    };
+    stream.write(reinterpret_cast<char*>(outerVectorSize), 4ul);
+    // Write first inner vector
+    stream.write(reinterpret_cast<char*>(innerVectorSize), 4ul);
+    stream.write(reinterpret_cast<char*>(innerVectorBytes), 8ul);
+    // Write second inner vector
+    stream.write(reinterpret_cast<char*>(innerVectorSize), 4ul);
+    stream.write(reinterpret_cast<char*>(innerVectorBytes), 8ul);
+    stream.seek(0ul);
+
+    pma::AlignedMemoryResource memRes;
+    terse::BinaryInputArchive<tersetests::FakeStream> archive(&stream);
+    pma::Vector<pma::Vector<std::uint16_t> > outer{&memRes};
+    archive(outer);
+
+    // Make sure values and structure is what's expected
+    ASSERT_EQ(outer.size(), 2ul);
+    for (const auto& inner : outer) {
+        ASSERT_EQ(inner.size(), 4ul);
+        ASSERT_EQ(inner[0], 4);
+        ASSERT_EQ(inner[1], 16);
+        ASSERT_EQ(inner[2], 256);
+        ASSERT_EQ(inner[3], 4096);
+    }
+
+    // Make sure the parent allocator is propagated to nested elements
+    ASSERT_EQ(&memRes, outer.get_allocator().getMemoryResource());
+    for (const auto& inner : outer) {
+        ASSERT_EQ(&memRes, inner.get_allocator().getMemoryResource());
+    }
+}
+
+TEST(BinaryInputArchiveTest, MixedNestedTypes) {
+    tersetests::FakeStream stream;
+    unsigned char bytes[] = {
+        0x00, 0x00, 0x00, 0x01,  // 1
+        0x00, 0x00, 0x00, 0x04,  // 4
+        0x72, 0x6f, 0x6f, 0x74,  // "root"
+        0x00, 0x00, 0x00, 0x02,  // 2
+        0x00, 0x00, 0x00, 0x03,  // 3
+        0x00, 0x01, 0x00, 0x02, 0x00, 0x03,  // 1, 2, 3
+        0x00, 0x00, 0x00, 0x03,  // 3
+        0x3c, 0x23, 0xd7, 0x0a, 0x3f, 0x7d, 0x70, 0xa4, 0x40, 0x48, 0xf5, 0xc3,  // 0.01f, 0.99f, 3.14f
+        0x00, 0x00, 0x00, 0x03,  // 3
+        0x00, 0x01, 0x00, 0x02, 0x00, 0x03,  // 1, 2, 3
+        0x00, 0x00, 0x00, 0x03,  // 3
+        0x3c, 0x23, 0xd7, 0x0a, 0x3f, 0x7d, 0x70, 0xa4, 0x40, 0x48, 0xf5, 0xc3  // 0.01f, 0.99f, 3.14f
+    };
+    stream.write(reinterpret_cast<char*>(bytes), 68ul);
+    stream.seek(0ul);
+
+    pma::AlignedMemoryResource memRes;
+    terse::BinaryInputArchive<tersetests::FakeStream> archive(&stream);
+    pma::Vector<Root> dest{&memRes};
+
+    archive(dest);
+
+    ASSERT_EQ(dest.size(), 1ul);
+    for (const auto& root : dest) {
+        ASSERT_EQ(root.name, "root");
+        ASSERT_EQ(root.children.size(), 2ul);
+        ASSERT_EQ(root.children.get_allocator().getMemoryResource(), &memRes);
+        for (const auto& child : root.children) {
+            ASSERT_EQ(child.integers.size(), 3ul);
+            ASSERT_EQ(child.integers.get_allocator().getMemoryResource(), &memRes);
+            ASSERT_EQ(child.integers[0], 1);
+            ASSERT_EQ(child.integers[1], 2);
+            ASSERT_EQ(child.integers[2], 3);
+            ASSERT_EQ(child.floats.size(), 3ul);
+            ASSERT_EQ(child.floats.get_allocator().getMemoryResource(), &memRes);
+            ASSERT_EQ(child.floats[0], 0.01f);
+            ASSERT_EQ(child.floats[1], 0.99f);
+            ASSERT_EQ(child.floats[2], 3.14f);
+        }
+    }
+}
diff --git a/Engine/Plugins/Animation/RigLogic/Source/RigLogicLibTest/Private/tersetests/archives/binary/TestOutputArchive.cpp b/Engine/Plugins/Animation/RigLogic/Source/RigLogicLibTest/Private/tersetests/archives/binary/TestOutputArchive.cpp
new file mode 100644
--- /dev/null
+++ b/Engine/Plugins/Animation/RigLogic/Source/RigLogicLibTest/Private/tersetests/archives/binary/TestOutputArchive.cpp
@@ -0,0 +1,324 @@
+#include "tersetests/Defs.h"
+#include "tersetests/archives/binary/Common.h"
+#include "tersetests/archives/binary/FakeStream.h"
+
+#include "terse/archives/binary/OutputArchive.h"
+
+#include <pma/TypeDefs.h>
+#include <pma/resources/AlignedMemoryResource.h>
+
+#ifdef _MSC_VER
+    #pragma warning(push)
+    #pragma warning(disable : 4365 4987)
+#endif
+#include <cstdint>
+#include <string>
+#include <vector>
+#ifdef _MSC_VER
+    #pragma warning(pop)
+#endif
+
+TEST(BinaryOutputArchiveTest, StreamInsertionOperator) {
+    std::int32_t source1 = 1234;
+    std::int32_t source2 = 5678;
+
+    tersetests::FakeStream stream;
+    terse::BinaryOutputArchive<tersetests::FakeStream> archive(&stream);
+    archive << source1 << source2;
+
+    unsigned char expected[] = {
+        0x00, 0x00, 0x04, 0xd2,  // 1234
+        0x00, 0x00, 0x16, 0x2e  // 5678
+    };
+    unsigned char bytes[8ul];
+    stream.seek(0ul);
+    stream.read(reinterpret_cast<char*>(bytes), sizeof(std::int32_t) * 2ul);
+
+    ASSERT_ELEMENTS_EQ(bytes, expected, 8ul);
+}
+
+TEST(BinaryOutputArchiveTest, ArchiveOffset) {
+    OffsetUtilizer source{};
+    source.a = 1234;
+    source.b = 'x';
+    source.c = "abcd";
+    source.firstInt = 256;
+    source.secondInt = 1024;
+
+    tersetests::FakeStream stream;
+    terse::BinaryOutputArchive<tersetests::FakeStream> archive(&stream);
+    archive << source;
+
+    unsigned char expected[] = {
+        0x00, 0x00, 0x00, 0x19,  // layerStart offset
+        0x00, 0x00, 0x00, 0x19,  // first offset
+        0x00, 0x00, 0x00, 0x1d,  // second offset
+        0x00, 0x00, 0x04, 0xd2,  // a = 1234
+        0x78,  // b = 'x'
+        0x00, 0x00, 0x00, 0x04,  // c size = 4
+        0x61, 0x62, 0x63, 0x64,  // c = 'abcd'
+        0x00, 0x00, 0x01, 0x00,  // firstInt = 256
+        0x00, 0x00, 0x04, 0x00  // secondInt = 1024
+
+    };
+    unsigned char bytes[33ul];
+    stream.seek(0ul);
+    stream.read(reinterpret_cast<char*>(bytes), sizeof(bytes));
+    ASSERT_ELEMENTS_EQ(bytes, expected, sizeof(bytes));
+}
+
+TEST(BinaryOutputArchiveTest, PrimitiveTypeSerialization) {
+    std::int32_t source = 1234;
+
+    tersetests::FakeStream stream;
+    terse::BinaryOutputArchive<tersetests::FakeStream> archive(&stream);
+    archive(source);
+
+    unsigned char expected[] = {0x00, 0x00, 0x04, 0xd2};  // 1234
+    unsigned char bytes[4ul];
+    stream.seek(0ul);
+    stream.read(reinterpret_cast<char*>(bytes), sizeof(std::int32_t));
+
+    ASSERT_ELEMENTS_EQ(bytes, expected, 4ul);
+}
+
+TEST(BinaryOutputArchiveTest, ComplexTypeSerialization) {
+    ComplexType source{127, 255, 32767, 65535, 2147483647, 4294967295};
+
+    tersetests::FakeStream stream;
+    terse::BinaryOutputArchive<tersetests::FakeStream> archive(&stream);
+    archive(source);
+
+    unsigned char expected[] = {
+        0x7f,  // 127
+        0xff,  // 255
+        0x7f, 0xff,  // 32767
+        0xff, 0xff,  // 65535
+        0x7f, 0xff, 0xff, 0xff,  // 2147483647
+        0xff, 0xff, 0xff, 0xff,  // 4294967295
+    };
+    unsigned char bytes[14ul];
+    stream.seek(0ul);
+    stream.read(reinterpret_cast<char*>(bytes), 14ul);
+
+    ASSERT_ELEMENTS_EQ(bytes, expected, 14ul);
+}
+
+TEST(BinaryOutputArchiveTest, StringSerialization) {
+    std::string source{"abcdefgh"};
+
+    tersetests::FakeStream stream;
+    terse::BinaryOutputArchive<tersetests::FakeStream> archive(&stream);
+    archive(source);
+
+    unsigned char expected[] = {
+        0x00, 0x00, 0x00, 0x08,  // 8
+        0x61,  // 'a'
+        0x62,  // 'b'
+        0x63,  // 'c'
+        0x64,  // 'd'
+        0x65,  // 'e'
+        0x66,  // 'f'
+        0x67,  // 'g'
+        0x68  // 'h'
+    };
+    unsigned char bytes[12ul];
+    stream.seek(0ul);
+    stream.read(reinterpret_cast<char*>(bytes), 12ul);
+
+    ASSERT_ELEMENTS_EQ(bytes, expected, 12ul);
+}
+
+TEST(BinaryOutputArchiveTest, VectorOfIntegerTypeSerialization) {
+    std::vector<std::int32_t> source{16, 256, 4096, 65536};
+
+    tersetests::FakeStream stream;
+    terse::BinaryOutputArchive<tersetests::FakeStream> archive(&stream);
+    archive(source);
+
+    unsigned char expected[] = {
+        0x00, 0x00, 0x00, 0x04,  // 4
+        0x00, 0x00, 0x00, 0x10,  // 16
+        0x00, 0x00, 0x01, 0x00,  // 256
+        0x00, 0x00, 0x10, 0x00,  // 4096
+        0x00, 0x01, 0x00, 0x00  // 65536
+    };
+    unsigned char bytes[20ul];
+    stream.seek(0ul);
+    stream.read(reinterpret_cast<char*>(bytes), 20ul);
+
+    ASSERT_ELEMENTS_EQ(bytes, expected, 20ul);
+}
+
+TEST(BinaryOutputArchiveTest, VectorOfFloatTypeSerialization) {
+    std::vector<float> source{1.2f, 3.14f, 0.05f, 0.74f};
+
+    tersetests::FakeStream stream;
+    terse::BinaryOutputArchive<tersetests::FakeStream> archive(&stream);
+    archive(source);
+
+    unsigned char expected[] = {
+        0x00, 0x00, 0x00, 0x04,  // 4
+        0x3f, 0x99, 0x99, 0x9a,  // 1.2f
+        0x40, 0x48, 0xf5, 0xc3,  // 3.14f
+        0x3d, 0x4c, 0xcc, 0xcd,  // 0.05f
+        0x3f, 0x3d, 0x70, 0xa4  // 0.74f
+    };
+    unsigned char bytes[20ul];
+    stream.seek(0ul);
+    stream.read(reinterpret_cast<char*>(bytes), 20ul);
+
+    ASSERT_ELEMENTS_EQ(bytes, expected, 20ul);
+}
+
+TEST(BinaryOutputArchiveTest, VectorOfComplexTypeSerialization) {
+    std::vector<ComplexType> source;
+    source.push_back({127, 255, 32767, 65535, 2147483647, 4294967295});
+    source.push_back({127, 255, 32767, 65535, 2147483647, 4294967295});
+
+    tersetests::FakeStream stream;
+    terse::BinaryOutputArchive<tersetests::FakeStream> archive(&stream);
+    archive(source);
+
+    unsigned char expected[] = {
+        0x00, 0x00, 0x00, 0x02,  // 2
+        0x7f,  // 127
+        0xff,  // 255
+        0x7f, 0xff,  // 32767
+        0xff, 0xff,  // 65535
+        0x7f, 0xff, 0xff, 0xff,  // 2147483647
+        0xff, 0xff, 0xff, 0xff,  // 4294967295
+        0x7f,  // 127
+        0xff,  // 255
+        0x7f, 0xff,  // 32767
+        0xff, 0xff,  // 65535
+        0x7f, 0xff, 0xff, 0xff,  // 2147483647
+        0xff, 0xff, 0xff, 0xff  // 4294967295
+    };
+    unsigned char bytes[32ul];
+    stream.seek(0ul);
+    stream.read(reinterpret_cast<char*>(bytes), 32ul);
+
+    ASSERT_ELEMENTS_EQ(bytes, expected, 32ul);
+}
+
+TEST(BinaryOutputArchiveTest, VectorOfStringSerialization) {
+    std::vector<std::string> source{"abcdefgh", "abcdefgh"};
+
+    tersetests::FakeStream stream;
+    terse::BinaryOutputArchive<tersetests::FakeStream> archive(&stream);
+    archive(source);
+
+    unsigned char expected[] = {
+        0x00, 0x00, 0x00, 0x02,  // 2
+        0x00, 0x00, 0x00, 0x08,  // 8
+        0x61, 0x62, 0x63, 0x64, 0x65, 0x66, 0x67, 0x68,  // "abcdefgh"
+        0x00, 0x00, 0x00, 0x08,  // 8
+        0x61, 0x62, 0x63, 0x64, 0x65, 0x66, 0x67, 0x68  // "abcdefgh"
+    };
+    unsigned char bytes[28ul];
+    stream.seek(0ul);
+    stream.read(reinterpret_cast<char*>(bytes), 28ul);
+
+    ASSERT_ELEMENTS_EQ(bytes, expected, 28ul);
+}
+
+TEST(BinaryOutputArchiveTest, VectorOfPairOfStringsSerialization) {
+    using KeyValue = std::pair<std::string, std::string>;
+    std::vector<KeyValue> source{{"abcd", "efgh"}, {"abcd", "efgh"}};
+
+    tersetests::FakeStream stream;
+    terse::BinaryOutputArchive<tersetests::FakeStream> archive(&stream);
+    archive(source);
+
+    unsigned char expected[] = {
+        0x00, 0x00, 0x00, 0x02,  // 2
+        0x00, 0x00, 0x00, 0x04,  // 4
+        0x61, 0x62, 0x63, 0x64,  // "abcd"
+        0x00, 0x00, 0x00, 0x04,  // 4
+        0x65, 0x66, 0x67, 0x68,  // "efgh"
+        0x00, 0x00, 0x00, 0x04,  // 4
+        0x61, 0x62, 0x63, 0x64,  // "abcd"
+        0x00, 0x00, 0x00, 0x04,  // 4
+        0x65, 0x66, 0x67, 0x68,  // "efgh"
+    };
+    unsigned char bytes[36ul];
+    stream.seek(0ul);
+    stream.read(reinterpret_cast<char*>(bytes), 36ul);
+
+    ASSERT_ELEMENTS_EQ(bytes, expected, 36ul);
+}
+
+TEST(BinaryOutputArchiveTest, VectorOfVectorOfPrimitiveTypeSerialization) {
+    pma::Vector<std::uint16_t> values{4, 16, 256, 4096};
+    std::vector<pma::Vector<std::uint16_t> > source{values, values};
+
+    tersetests::FakeStream stream;
+    terse::BinaryOutputArchive<tersetests::FakeStream> archive(&stream);
+    archive(source);
+
+    unsigned char expected[] = {
+        0x00, 0x00, 0x00, 0x02,  // 2
+        0x00, 0x00, 0x00, 0x04,  // 4
+        0x00, 0x04, 0x00, 0x10, 0x01, 0x00, 0x10, 0x00,  // 4, 16, 256, 4096
+        0x00, 0x00, 0x00, 0x04,  // 4
+        0x00, 0x04, 0x00, 0x10, 0x01, 0x00, 0x10, 0x00  // 4, 16, 256, 4096
+    };
+    unsigned char bytes[28ul];
+    stream.seek(0ul);
+    stream.read(reinterpret_cast<char*>(bytes), 28ul);
+
+    ASSERT_ELEMENTS_EQ(bytes, expected, 28ul);
+}
+
+TEST(BinaryOutputArchiveTest, MixedNestedTypes) {
+    pma::AlignedMemoryResource memRes;
+
+    Root root{&memRes};
+    root.name = "root";
+
+    Child child1{&memRes};
+    child1.integers.push_back(1);
+    child1.integers.push_back(2);
+    child1.integers.push_back(3);
+    child1.floats.push_back(0.01f);
+    child1.floats.push_back(0.99f);
+    child1.floats.push_back(3.14f);
+    Child child2{&memRes};
+    child2.integers.push_back(1);
+    child2.integers.push_back(2);
+    child2.integers.push_back(3);
+    child2.floats.push_back(0.01f);
+    child2.floats.push_back(0.99f);
+    child2.floats.push_back(3.14f);
+
+    root.children.push_back(std::move(child1));
+    root.children.push_back(std::move(child2));
+
+    pma::Vector<Root> source{&memRes};
+    source.push_back(std::move(root));
+
+    tersetests::FakeStream stream;
+    terse::BinaryOutputArchive<tersetests::FakeStream> archive(&stream);
+    archive(source);
+
+    unsigned char expected[] = {
+        0x00, 0x00, 0x00, 0x01,  // 1
+        0x00, 0x00, 0x00, 0x04,  // 4
+        0x72, 0x6f, 0x6f, 0x74,  // "root"
+        0x00, 0x00, 0x00, 0x02,  // 2
+        0x00, 0x00, 0x00, 0x03,  // 3
+        0x00, 0x01, 0x00, 0x02, 0x00, 0x03,  // 1, 2, 3
+        0x00, 0x00, 0x00, 0x03,  // 3
+        0x3c, 0x23, 0xd7, 0x0a, 0x3f, 0x7d, 0x70, 0xa4, 0x40, 0x48, 0xf5, 0xc3,  // 0.01f, 0.99f, 3.14f
+        0x00, 0x00, 0x00, 0x03,  // 3
+        0x00, 0x01, 0x00, 0x02, 0x00, 0x03,  // 1, 2, 3
+        0x00, 0x00, 0x00, 0x03,  // 3
+        0x3c, 0x23, 0xd7, 0x0a, 0x3f, 0x7d, 0x70, 0xa4, 0x40, 0x48, 0xf5, 0xc3  // 0.01f, 0.99f, 3.14f
+    };
+    unsigned char bytes[68ul];
+    stream.seek(0ul);
+    stream.read(reinterpret_cast<char*>(bytes), 68ul);
+
+    ASSERT_ELEMENTS_EQ(bytes, expected, 68ul);
+}
diff --git a/Engine/Plugins/Animation/RigLogic/Source/RigLogicLibTest/RigLogicLibTest.Build.cs b/Engine/Plugins/Animation/RigLogic/Source/RigLogicLibTest/RigLogicLibTest.Build.cs
new file mode 100644
--- /dev/null
+++ b/Engine/Plugins/Animation/RigLogic/Source/RigLogicLibTest/RigLogicLibTest.Build.cs
@@ -0,0 +1,53 @@
+// Copyright 1998-2019 Epic Games, Inc. All Rights Reserved.
+
+using UnrealBuildTool;
+using System;
+using System.IO;
+using Tools.DotNETCommon;
+
+public class RigLogicLibTest : ModuleRules
+{
+    public RigLogicLibTest(ReadOnlyTargetRules Target) : base(Target)
+    {
+        PrivateDefinitions.Add("RIGLOGIC_MODULE_DISCARD");
+        if (Target.Platform == UnrealTargetPlatform.Win32 ||
+            Target.Platform == UnrealTargetPlatform.Win64 ||
+            Target.Platform == UnrealTargetPlatform.Linux ||
+            Target.Platform == UnrealTargetPlatform.Mac)
+        {
+            PrivateDefinitions.Add("RL_BUILD_WITH_SSE=1");
+        }
+
+        string RigLogicLibPath = Path.GetFullPath(Path.Combine(ModuleDirectory, "../RigLogicLib"));
+        ConditionalAddModuleDirectory(new DirectoryReference(RigLogicLibPath));
+
+        PublicDependencyModuleNames.AddRange(
+            new string[]
+            {
+                "Core",
+                "GoogleTest"
+            }
+        );
+
+        if (Target.Platform == UnrealTargetPlatform.Win32 || Target.Platform == UnrealTargetPlatform.Win64 || Target.Platform == UnrealTargetPlatform.XboxOne)
+        {
+            PublicDefinitions.Add("GTEST_OS_WINDOWS=1");
+        }
+        else if (Target.Platform == UnrealTargetPlatform.Mac)
+        {
+            PublicDefinitions.Add("GTEST_OS_MAC=1");
+        }
+        else if (Target.Platform == UnrealTargetPlatform.IOS || Target.Platform == UnrealTargetPlatform.TVOS)
+        {
+            PublicDefinitions.Add("GTEST_OS_IOS=1");
+        }
+        else if (Target.Platform == UnrealTargetPlatform.Android || Target.Platform == UnrealTargetPlatform.Lumin)
+        {
+            PublicDefinitions.Add("GTEST_OS_LINUX_ANDROID=1");
+        }
+        else if (Target.IsInPlatformGroup(UnrealPlatformGroup.Unix) || Target.Platform == UnrealTargetPlatform.PS4)
+        {
+            PublicDefinitions.Add("GTEST_OS_LINUX=1");
+        }
+    }
+}
diff --git a/Engine/Plugins/Animation/RigLogic/Source/RigLogicModule/Private/DNAAsset.cpp b/Engine/Plugins/Animation/RigLogic/Source/RigLogicModule/Private/DNAAsset.cpp
new file mode 100644
--- /dev/null
+++ b/Engine/Plugins/Animation/RigLogic/Source/RigLogicModule/Private/DNAAsset.cpp
@@ -0,0 +1,135 @@
+// Copyright 1998-2019 Epic Games, Inc. All Rights Reserved.
+
+#include "DNAAsset.h"
+#include "DNAUtils.h"
+
+#include "DNAReaderAdapter.h"
+#include "FMemoryResource.h"
+#include "RigLogicMemoryStream.h"
+#include "SkelMeshDNAReader.h"
+
+#if WITH_EDITORONLY_DATA
+    #include "EditorFramework/AssetImportData.h"
+#endif
+#include "Engine/AssetUserData.h"
+#include "Runtime/Core/Public/Serialization/BufferArchive.h"
+#include "Runtime/Core/Public/Serialization/MemoryReader.h"
+#include "Runtime/Core/Public/Misc/Paths.h"
+#include "Runtime/Core/Public/Misc/FileHelper.h"
+#include "Runtime/CoreUObject/Public/UObject/UObjectGlobals.h"
+
+#include "riglogic/RigLogic.h"
+
+DEFINE_LOG_CATEGORY(LogDNAAsset);
+
+static constexpr uint32 AVG_BEHAVIOR_SIZE = 5 * 1024 * 1024;
+static constexpr uint32 AVG_GEOMETRY_SIZE = 150 * 1024 * 1024;
+
+void UDNAAsset::WriteToBuffer( IDNAReader* SourcePartialDNAReader, TArray<uint8>* WriteBuffer, uint32 PredictedSize)
+{	
+	WriteBuffer->Empty();
+	WriteBuffer->Reserve(PredictedSize);
+	// Use that stream to write data into TArray, so UE can serialize it with this asset
+	FRigLogicMemoryStream PartialDataWriteStream(WriteBuffer);
+	rl4::ScopedPtr<dna::StreamWriter> PartialStreamWriter = rl4::makeScoped<dna::StreamWriter>(&PartialDataWriteStream, FMemoryResource::Instance()); //where to write
+	PartialStreamWriter->setFrom(SourcePartialDNAReader->Unwrap()); //from where to read
+	PartialStreamWriter->write(); //DNAData TArray now holds and serializes behavior data
+}
+
+bool UDNAAsset::Init(const FString DNAFilename)
+{
+	if (!rl4::Status::isOk()) {
+		UE_LOG(LogDNAAsset, Warning, TEXT("%s"), ANSI_TO_TCHAR(rl4::Status::get().message));
+	}
+
+	this->DNAFileName = DNAFilename; //memorize for re-import
+	
+	if (!FPaths::FileExists(DNAFilename))
+	{
+		UE_LOG(LogDNAAsset, Error, TEXT("DNA file %s doesn't exist!"), *DNAFilename);
+		return false;
+	}
+	
+	// Temporary buffer for the DNA file
+	TArray<uint8> TempFileBuffer;
+	
+	if (!FFileHelper::LoadFileToArray(TempFileBuffer, *DNAFilename)) //load entire DNA file into the array
+	{
+		UE_LOG(LogDNAAsset, Error, TEXT("Couldn't read DNA file %s!"), *DNAFilename);
+		return false;
+	}
+	
+	// Load run-time data (behavior) from whole-DNA buffer into BehaviorStreamReader
+	TSharedPtr<IDNAReader> DNAReaderBehavior = ReadDNALayerFromBuffer(&TempFileBuffer, EDNADataLayer::Behavior, 0u); //0u = MaxLOD
+
+	if (!DNAReaderBehavior.IsValid())
+	{
+		return false;
+	}
+
+	BehaviorStreamReader = DNAReaderBehavior;
+	WriteToBuffer(DNAReaderBehavior.Get(), &BehaviorData, AVG_BEHAVIOR_SIZE);
+
+#if WITH_EDITOR
+	//We use geometry part of the data in MHC only (for updating the SkeletalMesh with
+	//result of GeneSplicer), so we can drop geometry part when cooking for runtime
+	TSharedPtr<IDNAReader> DNAReaderGeometry = ReadDNALayerFromBuffer(&TempFileBuffer, EDNADataLayer::Geometry, 0u); //0u = MaxLOD
+	GeometryStreamReader = DNAReaderGeometry;
+	WriteToBuffer(DNAReaderGeometry.Get(), &DesignTimeData, AVG_GEOMETRY_SIZE);
+	//Note: in future, we will want to load geometry data in-game too 
+	//to enable GeneSplicer to read geometry directly from SkeletalMeshes, as
+	//a way to save memory, as on consoles the "database" will be exactly the set of characters
+	//used in the game
+#endif
+
+	return true;
+}
+
+void UDNAAsset::PostLoad()
+{
+	Super::PostLoad();
+
+	if (BehaviorData.Num() == 0)
+	{
+		return;
+	}
+
+	BehaviorStreamReader = ReadDNALayerFromBuffer(&BehaviorData, EDNADataLayer::Behavior, 0u); //0u = max LOD
+
+#if WITH_EDITORONLY_DATA
+	if (DesignTimeData.Num() == 0)
+	{
+		return;
+	}
+	GeometryStreamReader = ReadDNALayerFromBuffer(&DesignTimeData, EDNADataLayer::Geometry, 0u); //0u = max LOD
+#endif
+}
+
+TSharedPtr<IDNAReader> UDNAAsset::CopyDNALayer(const IDNAReader* Source, EDNADataLayer DNADataLayer, uint32 PredictedSize )
+{
+	// To avoid lots of reallocations in `FRigLogicMemoryStream`, reserve an approximate size
+	// that we know would cause at most one reallocation in the worst case (but none for the average DNA)
+	
+	TArray<uint8> MemoryBuffer;
+	MemoryBuffer.Reserve(PredictedSize);
+
+	FRigLogicMemoryStream MemoryStream(&MemoryBuffer);
+	auto DNAWriter = rl4::makeScoped<dna::StreamWriter>(&MemoryStream, FMemoryResource::Instance());
+	DNAWriter->setFrom(Source->Unwrap(), static_cast<dna::DataLayer>(DNADataLayer));
+	DNAWriter->write();
+
+	MemoryStream.seek(0ul);
+
+	return ReadDNALayerFromBuffer(&MemoryBuffer, DNADataLayer);
+}
+
+void UDNAAsset::SetBehaviorReader(const TSharedPtr<IDNAReader> SourceDNAReader)
+{
+
+	BehaviorStreamReader = CopyDNALayer(SourceDNAReader.Get(), EDNADataLayer::Behavior, AVG_BEHAVIOR_SIZE);
+}
+
+void UDNAAsset::SetGeometryReader(const TSharedPtr<IDNAReader> SourceDNAReader)
+{
+	GeometryStreamReader = CopyDNALayer(SourceDNAReader.Get(), EDNADataLayer::Geometry, AVG_GEOMETRY_SIZE);
+}
\ No newline at end of file
diff --git a/Engine/Plugins/Animation/RigLogic/Source/RigLogicModule/Private/DNAReader.cpp b/Engine/Plugins/Animation/RigLogic/Source/RigLogicModule/Private/DNAReader.cpp
new file mode 100644
--- /dev/null
+++ b/Engine/Plugins/Animation/RigLogic/Source/RigLogicModule/Private/DNAReader.cpp
@@ -0,0 +1,5 @@
+// Copyright 1998-2019 Epic Games, Inc. All Rights Reserved.
+
+#include "DNAReader.h"
+
+DEFINE_LOG_CATEGORY(LogDNAReader);
diff --git a/Engine/Plugins/Animation/RigLogic/Source/RigLogicModule/Private/DNAToSkelMeshMap.cpp b/Engine/Plugins/Animation/RigLogic/Source/RigLogicModule/Private/DNAToSkelMeshMap.cpp
new file mode 100644
--- /dev/null
+++ b/Engine/Plugins/Animation/RigLogic/Source/RigLogicModule/Private/DNAToSkelMeshMap.cpp
@@ -0,0 +1,279 @@
+// Copyright 1998-2019 Epic Games, Inc. All Rights Reserved.
+
+#include "DNAToSkelMeshMap.h"
+
+#include "DNAAsset.h"
+#include "SkelMeshDNAReader.h"
+
+#include "Runtime/Core/Public/Serialization/BufferArchive.h"
+#include "Runtime/Core/Public/Serialization/MemoryReader.h"
+#include "Runtime/Core/Public/Misc/Paths.h"
+#include "Runtime/Core/Public/Misc/FileHelper.h"
+#include "Runtime/CoreUObject/Public/UObject/UObjectGlobals.h"
+#include "Runtime/Engine/Public/Rendering/SkeletalMeshModel.h"
+#include "Runtime/Engine/Public/Rendering/SkeletalMeshLODModel.h"
+#include "Runtime/Engine/Classes/Animation/SmartName.h"
+
+#include "riglogic/RigLogic.h"
+
+DEFINE_LOG_CATEGORY(LogDNAToSkelMeshMap);
+
+#if WITH_EDITORONLY_DATA
+/** Creates mappings between source DNA and target SkelMesh. */
+void FDNAToSkelMeshMap::InitBaseMesh(IDNAReader* SourceDNAReader, USkeletalMesh* TargetSkeletalMesh)
+{
+	TargetSkelMesh = TargetSkeletalMesh; //maps are created for specific SkelMesh, so we can memorize it for clarity
+
+	InitVertexMap(SourceDNAReader);
+
+	//NOTE: this is just init for updating base mesh
+	// before updating joints and morph targets, MapJoints and MapMorphTarget must be called respectively
+}
+#endif
+
+
+/** Creates and initiates DNAReader and creates mapping between DNA and provided SkelMesh. */
+bool FDNAToSkelMeshMap::InitFromDNAAsset(USkeletalMesh* SkelMesh)
+{
+#if WITH_EDITORONLY_DATA
+	// Fetch DNAAsset from SkelMesh and use its StreamReader
+	UAssetUserData* UserData = SkelMesh->GetAssetUserDataOfClass(UDNAAsset::StaticClass());
+	if (UserData == nullptr)
+	{
+		FString PathName = GetPathNameSafe(SkelMesh);
+		UE_LOG(LogDNAToSkelMeshMap, Warning, TEXT("Could not find DNAAsset user data for %s"), *PathName);
+		return false;
+	}
+	UDNAAsset* DNAAsset = Cast<UDNAAsset>(UserData);
+	FSkelMeshDNAReader SkelMeshDNAReader{ DNAAsset };
+	InitBaseMesh(&SkelMeshDNAReader, SkelMesh);
+
+	return true;
+#else
+	return false;
+#endif //WITH_EDITORONLY_DATA
+}
+
+#if WITH_EDITORONLY_DATA
+/** Maps vertices and meshes from SkelMesh to DNA. */
+void FDNAToSkelMeshMap::InitVertexMap(IDNAReader* DNAReader)
+{
+	double StartTime = FPlatformTime::Seconds();
+
+	ImportVtxToDNAMeshIndex.Empty();
+	ImportVtxToDNAVtxIndex.Empty();
+
+	FSkeletalMeshModel* ImportedModel = TargetSkelMesh->GetImportedModel();
+	int32 NumEngineLods = ImportedModel->LODModels.Num();
+	ImportVtxToDNAMeshIndex.AddZeroed(NumEngineLods);
+	ImportVtxToDNAVtxIndex.AddZeroed(NumEngineLods);
+	for (int32 i = 0; i < NumEngineLods; i++)
+	{
+		FSkeletalMeshLODModel& LODModel = ImportedModel->LODModels[i];
+		ImportVtxToDNAMeshIndex[i].AddZeroed(LODModel.MeshToImportVertexMap.Num());
+		ImportVtxToDNAVtxIndex[i].AddZeroed(LODModel.MeshToImportVertexMap.Num());
+	}
+
+	TArray<TArray<int32>> MeshStartIndices; // For each LOD, store the start index offset of each mesh.
+	TArray<int32> TotalVertices; // Counts total indices for each mesh in DNA, per LOD.
+	uint16 NumLODs = DNAReader->GetLODCount();
+	MeshStartIndices.AddZeroed(NumLODs);
+	TotalVertices.AddZeroed(NumLODs);
+	ImportDNAVtxToUEVtxIndex.AddZeroed(NumLODs);
+	// First find start indices for each mesh and for each LOD in DNA.
+	for (uint16 LODIndex = 0; LODIndex < NumLODs; LODIndex++)
+	{
+		const TArrayView<const uint16> LODMeshIndices = DNAReader->GetMeshIndicesForLOD(LODIndex);
+
+		for (size_t MeshIndex = 0; MeshIndex < LODMeshIndices.Num(); MeshIndex++)
+		{
+			MeshStartIndices[LODIndex].Add(TotalVertices[LODIndex]);
+			const int32 VertexCount = DNAReader->GetVertexPositionCount(LODMeshIndices[MeshIndex]);
+			TotalVertices[LODIndex] += VertexCount;
+		}
+		MeshStartIndices[LODIndex].Add(TotalVertices[LODIndex]);
+		ImportDNAVtxToUEVtxIndex[LODIndex].Init(INDEX_NONE, TotalVertices[LODIndex]);
+	}
+	// Then for each LOD and each Vertex in LODModel Vertex Map find corresponding DNA Vertex.
+	int32 ImportedVertexCount = 0;
+	for (int32 LODIndex = 0; LODIndex < ImportedModel->LODModels.Num(); LODIndex++)
+	{
+		const FSkeletalMeshLODModel& LODModel = ImportedModel->LODModels[LODIndex];
+		const TArrayView<const uint16> LODMeshIndices = DNAReader->GetMeshIndicesForLOD(LODIndex);
+		const int32 StartIndicesCount = MeshStartIndices[LODIndex].Num();
+
+		int32 LODMeshVtxCount = LODModel.MeshToImportVertexMap.Num();
+		for (int32 LODMeshVtxIndex = 0; LODMeshVtxIndex < LODMeshVtxCount; LODMeshVtxIndex++)
+		{
+			int32 FbxVertexIndex = LODModel.MeshToImportVertexMap[LODMeshVtxIndex];
+			bool bFound = false;
+			// For each LODModel Vertex try to find the corresponding DNA Vertex Index by checking prepared list of Start Indices.
+			for (int k = 1; k < StartIndicesCount; k++)
+			{
+				if (FbxVertexIndex < MeshStartIndices[LODIndex][k])
+				{
+					bFound = true;
+					ImportVtxToDNAMeshIndex[LODIndex][LODMeshVtxIndex] = LODMeshIndices[k - 1];
+					const int32 DNAVertexIndex = FbxVertexIndex - MeshStartIndices[LODIndex][k - 1];
+					ImportVtxToDNAVtxIndex[LODIndex][LODMeshVtxIndex] = DNAVertexIndex;
+					ImportDNAVtxToUEVtxIndex[LODIndex][DNAVertexIndex] = LODMeshVtxIndex;
+					break;
+				}
+			}
+			// If no matching DNA Vertex Index is found mark LODModel Vertex Index with -1.
+			if (!bFound)
+			{
+#ifdef DEBUG
+				UE_LOG(LogDNAToSkelMeshMap, Warning, TEXT("Not sorted fbx vertex found %d"), FbxVertexIndex);
+#endif
+				ImportVtxToDNAMeshIndex[LODIndex][LODMeshVtxIndex] = INDEX_NONE;
+				ImportVtxToDNAVtxIndex[LODIndex][LODMeshVtxIndex] = INDEX_NONE;
+			}
+		}
+	}
+
+	// Find and map overlapping indices.
+	OverlappingVertices.AddZeroed(NumEngineLods);
+	for (int32 i = 0; i < NumEngineLods; i++)
+	{
+		FSkeletalMeshLODModel& LODModel = ImportedModel->LODModels[i];
+		const int32 NumSections = LODModel.Sections.Num();
+		OverlappingVertices[i].AddZeroed(NumSections);
+		for (int32 SectionIdx = 0; SectionIdx < NumSections; SectionIdx++)
+		{
+			const FSkelMeshSection& Section = LODModel.Sections[SectionIdx];
+			const int32 NumVertices = Section.GetNumVertices();
+
+			TArray<bool> VerticesCovered;
+			VerticesCovered.AddZeroed(NumVertices);
+
+			OverlappingVertices[i][SectionIdx].AddZeroed(NumVertices);
+			for (int32 VertexIndex = 0; VertexIndex < NumVertices; VertexIndex++)
+			{
+				const int32 VertexBufferIndex = Section.GetVertexBufferIndex() + VertexIndex;
+				const int32 DNAVertexIndex = ImportVtxToDNAVtxIndex[i][VertexBufferIndex];
+				if (DNAVertexIndex >= 0)
+				{
+					VerticesCovered[VertexIndex] = true;
+
+					if (Section.OverlappingVertices.Contains(VertexIndex))
+					{
+						const int32 OverlappingCount = Section.OverlappingVertices[VertexIndex].Num();
+						for (int32 OverlappingIndex = 0; OverlappingIndex < OverlappingCount; ++OverlappingIndex)
+						{
+							const int32 OverlappingVertexIndex = Section.OverlappingVertices[VertexIndex][OverlappingIndex];
+							if (!VerticesCovered[OverlappingVertexIndex]) // No need to update the same vertex twice.
+							{
+								OverlappingVertices[i][SectionIdx][VertexIndex].Add(OverlappingVertexIndex);
+							}
+						}
+					}
+				}
+			}
+		}
+	}
+
+	double TimeToMap = FPlatformTime::Seconds();
+	TimeToMap -= StartTime;
+
+	UE_LOG(LogDNAToSkelMeshMap, Log, TEXT("	InitVertexMap:	%.6f"), TimeToMap * 1000.0f);
+}
+
+/** Makes a map of all Joints from DNA to Bones in Reference Skeleton.
+ * Additionally makes a map of all Skinned joints. */
+void FDNAToSkelMeshMap::MapJoints(IDNAReader* DNAReader)
+{
+
+	double StartTime = FPlatformTime::Seconds();
+
+	FReferenceSkeleton* RefSkeleton = &TargetSkelMesh->RefSkeleton;
+	uint32 JointCount = DNAReader->GetJointCount();
+
+	// Map Joints to Bones.
+	RLJointToUEBoneIndices.Empty();
+	RLJointToUEBoneIndices.Reserve(JointCount);
+	for (uint32 JntIndex = 0; JntIndex < JointCount; ++JntIndex)
+	{
+		const FString boneNameFStr = DNAReader->GetJointName(JntIndex);
+		const FName BoneName(*boneNameFStr, 0, FNAME_Add, false);
+		const int32 BoneIndex = RefSkeleton->FindBoneIndex(BoneName);
+
+		// BoneIndex can be INDEX_NONE;
+		// We can safely put it into the map with other indices, it will be handled in the Evaluate method.
+		RLJointToUEBoneIndices.Add(BoneIndex);
+	}
+
+
+	double TimeToMap = FPlatformTime::Seconds();
+	TimeToMap -= StartTime;
+
+	UE_LOG(LogDNAToSkelMeshMap, Log, TEXT("	Map joints:	%.6f"), TimeToMap * 1000.0f);
+}
+
+void FDNAToSkelMeshMap::MapMorphTargets(IDNAReader* DNAReader)
+{
+	double StartTime = FPlatformTime::Seconds();
+
+	uint16 MeshCount = DNAReader->GetMeshCount();
+	//Initialize map which stores DeltaIndex for each Blend Shape Target, for each Vertex in each Mesh.
+	BlendShapeVertexDeltaIndices.Empty();
+	BlendShapeVertexDeltaIndices.AddZeroed(MeshCount);
+	// Assemble the reverse map, for each MorphTarget index save corresponding FDNABlendShapeTarget Index
+	// Keeping MeshIndex and TargetIndex for each MorphTarget.
+	MeshBlendShapeTargets.Empty();
+	MeshBlendShapeTargets.Init(FDNABlendShapeTarget(), TargetSkelMesh->MorphTargets.Num());
+
+	for (int16 MeshIndex = 0; MeshIndex < MeshCount; MeshIndex++)
+	{
+		uint16 TargetCount = DNAReader->GetBlendShapeTargetCount(MeshIndex);
+		for (int16 TargetIndex = 0; TargetIndex < TargetCount; TargetIndex++)
+		{
+			if (BlendShapeVertexDeltaIndices[MeshIndex].Num() == 0)
+			{
+				const int32 VertexCount = DNAReader->GetVertexPositionCount(MeshIndex);
+				BlendShapeVertexDeltaIndices[MeshIndex].AddZeroed(VertexCount);
+				for (int32 VertexIndex = 0; VertexIndex < VertexCount; VertexIndex++)
+				{
+					BlendShapeVertexDeltaIndices[MeshIndex][VertexIndex].Init(INDEX_NONE, TargetCount);
+				}
+			}
+
+			// Construct MorphTarget name by combining Blend Shape Channel name and Mesh name from DNA.
+			const uint16 ChannelIndex = DNAReader->GetBlendShapeChannelIndex(MeshIndex, TargetIndex);
+			const FString BlendShapeStr = DNAReader->GetBlendShapeChannelName(ChannelIndex);
+			const FString MeshStr = DNAReader->GetMeshName(MeshIndex);
+			const FString MorphTargetStr = MeshStr + TEXT("__") + BlendShapeStr;
+			const FName MorphTargetName(*MorphTargetStr, 0, FNAME_Add, false);
+
+			// Find MorphTarget from SkelMesh by previously constructed name.
+			int32 MorphTargetIndex;
+			UMorphTarget* MorphTarget = TargetSkelMesh->FindMorphTargetAndIndex(MorphTargetName, MorphTargetIndex);
+			if (MorphTarget != nullptr)
+			{
+				// Store MeshIndex and TargetIndex for found MorphTarget.
+				MeshBlendShapeTargets[MorphTargetIndex].MeshIndex = MeshIndex;
+				MeshBlendShapeTargets[MorphTargetIndex].TargetIndex = TargetIndex;
+
+				const TArrayView<const uint32> BlendShapeTargetVertices = DNAReader->GetBlendShapeTargetVertexIndices(MeshIndex, TargetIndex);
+				for (int32 DeltaIndex = 0; DeltaIndex < BlendShapeTargetVertices.Num(); DeltaIndex++)
+				{
+					// For each BlendShape Target Delta find corresponding Vertex and store Delta Index for this Vertex and Target in the map.
+					const int32 DNAVertex = BlendShapeTargetVertices[DeltaIndex];
+					BlendShapeVertexDeltaIndices[MeshIndex][DNAVertex][TargetIndex] = DeltaIndex;
+				}
+			}
+			else
+			{
+#ifdef DEBUG
+				UE_LOG(LogDNAToSkelMeshMap, Error, TEXT("Could not find morph target %s (probably below threshold) Channel %d, Target %d"), *MorphTargetStr, ChannelIndex, TargetIndex);
+#endif
+			}
+		}
+	}
+
+	double TimeToMap = FPlatformTime::Seconds();
+	TimeToMap -= StartTime;
+
+	//UE_LOG(LogDNAToSkelMeshMap, Log, TEXT("	MapMorphTargets:	%.6f"), TimeToMap * 1000.0f);
+}
+#endif // WITH_EDITORONLY_DATA
\ No newline at end of file
diff --git a/Engine/Plugins/Animation/RigLogic/Source/RigLogicModule/Private/DNAUtils.cpp b/Engine/Plugins/Animation/RigLogic/Source/RigLogicModule/Private/DNAUtils.cpp
new file mode 100644
--- /dev/null
+++ b/Engine/Plugins/Animation/RigLogic/Source/RigLogicModule/Private/DNAUtils.cpp
@@ -0,0 +1,57 @@
+// Copyright 1998-2019 Epic Games, Inc. All Rights Reserved.
+
+#include "DNAUtils.h"
+
+#include "CoreMinimal.h"
+#include "DNAReaderAdapter.h"
+#include "FMemoryResource.h"
+#include "RigLogicMemoryStream.h"
+
+#include "riglogic/RigLogic.h"
+
+static TSharedPtr<IDNAReader> ReadDNAStream(rl4::ScopedPtr<dna::StreamReader> DNAStreamReader)
+{
+	DNAStreamReader->read();
+	if (!rl4::Status::isOk())
+	{
+		UE_LOG(LogDNAReader, Error, TEXT("%s"), ANSI_TO_TCHAR(rl4::Status::get().message));
+		return nullptr;
+	}
+	return MakeShared<FDNAReader<dna::StreamReader>>(DNAStreamReader.release());
+}
+
+TSharedPtr<IDNAReader> ReadDNAFromFile(const FString& Path, EDNADataLayer Layer, uint16_t MaxLOD)
+{
+	auto DNAFileStream = rl4::makeScoped<rl4::FileStream>(TCHAR_TO_ANSI(*Path), rl4::FileStream::AccessMode::Read, rl4::FileStream::OpenMode::Binary, FMemoryResource::Instance());
+	auto DNAStreamReader = rl4::makeScoped<dna::StreamReader>(DNAFileStream.get(), static_cast<dna::DataLayer>(Layer), MaxLOD, FMemoryResource::Instance());
+	return ReadDNAStream(MoveTemp(DNAStreamReader));
+}
+
+TSharedPtr<IDNAReader> ReadDNAFromFile(const FString& Path, EDNADataLayer Layer, TArrayView<uint16_t> LODs)
+{
+	auto DNAFileStream = rl4::makeScoped<rl4::FileStream>(TCHAR_TO_ANSI(*Path), rl4::FileStream::AccessMode::Read, rl4::FileStream::OpenMode::Binary, FMemoryResource::Instance());
+	auto DNAStreamReader = rl4::makeScoped<dna::StreamReader>(DNAFileStream.get(), static_cast<dna::DataLayer>(Layer), LODs.GetData(), static_cast<uint16>(LODs.Num()), FMemoryResource::Instance());
+	return ReadDNAStream(MoveTemp(DNAStreamReader));
+}
+
+TSharedPtr<IDNAReader> ReadDNALayerFromBuffer(TArray<uint8>* DNABuffer, EDNADataLayer Layer, uint16_t MaxLOD)
+{
+	FRigLogicMemoryStream DNAMemoryStream(DNABuffer);
+	auto DNAStreamReader = rl4::makeScoped<dna::StreamReader>(&DNAMemoryStream, static_cast<dna::DataLayer>(Layer), MaxLOD, FMemoryResource::Instance());
+	return ReadDNAStream(MoveTemp(DNAStreamReader));
+}
+
+TSharedPtr<IDNAReader> ReadDNALayerFromBuffer(TArray<uint8>* DNABuffer, EDNADataLayer Layer, TArrayView<uint16_t> LODs)
+{
+	FRigLogicMemoryStream DNAMemoryStream(DNABuffer);
+	auto DNAStreamReader = rl4::makeScoped<dna::StreamReader>(&DNAMemoryStream, static_cast<dna::DataLayer>(Layer), LODs.GetData(), static_cast<uint16>(LODs.Num()), FMemoryResource::Instance());
+	return ReadDNAStream(MoveTemp(DNAStreamReader));
+}
+
+void WriteDNAToFile(const IDNAReader* Reader, EDNADataLayer Layer, const FString& Path)
+{
+	auto DNAFileStream = rl4::makeScoped<rl4::FileStream>(TCHAR_TO_ANSI(*Path), rl4::FileStream::AccessMode::Write, rl4::FileStream::OpenMode::Binary, FMemoryResource::Instance());
+	auto DNAStreamWriter = rl4::makeScoped<dna::StreamWriter>(DNAFileStream.get(), FMemoryResource::Instance());
+	DNAStreamWriter->setFrom(Reader->Unwrap(), static_cast<dna::DataLayer>(Layer), FMemoryResource::Instance());
+	DNAStreamWriter->write();
+}
diff --git a/Engine/Plugins/Animation/RigLogic/Source/RigLogicModule/Private/RigInstance.cpp b/Engine/Plugins/Animation/RigLogic/Source/RigLogicModule/Private/RigInstance.cpp
new file mode 100644
--- /dev/null
+++ b/Engine/Plugins/Animation/RigLogic/Source/RigLogicModule/Private/RigInstance.cpp
@@ -0,0 +1,89 @@
+// Copyright 1998-2019 Epic Games, Inc. All Rights Reserved.
+
+#include "RigInstance.h"
+
+#include "FMemoryResource.h"
+#include "RigLogic.h"
+
+#include "riglogic/RigLogic.h"
+
+void FRigInstance::FRigInstanceDeleter::operator()(rl4::RigInstance* Pointer)
+{
+	rl4::RigInstance::destroy(Pointer);
+}
+
+FRigInstance::FRigInstance(FRigLogic* RigLogic) :
+	RigInstance{rl4::RigInstance::create(RigLogic->Unwrap(), FMemoryResource::Instance())}
+{
+}
+
+FRigInstance::~FRigInstance() = default;
+
+uint16 FRigInstance::GetGUIControlCount() const
+{
+	return RigInstance->getGUIControlCount();
+}
+
+void FRigInstance::SetGUIControl(uint16 Index, float Value)
+{
+	RigInstance->setGUIControl(Index, Value);
+}
+
+void FRigInstance::SetGUIControlValues(const float* Values)
+{
+	RigInstance->setGUIControlValues(Values);
+}
+
+uint16 FRigInstance::GetRawControlCount() const
+{
+	return RigInstance->getRawControlCount();
+}
+
+void FRigInstance::SetRawControl(uint16 Index, float Value)
+{
+	RigInstance->setRawControl(Index, Value);
+}
+
+void FRigInstance::SetRawControlValues(const float* Values)
+{
+	RigInstance->setRawControlValues(Values);
+}
+
+TArrayView<const float> FRigInstance::GetRawJointOutputs() const
+{
+	rl4::ConstArrayView<float> Outputs = RigInstance->getRawJointOutputs();
+	return TArrayView<const float>{Outputs.data(), static_cast<int32>(Outputs.size())};
+}
+
+FTransformArrayView FRigInstance::GetJointOutputs() const
+{
+	rl4::ConstArrayView<float> Outputs = RigInstance->getRawJointOutputs();
+	return FTransformArrayView(Outputs.data(), Outputs.size());
+}
+
+TArrayView<const float> FRigInstance::GetBlendShapeOutputs() const
+{
+	rl4::ConstArrayView<float> Outputs = RigInstance->getBlendShapeOutputs();
+	return TArrayView<const float>{Outputs.data(), static_cast<int32>(Outputs.size())};
+}
+
+TArrayView<const float> FRigInstance::GetAnimatedMapOutputs() const
+{
+	rl4::ConstArrayView<float> Outputs = RigInstance->getAnimatedMapOutputs();
+	return TArrayView<const float>{Outputs.data(), static_cast<int32>(Outputs.size())};
+}
+
+uint16 FRigInstance::GetLOD() const
+{
+	return RigInstance->getLOD();
+}
+
+void FRigInstance::SetLOD(uint16 LOD)
+{
+	RigInstance->setLOD(LOD);
+}
+
+rl4::RigInstance* FRigInstance::Unwrap() const
+{
+	return RigInstance.Get();
+}
diff --git a/Engine/Plugins/Animation/RigLogic/Source/RigLogicModule/Private/RigLogicMemoryStream.cpp b/Engine/Plugins/Animation/RigLogic/Source/RigLogicModule/Private/RigLogicMemoryStream.cpp
new file mode 100644
--- /dev/null
+++ b/Engine/Plugins/Animation/RigLogic/Source/RigLogicModule/Private/RigLogicMemoryStream.cpp
@@ -0,0 +1,55 @@
+#include "RigLogicMemoryStream.h"
+
+FRigLogicMemoryStream::FRigLogicMemoryStream(TArray<uint8>* Buffer)
+{
+	BitStreamBuffer = Buffer; //don't allocate new memory, just point to the existing one
+}
+
+void FRigLogicMemoryStream::seek(size_t Position)
+{
+	PositionInBuffer = Position;
+}
+
+size_t FRigLogicMemoryStream::tell()
+{
+	return PositionInBuffer;
+}
+
+void FRigLogicMemoryStream::open()
+{
+	PositionInBuffer = 0;
+}
+
+void FRigLogicMemoryStream::read(char* ReadToBuffer, size_t Size)
+{
+	FMemory::Memcpy(ReadToBuffer, BitStreamBuffer->GetData() + PositionInBuffer, Size);
+	PositionInBuffer = PositionInBuffer + Size;
+}
+
+void FRigLogicMemoryStream::write(const char* WriteFromBuffer, size_t Size)
+{
+	FString writeStr;
+	int32 BufferSize = BitStreamBuffer->Num();
+	if (PositionInBuffer + Size > BufferSize)
+	{
+		int32 difference = PositionInBuffer + Size - BufferSize;
+		BitStreamBuffer->AddUninitialized(difference);
+	}
+	FMemory::Memcpy(BitStreamBuffer->GetData() + PositionInBuffer, WriteFromBuffer, Size);
+	PositionInBuffer = PositionInBuffer + Size;
+}
+
+size_t FRigLogicMemoryStream::size()
+{
+	return BitStreamBuffer->Num();
+}
+
+bool FRigLogicMemoryStream::hasError()
+{
+	return false;
+}
+
+TArray<uint8>* FRigLogicMemoryStream::GetBitStreamBuffer()
+{
+	return BitStreamBuffer;
+}
\ No newline at end of file
diff --git a/Engine/Plugins/Animation/RigLogic/Source/RigLogicModule/Private/RigLogicMemoryStream.h b/Engine/Plugins/Animation/RigLogic/Source/RigLogicModule/Private/RigLogicMemoryStream.h
new file mode 100644
--- /dev/null
+++ b/Engine/Plugins/Animation/RigLogic/Source/RigLogicModule/Private/RigLogicMemoryStream.h
@@ -0,0 +1,34 @@
+// Copyright 1998-2019 Epic Games, Inc. All Rights Reserved.
+
+#pragma once
+
+#include "CoreMinimal.h"
+
+#include "riglogic/RigLogic.h"
+
+/** A simple implementation of in-memory stream for interfacing with RigLogic API,
+  so RigLogic can consume DNA as a Stream either from file or memory */
+
+class FRigLogicMemoryStream: public rl4::BoundedIOStream
+{
+public:
+	/** The buffer is not copied, the pointer to it is stored inside this object**/
+	FRigLogicMemoryStream(TArray<uint8>* Buffer);
+	~FRigLogicMemoryStream() {} //as no memory was allocated, we do not release it
+
+	void seek(size_t Position) override;
+	size_t tell() override;
+	void open() override;
+	void close() override {}
+	void read(char* ReadToBuffer, size_t Size) override;
+	void write(const char* WriteFromBuffer, size_t Size) override;
+	size_t size() override;
+	void release() {}
+	static void deallocResource(void* ptr) {}; //no memory was allocated
+	bool hasError();
+	TArray<uint8>* GetBitStreamBuffer();
+
+private:
+	TArray<uint8>* BitStreamBuffer; //doesn't contain the array, only points to the array given to it
+	size_t PositionInBuffer = 0;
+};
diff --git a/Engine/Plugins/Animation/RigLogic/Source/RigLogicModule/Private/RigLogicModule.cpp b/Engine/Plugins/Animation/RigLogic/Source/RigLogicModule/Private/RigLogicModule.cpp
new file mode 100644
--- /dev/null
+++ b/Engine/Plugins/Animation/RigLogic/Source/RigLogicModule/Private/RigLogicModule.cpp
@@ -0,0 +1,32 @@
+// Copyright 1998-2019 Epic Games, Inc. All Rights Reserved.
+
+#include "RigLogicModule.h"
+
+#include "RigUnit_RigLogic.h"
+
+#include "Interfaces/IPluginManager.h"
+#include "Runtime/Core/Public/Misc/Paths.h"
+
+#include "Runtime/Core/Public/Modules/ModuleInterface.h"
+#include "Runtime/Core/Public/Modules/ModuleManager.h"
+
+#include "Runtime/Core/Public/HAL/PlatformProcess.h"
+
+#define LOCTEXT_NAMESPACE "FRigLogicModule"
+
+DEFINE_LOG_CATEGORY(LogRigLogic);
+
+void FRigLogicModule::StartupModule()
+{
+
+}
+
+void FRigLogicModule::ShutdownModule()
+{
+
+}
+
+#undef LOCTEXT_NAMESPACE
+
+IMPLEMENT_MODULE(FRigLogicModule, RigLogicModule)
+
diff --git a/Engine/Plugins/Animation/RigLogic/Source/RigLogicModule/Private/RigUnit_RigLogic.cpp b/Engine/Plugins/Animation/RigLogic/Source/RigLogicModule/Private/RigUnit_RigLogic.cpp
new file mode 100644
--- /dev/null
+++ b/Engine/Plugins/Animation/RigLogic/Source/RigLogicModule/Private/RigUnit_RigLogic.cpp
@@ -0,0 +1,431 @@
+// Copyright 1998-2019 Epic Games, Inc. All Rights Reserved.
+
+#include "RigUnit_RigLogic.h"
+
+#include "ControlRig.h"
+#include "DNAReader.h"
+#include "RigInstance.h"
+#include "RigLogic.h"
+#include "SkelMeshDNAUtils.h"
+#include "Runtime/Core/Public/Math/TransformNonVectorized.h"
+#include "Units/RigUnitContext.h"
+
+DEFINE_LOG_CATEGORY(LogRigLogicUnit);
+
+const uint8 FRigUnit_RigLogic_Data::MAX_ATTRS_PER_JOINT = 9;
+
+FRigUnit_RigLogic_Data::FRigUnit_RigLogic_Data()
+: SkelMeshComponent(nullptr)
+, CurrentLOD(0)
+, RigLogic(nullptr)
+, RigInstance(nullptr)
+{
+}
+
+FRigUnit_RigLogic_Data::~FRigUnit_RigLogic_Data()
+{
+	if (RigInstance != nullptr)
+	{
+		delete RigInstance;
+		RigInstance = nullptr;
+	}
+
+	if (RigLogic != nullptr)
+	{
+		delete RigLogic;
+		RigLogic = nullptr;
+	}
+}
+
+FRigUnit_RigLogic_Data::FRigUnit_RigLogic_Data(const FRigUnit_RigLogic_Data& Other)
+{
+	*this = Other;
+}
+
+FRigUnit_RigLogic_Data& FRigUnit_RigLogic_Data::operator = (const FRigUnit_RigLogic_Data& Other)
+{
+
+	if (RigInstance != nullptr)
+	{
+		delete RigInstance;
+		RigInstance = nullptr;
+	}
+
+	if (RigLogic != nullptr)
+	{
+		delete RigLogic;
+		RigLogic = nullptr;
+	}
+
+	SkelMeshComponent = Other.SkelMeshComponent;
+    //@helge: is this required, or the unit is always initialized afterwards?
+	//RigLogic = Other.RigLogic;
+	//RigInstance = Other.RigInstance;
+	CurrentLOD = Other.CurrentLOD;
+	InputCurveIndices = Other.InputCurveIndices;
+	HierarchyBoneIndices = Other.HierarchyBoneIndices;
+	MorphTargetCurveIndices = Other.MorphTargetCurveIndices;
+	BlendShapeIndices = Other.BlendShapeIndices;
+	CurveContainerIndicesForAnimMaps = Other.CurveContainerIndicesForAnimMaps;
+	RigLogicIndicesForAnimMaps = Other.RigLogicIndicesForAnimMaps;
+	UpdatedJoints = Other.UpdatedJoints;
+	BlendShapeMappingCount = Other.BlendShapeMappingCount;
+	NeckFemaleAverageCorCurveIndex = Other.NeckFemaleAverageCorCurveIndex;
+	NeckMaleMuscularCorExpCurveIndex = Other.NeckFemaleAverageCorCurveIndex;
+
+	return *this;
+}
+
+//maps indices of input curves from dna file to control rig curves
+void FRigUnit_RigLogic_Data::MapInputCurveIndices(const IBehaviorReader* DNABehavior, const FRigCurveContainer* CurveContainer)
+{
+	InputCurveIndices.Reset();
+
+	const uint32 ControlCount = DNABehavior->GetRawControlCount();
+	for (uint32_t ControlIndex = 0; ControlIndex < ControlCount; ++ControlIndex)
+	{
+		const FString DNAControlName = DNABehavior->GetRawControlName(ControlIndex);
+		const FString AnimatedControlName = ConstructCurveName(DNAControlName, TEXT("<obj>_<attr>"));
+		if (AnimatedControlName == TEXT(""))
+		{
+			return;
+		}
+
+		const FName ControlFName(*AnimatedControlName);
+		const int32 CurveIndex = CurveContainer ? CurveContainer->GetIndex(ControlFName) : INDEX_NONE;
+		InputCurveIndices.Add(CurveIndex); //can be INDEX_NONE
+	}
+}
+
+void FRigUnit_RigLogic_Data::MapJoints(const IBehaviorReader* DNABehavior, const FRigBoneHierarchy* Hierarchy)
+{
+	const uint16 JointCount = DNABehavior->GetJointCount();
+
+	HierarchyBoneIndices.Reset(JointCount);
+	for (uint16 JointIndex = 0; JointIndex < JointCount ; ++JointIndex)
+	{
+		const FString RLJointName = DNABehavior->GetJointName(JointIndex);
+		const FName JointFName = FName(*RLJointName);
+		const int32 BoneIndex = Hierarchy->GetIndex(JointFName);
+		HierarchyBoneIndices.Add(BoneIndex);
+	}
+}
+
+void FRigUnit_RigLogic_Data::MapMorphTargets(const IBehaviorReader* DNABehavior, const FRigCurveContainer* CurveContainer)
+{
+	uint16 LODCount = DNABehavior->GetLODCount();
+
+	MorphTargetCurveIndices.Reset();
+	MorphTargetCurveIndices.AddZeroed(LODCount);
+	BlendShapeIndices.Reset();
+	BlendShapeIndices.AddZeroed(LODCount);
+
+	for (uint16 LodIndex = 0; LodIndex < LODCount; ++LodIndex)
+	{
+		TArrayView<const uint16> BlendShapeChannelIndicesForLOD =
+			DNABehavior->GetMeshBlendShapeChannelMappingIndicesForLOD(LodIndex);
+		for (uint16 MappingIndex: BlendShapeChannelIndicesForLOD)
+		{
+			const FMeshBlendShapeChannelMapping Mapping = DNABehavior->GetMeshBlendShapeChannelMapping(MappingIndex);
+			const uint16 BlendShapeIndex = Mapping.BlendShapeChannelIndex;
+			const uint16 MeshIndex = Mapping.MeshIndex;
+			const FString BlendShapeStr = DNABehavior->GetBlendShapeChannelName(BlendShapeIndex);
+			const FString MeshStr = DNABehavior->GetMeshName(MeshIndex);
+			const FString MorphTargetStr = MeshStr + TEXT("__") + BlendShapeStr;
+			const FName MorphTargetName(*MorphTargetStr);
+
+			const int32 MorphTargetIndex = CurveContainer->GetIndex(MorphTargetName);
+			MorphTargetCurveIndices[LodIndex].Values.Add(MorphTargetIndex);
+			BlendShapeIndices[LodIndex].Values.Add(Mapping.BlendShapeChannelIndex);
+		}
+	}
+}
+
+void FRigUnit_RigLogic_Data::MapMaskMultipliers(const IBehaviorReader* DNABehavior, const FRigCurveContainer* CurveContainer)
+{
+	uint16 LODCount = DNABehavior->GetLODCount();
+	CurveContainerIndicesForAnimMaps.Reset();
+	CurveContainerIndicesForAnimMaps.AddZeroed(LODCount);
+
+	RigLogicIndicesForAnimMaps.Reset();
+	RigLogicIndicesForAnimMaps.AddZeroed(LODCount);
+
+	for (uint16 LodIndex = 0; LodIndex < LODCount; ++LodIndex)
+	{		
+		TArrayView<const uint16> AnimMapIndicesPerLOD = DNABehavior->GetAnimatedMapIndicesForLOD(LodIndex);
+		for (uint16 AnimMapIndexPerLOD: AnimMapIndicesPerLOD)
+		{
+			FString AnimMapNameFStr = DNABehavior->GetAnimatedMapName(AnimMapIndexPerLOD);
+			FString MaskMultiplierNameStr = ConstructCurveName(AnimMapNameFStr, TEXT("<obj>_<attr>"));
+			if (MaskMultiplierNameStr == "") {
+				return;
+			}
+			FName MaskMultiplierFName(*MaskMultiplierNameStr);
+
+			const int32 CurveIndex = CurveContainer->GetIndex(MaskMultiplierFName);
+			CurveContainerIndicesForAnimMaps[LodIndex].Values.Add(CurveIndex); //can be INDEX_NONE if curve was not found
+			RigLogicIndicesForAnimMaps[LodIndex].Values.Add(AnimMapIndexPerLOD);
+		}
+	}
+}
+
+FString FRigUnit_RigLogic_Data::ConstructCurveName(const FString& NameToSplit, const FString& FormatString)
+{
+	// constructs curve name from NameToSplit (always in form <obj>.<attr>)
+	// using FormatString of form x<obj>y<attr>z
+	// where x, y and z are arbitrary strings
+	// example:
+	// FormatString="mesh_<obj>_<attr>"
+	// 'head.blink_L' becomes 'mesh_head_blink_L'
+
+	FString ObjectName, AttributeName;
+	if (!NameToSplit.Split(".", &ObjectName, &AttributeName))
+	{
+		
+		UE_LOG(LogRigLogicUnit, Error, TEXT("RigUnit_R: Missing '.' in '%s'"),
+			*NameToSplit);
+		return TEXT("");
+	}
+
+	FString CurveName = FormatString;
+	CurveName = CurveName.Replace(TEXT("<obj>"), *ObjectName);
+	CurveName = CurveName.Replace(TEXT("<attr>"), *AttributeName);
+	return CurveName;
+}
+
+void FRigUnit_RigLogic_Data::CalculateRigLogic(FControlRigExecuteContext& ExecuteContext)
+{
+	const int32 RawControlCount = RigInstance->GetRawControlCount();
+	for (int32 ControlIndex = 0; ControlIndex < RawControlCount; ++ControlIndex)
+	{
+		const uint32 CurveIndex = InputCurveIndices[ControlIndex];
+		float Value = ExecuteContext.GetCurves()->GetValue(CurveIndex);
+	
+		RigInstance->SetRawControl(ControlIndex, Value);
+	}
+
+	RigLogic->Calculate(RigInstance);
+}
+
+void FRigUnit_RigLogic_Data::UpdateJoints(FRigHierarchyContainer* Hierarchy, const FRigUnit_RigLogic_JointUpdateParams& JointUpdateParams)
+{
+	const uint16 JointCount = UpdatedJoints.Num();
+	for (uint16 JointIndex = 0; JointIndex < JointCount; ++JointIndex)
+	{
+		UpdatedJoints[JointIndex] = false;
+	}
+
+	for (const uint16 AttrIndex : JointUpdateParams.VariableAttributes)
+	{
+		const uint16 JointIndex = AttrIndex / MAX_ATTRS_PER_JOINT;
+		if (!UpdatedJoints[JointIndex]) // Check if joint is already updated. It can happen max 9 times because there is 9 attributes for each joint.
+		{
+			UpdatedJoints[JointIndex] = true;
+			const int32 BoneIndex = HierarchyBoneIndices[JointIndex];
+			if (BoneIndex != INDEX_NONE)
+			{
+				const FTransform& Neutral = JointUpdateParams.NeutralJointTransforms[JointIndex];
+				const FTransform& Delta = JointUpdateParams.DeltaTransforms[JointIndex];
+				FTransform Transform
+				{
+					Neutral.GetRotation() * Delta.GetRotation(),
+					Neutral.GetTranslation() + Delta.GetTranslation(),
+					Neutral.GetScale3D() + Delta.GetScale3D()  // Neutral scale is always 1.0
+				};
+				Hierarchy->BoneHierarchy[BoneIndex].LocalTransform = Transform;
+			}
+		}
+	}
+
+	Hierarchy->BoneHierarchy.RecomputeGlobalTransforms();
+}
+
+void FRigUnit_RigLogic_Data::UpdateBlendShapeCurves(FRigCurveContainer* CurveContainer, TArrayView<const float>& BlendShapeValues )
+{
+	// set output blend shapes
+	if (BlendShapeIndices.IsValidIndex(CurrentLOD) && MorphTargetCurveIndices.IsValidIndex(CurrentLOD))
+	{
+		const uint32 BlendShapePerLODCount = static_cast<uint32>(BlendShapeIndices[CurrentLOD].Values.Num());
+
+		if (ensure(BlendShapePerLODCount == MorphTargetCurveIndices[CurrentLOD].Values.Num()))
+		{
+			for (uint32 MeshBlendIndex = 0; MeshBlendIndex < BlendShapePerLODCount; MeshBlendIndex++)
+			{
+				const int32 BlendShapeIndex = BlendShapeIndices[CurrentLOD].Values[MeshBlendIndex];
+
+				const int32 MorphTargetCurveIndex = MorphTargetCurveIndices[CurrentLOD].Values[MeshBlendIndex];
+				if (MorphTargetCurveIndex != INDEX_NONE)
+				{
+					float Value = BlendShapeValues[BlendShapeIndex];
+					CurveContainer->SetValue(MorphTargetCurveIndex, Value);
+				}
+			}
+		}
+	}
+	else
+	{
+		UE_LOG(LogRigLogicUnit, Warning, TEXT("Invalid LOD Index for the BlendShapes. Ensure your curve is set up correctly!"));
+	}
+}
+
+void FRigUnit_RigLogic_Data::UpdateAnimMapCurves(FRigCurveContainer* CurveContainer, TArrayView<const float>& AnimMapOutputs)
+{
+	// set output mask multipliers
+	// In case curves are not imported yet into CL, AnimatedMapsCurveIndices will be empty, so we need to check
+	// array bounds before trying to access it:
+	if (RigLogicIndicesForAnimMaps.IsValidIndex(CurrentLOD) && CurveContainerIndicesForAnimMaps.IsValidIndex(CurrentLOD))
+	{
+		uint32 AnimMapPerLODCount = RigLogicIndicesForAnimMaps[CurrentLOD].Values.Num();
+
+		for (uint32 AnimMapIndexForLOD = 0; AnimMapIndexForLOD < AnimMapPerLODCount; ++AnimMapIndexForLOD)
+		{
+			const int32 RigLogicAnimMapIndex = RigLogicIndicesForAnimMaps[CurrentLOD].Values[AnimMapIndexForLOD];
+
+			const int32 CurveContainerAnimMapIndex = CurveContainerIndicesForAnimMaps[CurrentLOD].Values[AnimMapIndexForLOD];
+			if (CurveContainerAnimMapIndex != INDEX_NONE)
+			{
+				float Value = AnimMapOutputs[RigLogicAnimMapIndex];
+				CurveContainer->SetValue(CurveContainerAnimMapIndex, Value);
+			}
+		}
+	}
+	else
+	{
+		UE_LOG(LogRigLogicUnit, Warning, TEXT("Invalid LOD Index for the AnimationMaps. Ensure your curve is set up correctly!"));
+	}	
+}
+
+FString FRigUnit_RigLogic::GetUnitLabel() const
+{
+	return FString::Printf(TEXT("RigUnit"));
+}
+
+IBehaviorReader* FRigUnit_RigLogic_Data::FetchBehaviorReaderFromOwner() 
+{
+	if (SkelMeshComponent.IsValid())
+	{
+		USkeletalMesh* SkelMesh = SkelMeshComponent->SkeletalMesh;
+		UAssetUserData* UserData = SkelMesh->GetAssetUserDataOfClass(UDNAAsset::StaticClass());
+		if (UserData == nullptr)
+		{
+			return nullptr;
+		}
+		UDNAAsset* DNAAsset = Cast<UDNAAsset>(UserData);
+		//inside DNAAsset there can be (Behavior)StreamReader or GeneSplicerDNAReader
+		//both implement behavior reader interface, so implicitly casting is safe:
+		return DNAAsset->GetBehaviorReader().Get();
+	}
+
+	return nullptr;
+}
+
+bool FRigUnit_RigLogic_Data::IsRigLogicInitialized()
+{
+	return RigLogic != nullptr && RigInstance != nullptr;
+}
+
+void FRigUnit_RigLogic_Data::InitializeRigLogic(IBehaviorReader* DNABehavior )
+{
+	if (DNABehavior == nullptr)
+	{
+		return;
+	}
+
+	RigLogic = new FRigLogic(DNABehavior); //needs only BehaviorReader
+	RigInstance = new FRigInstance(RigLogic);
+	RigInstance->SetLOD(0);
+}
+
+void FRigUnit_RigLogic_Data::ChangeRigLogicLODIfNeeded()
+{
+	if (!SkelMeshComponent.IsValid())
+	{
+		return;
+	}
+
+	int32 SkelMeshPredictedLOD = SkelMeshComponent->PredictedLODLevel;
+	// Set current LOD to RigLogic only if it changed.
+	if (CurrentLOD != SkelMeshPredictedLOD)
+	{
+		CurrentLOD = SkelMeshPredictedLOD;
+		RigInstance->SetLOD(CurrentLOD);
+	}
+}
+
+FRigUnit_RigLogic_Execute()
+{
+ 	FRigHierarchyContainer* Hierarchy = ExecuteContext.Hierarchy;
+	if (Hierarchy)
+	{
+		switch (Context.State)
+		{
+			case EControlRigState::Init:
+			{				
+				if (Data.SkelMeshComponent == nullptr)
+				{
+					Data.SkelMeshComponent = Context.DataSourceRegistry->RequestSource<USkeletalMeshComponent>(UControlRig::OwnerComponent);
+					// In normal execution, Data.SkelMeshComponent will be nullptr at the beginning
+					// however, during unit testing we cannot fetch it from DataSourceRegistry 
+					// in that case, a mock version will be inserted into Data by unit test beforehand
+				}
+				if (Data.SkelMeshComponent== nullptr || Data.SkelMeshComponent->SkeletalMesh == nullptr)
+				{
+					return;
+				}
+				Data.CurrentLOD = Data.SkelMeshComponent->PredictedLODLevel;
+
+				//Fetch BehaviorStreamReader from DNAAsset
+				//it can be imported into SkeletalMesh from DNA file, or overwritten by GeneSplicer when making
+				//a new character
+				IBehaviorReader* DNABehavior = Data.FetchBehaviorReaderFromOwner();				
+				Data.InitializeRigLogic(DNABehavior);
+				if (!Data.IsRigLogicInitialized())
+				{
+					return;
+				}
+
+				FRigCurveContainer* CurveContainer = ExecuteContext.GetCurves();
+
+				Data.MapJoints(DNABehavior, &Hierarchy->BoneHierarchy);
+				Data.MapInputCurveIndices(DNABehavior, CurveContainer);
+				Data.MapMorphTargets(DNABehavior, CurveContainer);
+				Data.MapMaskMultipliers(DNABehavior, CurveContainer);
+				Data.UpdatedJoints.SetNumZeroed(DNABehavior->GetJointCount());
+
+				break; 
+			}
+			case EControlRigState::Update:
+			{
+
+				FRigCurveContainer* CurveContainer = ExecuteContext.GetCurves();
+
+				if (!Data.IsRigLogicInitialized() || CurveContainer == nullptr)
+				{
+					return;
+				}
+				Data.ChangeRigLogicLODIfNeeded();
+				Data.CalculateRigLogic(ExecuteContext);
+
+				//Filing a struct so we can call the same method for updating joints from tests
+				FRigUnit_RigLogic_JointUpdateParams JointUpdateParamsTemp
+				(
+					Data.RigLogic->GetJointVariableAttributeIndices(Data.CurrentLOD),
+					Data.RigLogic->GetNeutralJointValues(),
+					Data.RigInstance->GetJointOutputs()
+				);
+
+				Data.UpdateJoints(Hierarchy, JointUpdateParamsTemp);
+
+				TArrayView<const float> BlendShapeValues = Data.RigInstance->GetBlendShapeOutputs();
+				Data.UpdateBlendShapeCurves( CurveContainer, BlendShapeValues );
+
+				TArrayView<const float> AnimMapOutputs = Data.RigInstance->GetAnimatedMapOutputs();
+				Data.UpdateAnimMapCurves( CurveContainer, AnimMapOutputs ); 
+			}
+			default:
+			{
+				break;
+			}
+		}
+	}
+}
