Editgrating RigLogic plugin from Dev-CharacterTech to UE5/Main @ CL 13891876.

diff --git a/Engine/Plugins/Animation/RigLogic/Source/RigLogicModule/Private/RigUnit_RigLogic_Test.cpp b/Engine/Plugins/Animation/RigLogic/Source/RigLogicModule/Private/RigUnit_RigLogic_Test.cpp
new file mode 100644
--- /dev/null
+++ b/Engine/Plugins/Animation/RigLogic/Source/RigLogicModule/Private/RigUnit_RigLogic_Test.cpp
@@ -0,0 +1,806 @@
+#if WITH_DEV_AUTOMATION_TESTS
+	#include "RigUnit_RigLogic_Test.h"
+#endif
+
+#include "RigUnit_RigLogic.h"
+#include "Units/RigUnitContext.h"
+#include "ControlRig.h"
+#include "riglogic/RigLogic.h"
+#include "Runtime/Core/Public/Math/TransformNonVectorized.h"
+#include "SkelMeshDNAUtils.h"
+
+#include "riglogic/RigLogic.h"
+#include "DNAUtils.h"
+
+#if WITH_DEV_AUTOMATION_TESTS
+#include "Units/RigUnitTest.h"
+#include "RigUnit_RigLogic.h"
+
+const uint8 FRigUnit_RigLogic::TestAccessor::MAX_ATTRS_PER_JOINT = 9;
+
+FRigUnit_RigLogic::TestAccessor::TestAccessor(FRigUnit_RigLogic* Unit)
+{
+	this->Unit = Unit;
+}
+
+/** ====== Map Input Curves ===== **/
+
+TestBehaviorReader* FRigUnit_RigLogic::TestAccessor::CreateBehaviorReaderEmpty()
+{
+	return new TestBehaviorReader();
+}
+
+FRigCurveContainer* FRigUnit_RigLogic::TestAccessor::CreateCurveContainerEmpty()
+{
+	return new FRigCurveContainer();
+}
+
+TestBehaviorReader* FRigUnit_RigLogic::TestAccessor::CreateBehaviorReaderOneCurve(FString ControlNameStr)
+{
+	TestBehaviorReader* BehaviorReader = new TestBehaviorReader;
+	//StringView ControlName(TCHAR_TO_ANSI(*ControlNameStr), ControlNameStr.Len());
+	BehaviorReader->rawControls.Add(ControlNameStr);
+	BehaviorReader->LODCount = 1;
+	return BehaviorReader;
+}
+
+FRigCurveContainer* FRigUnit_RigLogic::TestAccessor::CreateCurveContainerOneCurve(FString CurveNameStr)
+{
+	FRigCurveContainer* ValidCurveContainer = new  FRigCurveContainer();
+	ValidCurveContainer->Add(FName(*CurveNameStr));
+	ValidCurveContainer->Initialize();
+	return ValidCurveContainer;
+}
+
+void FRigUnit_RigLogic::TestAccessor::Exec_MapInputCurve(TestBehaviorReader* TestReader, FRigCurveContainer* TestCurveContainer)
+{
+	Unit->Data.MapInputCurveIndices(TestReader, TestCurveContainer); //inside the method so we can access data, which is a private member
+}
+
+/** ====== Map Joints ===== **/
+
+
+FRigBoneHierarchy* FRigUnit_RigLogic::TestAccessor::CreateBoneHierarchyEmpty()
+{
+	return new FRigBoneHierarchy();
+}
+
+FRigBoneHierarchy* FRigUnit_RigLogic::TestAccessor::CreateBoneHierarchyTwoBones( FString Bone1NameStr, FString Bone2NameStr )
+{
+	FRigBoneHierarchy* TestHierarchy = new FRigBoneHierarchy();
+	TestHierarchy->Reset();
+	TestHierarchy->Add(*Bone1NameStr, NAME_None, ERigBoneType::User, FTransform(FVector(1.f, 0.f, 0.f)));
+	TestHierarchy->Add(*Bone2NameStr, *Bone1NameStr, ERigBoneType::User, FTransform(FVector(1.f, 2.f, 0.f)));
+	TestHierarchy->Initialize();
+	return TestHierarchy;
+}
+
+TestBehaviorReader* FRigUnit_RigLogic::TestAccessor::CreateBehaviorReaderTwoJoints( FString Joint1NameStr, FString Joint2NameStr)
+{
+	TestBehaviorReader* TestReader = new TestBehaviorReader();
+	TestReader->addJoint(Joint1NameStr);
+	TestReader->addJoint(Joint2NameStr);
+	TestReader->LODCount = 1;
+	return TestReader;
+}
+
+void FRigUnit_RigLogic::TestAccessor::Exec_MapJoints(TestBehaviorReader* TestReader, FRigBoneHierarchy* TestHierarchy)
+{
+	Unit->Data.MapJoints(TestReader, TestHierarchy);
+}
+
+/** ====== Map Morph Targets ===== **/
+
+TestBehaviorReader* FRigUnit_RigLogic::TestAccessor::CreateBehaviorReaderNoBlendshapes(FString MeshNameStr)
+{
+	TestBehaviorReader* BehaviorReader = new TestBehaviorReader;
+	BehaviorReader->addMeshName(*MeshNameStr);
+	BehaviorReader->LODCount = 1; //there is one mesh, so LODs exist
+	return BehaviorReader;
+}
+
+TestBehaviorReader* FRigUnit_RigLogic::TestAccessor::CreateBehaviorReaderOneBlendShape(FString MeshNameStr, FString BlendShapeNameStr)
+{
+	TestBehaviorReader* BehaviorReader = new TestBehaviorReader;
+	BehaviorReader->addBlendShapeChannelName(*BlendShapeNameStr);
+	BehaviorReader->addBlendShapeChannelName(*BlendShapeNameStr);
+	BehaviorReader->addMeshName(*MeshNameStr);
+	BehaviorReader->addBlendShapeMapping(0, 0);
+	BehaviorReader->addBlendShapeMappingIndicesToLOD(0, 0); //mapping 0 to LOD0
+	BehaviorReader->LODCount = 1;
+
+	return BehaviorReader;
+}
+
+TestBehaviorReader* FRigUnit_RigLogic::TestAccessor::CreateBehaviorReaderTwoBlendShapes(FString MeshNameStr, FString BlendShape1Str, FString BlendShape2Str)
+{
+	TestBehaviorReader* BehaviorReader = new TestBehaviorReader;
+	BehaviorReader->addBlendShapeChannelName(*BlendShape1Str);
+	BehaviorReader->addBlendShapeChannelName(*BlendShape2Str);
+	BehaviorReader->addMeshName(*MeshNameStr);
+	BehaviorReader->addBlendShapeMapping(0, 0);
+	BehaviorReader->addBlendShapeMapping(0, 1);
+	//call 	BehaviorReader->addBlendShapeMappingIndicesToLOD(x, y) outside of this method to assign to various LODs
+	BehaviorReader->LODCount = 1;
+	return BehaviorReader;
+}
+
+FRigCurveContainer* FRigUnit_RigLogic::TestAccessor::CreateCurveContainerOneMorphTarget(FString MorphTargetStr)
+{
+	FRigCurveContainer* ValidCurveContainer = new FRigCurveContainer();
+	ValidCurveContainer->Add(FName(*MorphTargetStr));
+	ValidCurveContainer->Initialize();
+	return ValidCurveContainer;
+}
+
+FRigCurveContainer* FRigUnit_RigLogic::TestAccessor::CreateCurveContainerTwoMorphTargets(FString MorphTarget1Str, FString MorphTarget2Str)
+{
+	FRigCurveContainer* ValidCurveContainer = new FRigCurveContainer();
+	ValidCurveContainer->Add(FName(*MorphTarget1Str));
+	ValidCurveContainer->Add(FName(*MorphTarget2Str));
+	ValidCurveContainer->Initialize();
+	return ValidCurveContainer;
+}
+
+void FRigUnit_RigLogic::TestAccessor::Exec_MapMorphTargets(TestBehaviorReader* TestReader, FRigCurveContainer* TestCurveContainer)
+{
+	//put into a separate method so we can access the private Data member
+	Unit->Data.MapMorphTargets(TestReader, TestCurveContainer); 
+}
+
+/** ====== Map Mask Multipliers ===== **/
+
+TestBehaviorReader* FRigUnit_RigLogic::TestAccessor::CreateBehaviorReaderOneAnimatedMap(FString AnimatedMapNameStr)
+{
+	TestBehaviorReader* BehaviorReader = new TestBehaviorReader;
+	BehaviorReader->animatedMaps.Add(AnimatedMapNameStr);
+	BehaviorReader->addAnimatedMapIndicesToLOD(0, 0);
+	BehaviorReader->LODCount = 1;
+	return BehaviorReader;
+}
+
+void FRigUnit_RigLogic::TestAccessor::Exec_MapMaskMultipliers(TestBehaviorReader* TestReader, FRigCurveContainer* TestCurveContainer)
+{
+	Unit->Data.MapMaskMultipliers(TestReader, TestCurveContainer); //inside the method so we can access data, which is a private member
+}
+
+void FRigUnit_RigLogic::TestAccessor::AddToTransformArray(float* InArray, FTransform& Transform)
+{
+	uint32 FirstAttributeIndex = 0;
+	
+	FVector Rotation = Transform.GetRotation().Euler();
+	InArray[FirstAttributeIndex + 0] = Rotation.X;
+	InArray[FirstAttributeIndex + 1] = Rotation.Y;
+	InArray[FirstAttributeIndex + 2] = Rotation.Z;
+
+	FVector Translation = Transform.GetTranslation();
+	InArray[FirstAttributeIndex + 3] = Translation.X;
+	InArray[FirstAttributeIndex + 4] = Translation.Y;
+	InArray[FirstAttributeIndex + 5] = Translation.Z;
+
+	FVector Scale = Transform.GetScale3D();
+	InArray[FirstAttributeIndex + 6] = Scale.X;
+	InArray[FirstAttributeIndex + 7] = Scale.Y;
+	InArray[FirstAttributeIndex + 8] = Scale.Z;
+}
+
+
+FTransformArrayView FRigUnit_RigLogic::TestAccessor::CreateTwoJointNeutralTransforms(float *InValueArray)
+{
+	float* Transform1Ptr = InValueArray;
+
+	FTransform Joint1Transform;
+	Joint1Transform.TransformRotation(FQuat::MakeFromEuler(FVector(1.0f, 0.0f, 0.0f)));
+	Joint1Transform.TransformPosition(FVector(1.0f, 0.0f, 0.0f));
+	Joint1Transform.SetScale3D(FVector(1.0f, 1.0f, 0.0f));
+	AddToTransformArray(Transform1Ptr, Joint1Transform);
+
+	float* Transform2Ptr = InValueArray + MAX_ATTRS_PER_JOINT;
+	FTransform Joint2Transform;
+	Joint2Transform.TransformRotation(FQuat::MakeFromEuler(FVector(1.0f, 0.0f, 0.0f)));
+	Joint2Transform.TransformPosition(FVector(1.0f, 0.0f, 0.0f));
+	Joint2Transform.SetScale3D(FVector(1.0f, 1.0f, 0.0f));
+	AddToTransformArray(Transform2Ptr, Joint2Transform);
+
+	const float* ValuesPtr = InValueArray;
+	return FTransformArrayView(ValuesPtr, sizeof(FTransform));
+}
+
+TArrayView<const uint16> FRigUnit_RigLogic::TestAccessor::CreateTwoJointVariableAttributes(uint16* InVariableAttributeIndices, uint8 LOD)
+{
+	InVariableAttributeIndices[0] = 0;
+	InVariableAttributeIndices[1] = 1;
+	InVariableAttributeIndices[2] = 2;
+	InVariableAttributeIndices[3] = 3;
+	InVariableAttributeIndices[4] = 4;
+	InVariableAttributeIndices[5] = 5;
+	InVariableAttributeIndices[6] = 6;
+	InVariableAttributeIndices[7] = 7;
+	InVariableAttributeIndices[8] = 8;
+
+	if (LOD == 0) //LOD0 includes attributes for both bones
+	{
+		InVariableAttributeIndices[9 + 0] = 9 + 0;
+		InVariableAttributeIndices[9 + 1] = 9 + 1;
+		InVariableAttributeIndices[9 + 2] = 9 + 2;
+		InVariableAttributeIndices[9 + 3] = 9 + 3;
+		InVariableAttributeIndices[9 + 4] = 9 + 4;
+		InVariableAttributeIndices[9 + 5] = 9 + 5;
+		InVariableAttributeIndices[9 + 6] = 9 + 6;
+		InVariableAttributeIndices[9 + 7] = 9 + 7;
+		InVariableAttributeIndices[9 + 8] = 9 + 8;
+	
+		return TArrayView<const uint16>(InVariableAttributeIndices, 2 * FRigUnit_RigLogic::TestAccessor::MAX_ATTRS_PER_JOINT);
+	}
+
+	return TArrayView<const uint16>(InVariableAttributeIndices, FRigUnit_RigLogic::TestAccessor::MAX_ATTRS_PER_JOINT);
+}
+
+void FRigUnit_RigLogic::TestAccessor::Exec_UpdateJoints(FRigHierarchyContainer* TestHierarchyContainer, FRigUnit_RigLogic_JointUpdateParams& JointUpdateParams)
+{
+	Unit->Data.UpdateJoints(TestHierarchyContainer, JointUpdateParams);
+}
+
+IMPLEMENT_RIGUNIT_AUTOMATION_TEST(FRigUnit_RigLogic)
+{
+	FRigUnit_RigLogic::TestAccessor Test(&Unit);
+
+
+	//=============== INPUT CURVES MAPPING ====================
+
+
+	//=== MapInputCurve ValidReader ValidCurvesNameMismatch ===
+
+	//Prepare
+	TestBehaviorReader* TestReaderValid = Test.CreateBehaviorReaderOneCurve("CTRL_Expressions.Some_Control");
+	FRigCurveContainer* TestCurveContainerNameMismatch = Test.CreateCurveContainerOneCurve("CTRL_Expressions_NOT_ThatControl");
+	//Test
+	Test.Exec_MapInputCurve(TestReaderValid, TestCurveContainerNameMismatch);
+	//Assert
+	AddErrorIfFalse(
+		Test.GetData()->InputCurveIndices.Num() == 1 &&
+		Test.GetData()->InputCurveIndices[0] == INDEX_NONE,
+		TEXT("MapInputCurve - ValidReader CurveContainerWithNameMismatch")
+	);
+	delete TestCurveContainerNameMismatch;
+
+	//=== MapInputCurve EmptyReader ValidCurve ===
+
+	//Prepare
+	TestBehaviorReader* TestReaderEmpty = Test.CreateBehaviorReaderEmpty();
+	FRigCurveContainer* TestCurveContainerValid = Test.CreateCurveContainerOneCurve("CTRL_Expressions_Some_Control");
+	//Test
+	Test.Exec_MapInputCurve(TestReaderEmpty, TestCurveContainerValid);
+	//Assert
+	AddErrorIfFalse(
+		Test.GetData()->InputCurveIndices.Num() == 0,
+		TEXT("MapInputCurve - EmptyReader ValidCurveContainer")
+	);
+
+	//=== MapInputCurve ValidReader EmptyCurveContainer ===
+
+	//Prepare
+	FRigCurveContainer* TestCurveContainerEmpty = Test.CreateCurveContainerEmpty();
+	TestCurveContainerEmpty->Initialize();
+	//Test
+	Test.Exec_MapInputCurve(TestReaderValid, TestCurveContainerEmpty);
+	//Assert
+	AddErrorIfFalse(
+		Test.GetData()->InputCurveIndices.Num() == 1 &&
+		Test.GetData()->InputCurveIndices[0] == INDEX_NONE,
+		TEXT("MapInputCurve - ValidReader EmptyCurveContainer")
+	);
+
+	//=== MapInputCurve InvalidReader ValidCurveContainer ===
+	//Prepare
+	TestBehaviorReader* TestInvalidReader = Test.CreateBehaviorReaderOneCurve("InvalidControlNameNoDot");
+	//Expected error
+	AddExpectedError("RigUnit_R: Missing '.' in ");
+	//Test
+	Test.Exec_MapInputCurve(TestInvalidReader, TestCurveContainerValid);
+	//Assert
+	AddErrorIfFalse(
+		Test.GetData()->InputCurveIndices.Num() == 0,
+		TEXT("MapInputCurve - InvalidReader ValidCurveContainer")
+	);
+
+	delete TestInvalidReader;
+
+	//=== MapInputCurve Valid Inputs ===
+
+	//Prepare
+	//  ---
+	//Test
+	Test.Exec_MapInputCurve(TestReaderValid, TestCurveContainerValid);
+	//Assert
+	AddErrorIfFalse(
+		Test.GetData()->InputCurveIndices.Num() == 1 &&
+		Test.GetData()->InputCurveIndices[0] == 0,
+		TEXT("MapInputCurve - Valid Inputs")
+	);
+
+	delete TestReaderValid;
+	delete TestCurveContainerValid;
+
+	//===================== JOINTS MAPPING =====================
+
+
+	//=== MapJoints EmptyInputs ===
+	//Prepare
+	FRigBoneHierarchy* TestHierarchyEmpty = Test.CreateBoneHierarchyEmpty();
+	//Test
+	Test.Exec_MapJoints(TestReaderEmpty, TestHierarchyEmpty);
+	//Assert
+	AddErrorIfFalse(
+		Test.GetData()->HierarchyBoneIndices.Num() == 0,
+		TEXT("MapJoints - Empty Inputs")
+	);
+
+	//=== MapJoints EmptyReader TwoBones ===
+	//Prepare
+	FRigBoneHierarchy* TestHierarchyTwoBones = Test.CreateBoneHierarchyTwoBones("BoneA", "BoneB");
+	//Test
+	Test.Exec_MapJoints(TestReaderEmpty, TestHierarchyTwoBones);
+	AddErrorIfFalse(
+		Test.GetData()->HierarchyBoneIndices.Num() == 0,
+		TEXT("MapJoints - EmptyReader TwoBones")
+	);
+
+	//=== MapJoints TwoJoints NoBones ===
+	//Prepare
+	TestBehaviorReader* TestReaderTwoJoints = Test.CreateBehaviorReaderTwoJoints("BoneA", "BoneB");
+	//Test
+	Test.Exec_MapJoints(TestReaderTwoJoints, TestHierarchyEmpty);
+	AddErrorIfFalse(
+		Test.GetData()->HierarchyBoneIndices.Num() == 2,
+		TEXT("MapJoints - TwoJoints NoBones - expected 2 bone indices")
+	);
+	AddErrorIfFalse(
+		Test.GetData()->HierarchyBoneIndices.Num() == 2 && //repeated condition to prevent crash
+		Test.GetData()->HierarchyBoneIndices[0] == INDEX_NONE,
+		TEXT("MapJoints - TwoJoints NoBones - Expected joint 0 index to be NONE")
+	);
+	AddErrorIfFalse(
+		Test.GetData()->HierarchyBoneIndices.Num() == 2 && //repeated condition to prevent crash
+		Test.GetData()->HierarchyBoneIndices[1] == INDEX_NONE,
+		TEXT("MapJoints - TwoJoints NoBones - Expected joint 1 index to be NONE")
+	);
+
+	delete TestHierarchyEmpty;
+
+	//=== MapJoints TwoJoints TwoBones ===
+	//Prepare
+	//  already done
+	//Test
+	Test.Exec_MapJoints(TestReaderTwoJoints, TestHierarchyTwoBones);
+	//Assert
+	AddErrorIfFalse(
+		Test.GetData()->HierarchyBoneIndices.Num() == 2,
+		TEXT("MapJoints - TwoJoints TwoBones - Expected 2 bone indices")
+	);
+	AddErrorIfFalse(
+		Test.GetData()->HierarchyBoneIndices.Num() == 2 && //prevent crash
+		Test.GetData()->HierarchyBoneIndices[0] == 0,
+		TEXT("MapJoints - TwoJoints TwoBones - Expected bone 0 index to be 0")
+	);
+	AddErrorIfFalse(
+		Test.GetData()->HierarchyBoneIndices.Num() == 2 && //prevent crash
+		Test.GetData()->HierarchyBoneIndices[1] == 1,
+		TEXT("MapJoints - TwoJoints TwoBones - Expected bone index 1 index to be 1")
+	);
+
+	delete TestHierarchyTwoBones;
+	delete TestReaderTwoJoints;
+
+	//===================== BLENDSHAPES MAPPING =====================
+
+
+	//=== MapMorphTargets ValidReader MorphTargetWithNameMismatch ===
+
+	//Prepare
+	TestBehaviorReader* TestReaderBlendshapeValid = Test.CreateBehaviorReaderOneBlendShape("head", "blendshape");
+	FRigCurveContainer* TestMorphTargetNameMismatch = Test.CreateCurveContainerOneMorphTarget("head_NOT_that_blendshape");
+	//Test
+	Test.Exec_MapMorphTargets(TestReaderBlendshapeValid, TestMorphTargetNameMismatch);
+	//Assert
+	AddErrorIfFalse(
+		Test.GetData()->BlendShapeIndices.Num() == 1 && //LOD 0
+		Test.GetData()->MorphTargetCurveIndices.Num() == 1 && //LOD 0
+		Test.GetData()->BlendShapeIndices[0].Values.Num() == 1 && //at least one blendshape exists
+		Test.GetData()->BlendShapeIndices[0].Values[0] == 0 && //has index 0
+		Test.GetData()->MorphTargetCurveIndices[0].Values.Num() == 1 && //but morph target corresponding to that blendshape
+		Test.GetData()->MorphTargetCurveIndices[0].Values[0] == INDEX_NONE, //wasn't found
+		TEXT("MapMorphTargets - ValidReader MorphTargetWithNameMismatch")
+	);
+
+	delete TestMorphTargetNameMismatch;
+
+	//=== MapMorphTargets EmptyReader ValidMorphTargetCurve ===
+
+	//Prepare
+	//Empty reader (no meshes, no blendshapes)
+	FRigCurveContainer* TestMorphTargetCurveValid = Test.CreateCurveContainerOneMorphTarget("head__blendshape");
+	//Test
+	Test.Exec_MapMorphTargets(TestReaderEmpty, TestMorphTargetCurveValid);
+	//Assert
+	AddErrorIfFalse(
+		Test.GetData()->MorphTargetCurveIndices.Num() == 0 &&
+		Test.GetData()->BlendShapeIndices.Num() == 0,
+		TEXT("MapMorphTargets - EmptyReader ValidMorphTargetCurve")
+	);
+
+
+	//=== MapMorphTargets NoBlendShapes ValidMorphTargetCurve ===
+
+	//Prepare
+	TestBehaviorReader* TestReaderNoBlendshapes = Test.CreateBehaviorReaderNoBlendshapes("head"); //has a mesh, but no blend shapes
+	//Test
+	Test.Exec_MapMorphTargets(TestReaderNoBlendshapes, TestMorphTargetCurveValid);
+	//Assert
+	AddErrorIfFalse(
+		Test.GetData()->BlendShapeIndices.Num() == 1 && //LOD 0 exists
+		Test.GetData()->BlendShapeIndices[0].Values.Num() == 0 && //but no blend shapes mapped
+		Test.GetData()->MorphTargetCurveIndices[0].Values.Num() == 0, //or morph targets
+		TEXT("MapMorphTargets - NoBlendShapes ValidMorphTargetCurve")
+	);
+
+	delete TestReaderNoBlendshapes;
+
+	//=== MapMorphTargets ValidReader EmptyCurveContainer ===
+
+	//Prepare
+	//  ---
+	//Test
+	Test.Exec_MapMorphTargets(TestReaderBlendshapeValid, TestCurveContainerEmpty);
+	//Assert
+	AddErrorIfFalse(
+		Test.GetData()->BlendShapeIndices.Num() == 1 && //LOD 0
+		Test.GetData()->BlendShapeIndices[0].Values.Num() == 1 && //one blend shape
+		Test.GetData()->BlendShapeIndices[0].Values[0] == 0 && //of index 0
+		Test.GetData()->MorphTargetCurveIndices[0].Values.Num() == 1 && //we put in a morph target index corresponding to it
+		Test.GetData()->MorphTargetCurveIndices[0].Values[0] == INDEX_NONE, //but just to signal that it is not found
+		TEXT("MapMorphTargets - ValidReader EmptyCurveContainer")
+	);
+
+	//=== MapMorphTargets InvalidReader ValidMorphTargetCurve ===
+
+	//Prepare
+	TestBehaviorReader* TestReaderBlendshapesInvalid = Test.CreateBehaviorReaderOneBlendShape("head", "");
+	//Test
+	Test.Exec_MapMorphTargets(TestReaderBlendshapesInvalid, TestMorphTargetCurveValid);
+	//Assert
+	AddErrorIfFalse(
+		Test.GetData()->BlendShapeIndices.Num() == 1 && //LOD 0
+		Test.GetData()->BlendShapeIndices[0].Values.Num() == 1 && //one blend shape (empty named)
+		Test.GetData()->BlendShapeIndices[0].Values[0] == 0 && //of index 0
+		Test.GetData()->MorphTargetCurveIndices[0].Values.Num() == 1 && //we put in a morph target index corresponding to it
+		Test.GetData()->MorphTargetCurveIndices[0].Values[0] == INDEX_NONE, //but just to signal that it is not found
+		TEXT("MapMorphTargets - InvalidReader ValidMorphTargetCurve")
+	);
+
+	delete TestReaderBlendshapesInvalid;
+
+	//=== MapMorphTargets ValidReader InvalidMorphTargetCurve ===
+
+	//Prepare
+	FRigCurveContainer* TestMorphTargetCurvesInvalid = Test.CreateCurveContainerOneMorphTarget("");
+	//Test
+	Test.Exec_MapMorphTargets(TestReaderBlendshapeValid, TestMorphTargetCurvesInvalid);
+	//Assert
+	AddErrorIfFalse(
+		Test.GetData()->BlendShapeIndices.Num() == 1 && //LOD 0
+		Test.GetData()->BlendShapeIndices[0].Values.Num() == 1 && //one blend shape
+		Test.GetData()->BlendShapeIndices[0].Values[0] == 0 && //of index 0
+		Test.GetData()->MorphTargetCurveIndices[0].Values.Num() == 1 && //we put in a morph target index corresponding to it
+		Test.GetData()->MorphTargetCurveIndices[0].Values[0] == INDEX_NONE, //but just to signal that it is not found
+		TEXT("MapMorphTargets - ValidReader InvalidMorphTargetCurve")
+	);
+
+	delete TestMorphTargetCurvesInvalid;
+
+	//=== MapMorphTargets Valid Inputs ===
+
+	//Prepare
+	//  ---
+	//Test
+	Test.Exec_MapMorphTargets(TestReaderBlendshapeValid, TestMorphTargetCurveValid);
+	//Assert
+	AddErrorIfFalse(
+		Test.GetData()->BlendShapeIndices.Num() == 1 && //LOD 0
+		Test.GetData()->MorphTargetCurveIndices.Num() == 1 && //LOD 0
+		Test.GetData()->BlendShapeIndices[0].Values.Num() == 1 && //at least one blendshape exists
+		Test.GetData()->BlendShapeIndices[0].Values[0] == 0 && //has index 0
+		Test.GetData()->MorphTargetCurveIndices[0].Values.Num() == 1 && //morph target corresponding to that blendshape exists
+		Test.GetData()->MorphTargetCurveIndices[0].Values[0] == 0, //and actually points to the right index
+		TEXT("MapMorphTargets - ValidReader ValidTestMorphTarget")
+	);
+
+	delete TestReaderBlendshapeValid;
+	delete TestMorphTargetCurveValid;
+
+	//=== MapMorphTargets LOD0(AB) LOD1(A) ===
+
+	//Prepare
+	TestBehaviorReader* TestReaderBlendshapes_LOD0AB_LOD1A = Test.CreateBehaviorReaderTwoBlendShapes("head", "blendshapeA", "blendshapeB");
+	FRigCurveContainer* TestMorphTargetTwoCurves = Test.CreateCurveContainerTwoMorphTargets("head__blendshapeA", "head__blendshapeB");
+	//NOTE: indices in the first param here are not blendshape indices, but rather mappings from blendshapes to meshes
+	//in this test, they will correspond to blendshape indices
+
+
+	TestReaderBlendshapes_LOD0AB_LOD1A->addBlendShapeMappingIndicesToLOD(0, 0); //A -> LOD 0
+	TestReaderBlendshapes_LOD0AB_LOD1A->addBlendShapeMappingIndicesToLOD(1, 0); //B -> -||-
+	TestReaderBlendshapes_LOD0AB_LOD1A->addBlendShapeMappingIndicesToLOD(0, 1); //A -> LOD 1
+
+	TestReaderBlendshapes_LOD0AB_LOD1A->LODCount = 2; //needs to be set explicitly if not default (=1)
+	//Test
+	Test.Exec_MapMorphTargets(TestReaderBlendshapes_LOD0AB_LOD1A, TestMorphTargetTwoCurves);
+	//Assert
+	AddErrorIfFalse(
+		Test.GetData()->BlendShapeIndices.Num() == 2 &&  //2 LODs
+		Test.GetData()->MorphTargetCurveIndices.Num() == 2,
+		TEXT("MapMorphTargets LOD0(AB) LOD1(A) - Expected 2 LODs for both blendshapes and morph targets")
+	);
+
+	AddErrorIfFalse(
+		Test.GetData()->BlendShapeIndices.Num() == 2 &&         //condition repeated for crash prevention
+		Test.GetData()->BlendShapeIndices[0].Values.Num() == 2 && //two blendshapes at LOD 0
+		Test.GetData()->BlendShapeIndices[0].Values[0] == 0 &&  // A
+		Test.GetData()->BlendShapeIndices[0].Values[1] == 1 &&  // B
+		Test.GetData()->BlendShapeIndices[1].Values.Num() == 1 && //one blendshape at LOD 1
+		Test.GetData()->BlendShapeIndices[1].Values[0] == 0,    // A
+			TEXT("MapMorphTargets LOD0(AB) LOD1(A) - resulting blendshape indices not correct")
+	);
+
+	AddErrorIfFalse(
+		Test.GetData()->MorphTargetCurveIndices.Num() == 2 &&         //condition repeated for crash prevention
+		Test.GetData()->MorphTargetCurveIndices[0].Values.Num() == 2 && //two morph targets at LOD 0
+		Test.GetData()->MorphTargetCurveIndices[0].Values[0] == 0 &&  // A
+		Test.GetData()->MorphTargetCurveIndices[0].Values[1] == 1 &&  // B
+		Test.GetData()->MorphTargetCurveIndices[1].Values.Num() == 1 && //one morph target at LOD 1
+		Test.GetData()->MorphTargetCurveIndices[1].Values[0] == 0,    // A
+		TEXT("MapMorphTargets LOD0(AB) LOD1(A) - resulting morph target indices not correct")
+	);
+
+	delete TestReaderBlendshapes_LOD0AB_LOD1A;
+
+	//=== MapMorphTargets LOD0(AB) LOD1(-) ===
+
+	//Prepare
+	TestBehaviorReader* TestReaderBlendshapes_LOD0AB_LOD1N = Test.CreateBehaviorReaderTwoBlendShapes("head", "blendshapeA", "blendshapeB");
+	TestReaderBlendshapes_LOD0AB_LOD1N->addBlendShapeMappingIndicesToLOD(0, 0); //LOD 0
+	TestReaderBlendshapes_LOD0AB_LOD1N->addBlendShapeMappingIndicesToLOD(1, 0); //LOD 0
+	//LODCount = 1 by default
+
+	//Test
+	Test.Exec_MapMorphTargets(TestReaderBlendshapes_LOD0AB_LOD1N, TestMorphTargetTwoCurves);
+	//Assert
+	AddErrorIfFalse(
+		Test.GetData()->BlendShapeIndices.Num() == 1 &&  //1 LOD
+		Test.GetData()->MorphTargetCurveIndices.Num() == 1, //1 LOD
+		TEXT("MapMorphTargets LOD0(AB) LOD1(-) - Expected 1 LOD for both blendshapes and morph targets")
+	);
+
+
+	AddErrorIfFalse(
+		Test.GetData()->BlendShapeIndices.Num() == 1 &&         //condition repeated for crash prevention
+		Test.GetData()->BlendShapeIndices[0].Values.Num() == 2 && //two blendshapes at LOD 0
+		Test.GetData()->BlendShapeIndices[0].Values[0] == 0 &&  // A
+		Test.GetData()->BlendShapeIndices[0].Values[1] == 1,    // B
+		TEXT("MapMorphTargets LOD0(AB) LOD1(-) - Resulting blendshapes not correct")
+	);
+
+	AddErrorIfFalse(
+		Test.GetData()->MorphTargetCurveIndices.Num() == 1 &&         //condition repeated for crash prevention
+		Test.GetData()->MorphTargetCurveIndices[0].Values.Num() == 2 && //two morph targets at LOD 0
+		Test.GetData()->MorphTargetCurveIndices[0].Values[0] == 0 &&  // A
+		Test.GetData()->MorphTargetCurveIndices[0].Values[1] == 1,    // B
+		TEXT("MapMorphTargets LOD0(AB) LOD1(-) - Resulting morph targets not correct")
+	);
+
+	delete TestReaderBlendshapes_LOD0AB_LOD1N;
+
+	//=== MapMorphTargets LOD0(A) LOD1(B) ===
+	//Prepare
+	TestBehaviorReader* TestReaderBlendshapes_LOD0A_LOD1B = Test.CreateBehaviorReaderTwoBlendShapes("head", "blendshapeA", "blendshapeB");
+	TestReaderBlendshapes_LOD0A_LOD1B->addBlendShapeMappingIndicesToLOD(0, 0);
+	TestReaderBlendshapes_LOD0A_LOD1B->addBlendShapeMappingIndicesToLOD(1, 1);
+	TestReaderBlendshapes_LOD0A_LOD1B->LODCount = 2;
+
+	//Test
+	Test.Exec_MapMorphTargets(TestReaderBlendshapes_LOD0A_LOD1B, TestMorphTargetTwoCurves);
+	//Assert
+	AddErrorIfFalse(
+		Test.GetData()->BlendShapeIndices.Num() == 2 &&  //2 LODs
+		Test.GetData()->MorphTargetCurveIndices.Num() == 2, //2 LODs
+		TEXT("MapMorphTargets LOD0(A) LOD1(B) - Expected 2 LODs for both blendshapes and morph targets")
+	);
+
+	AddErrorIfFalse(
+		Test.GetData()->BlendShapeIndices.Num() == 2 &&         //condition repeated for crash prevention
+		Test.GetData()->BlendShapeIndices[0].Values.Num() == 1 && //1 blendshape at LOD 0
+		Test.GetData()->BlendShapeIndices[0].Values[0] == 0 &&  // A
+		Test.GetData()->BlendShapeIndices[1].Values.Num() == 1 && //1 blendshape at LOD 1
+		Test.GetData()->BlendShapeIndices[1].Values[0] == 1,    // B
+		TEXT("MapMorphTargets LOD0(A) LOD1(B) - Resulting blendshape indices not correct")
+	);
+
+	AddErrorIfFalse(
+		Test.GetData()->MorphTargetCurveIndices.Num() == 2 &&         //condition repeated for crash prevention
+		Test.GetData()->MorphTargetCurveIndices[0].Values.Num() == 1 && //1 morph target at LOD 0
+		Test.GetData()->MorphTargetCurveIndices[0].Values[0] == 0 &&  // A
+		Test.GetData()->MorphTargetCurveIndices[1].Values.Num() == 1 && //1 morph target at LOD 1
+		Test.GetData()->MorphTargetCurveIndices[1].Values[0] == 1,    // B
+		TEXT("MapMorphTargets LOD0(A) LOD1(B) - Resulting morph target indices not correct")
+	);
+
+	delete TestMorphTargetTwoCurves;
+	delete TestReaderBlendshapes_LOD0A_LOD1B;
+
+	//=============== MASK MULTIPLIERS MAPPING ====================
+
+
+	//=== MapMaskMultipliers ValidReader ValidAnimatedMapNameMismatch ===
+
+	//Prepare
+	TestBehaviorReader* TestReaderAnimMapsValid = Test.CreateBehaviorReaderOneAnimatedMap("CTRL_AnimMap.Some_Multiplier");
+	FRigCurveContainer* TestCurveContainerForAnimMapsNameMismatch = Test.CreateCurveContainerOneCurve("CTRL_AnimMap_NOT_ThatMultiploer");
+	//Test
+	Test.Exec_MapMaskMultipliers(TestReaderAnimMapsValid, TestCurveContainerForAnimMapsNameMismatch);
+	//Assert
+	AddErrorIfFalse(
+		Test.GetData()->CurveContainerIndicesForAnimMaps.Num() == 1 &&
+		Test.GetData()->CurveContainerIndicesForAnimMaps[0].Values.Num() == 1 &&
+		Test.GetData()->CurveContainerIndicesForAnimMaps[0].Values[0] == INDEX_NONE &&
+		Test.GetData()->RigLogicIndicesForAnimMaps.Num() == 1 &&
+		Test.GetData()->RigLogicIndicesForAnimMaps[0].Values.Num() == 1 &&
+		Test.GetData()->RigLogicIndicesForAnimMaps[0].Values[0] == 0,
+		TEXT("MapMaskMultipliers - ValidReader ValidAnimatedMapNameMismatch")
+	);
+
+	delete TestCurveContainerForAnimMapsNameMismatch;
+
+	//=== MapMaskMultipliers EmptyReader ValidAnimatedMap ===
+
+	//Prepare
+	FRigCurveContainer* TestCurveContainerForAnimMapsValid = Test.CreateCurveContainerOneCurve("CTRL_AnimMap_Some_Multiplier");
+	//Test
+	Test.Exec_MapMaskMultipliers(TestReaderEmpty, TestCurveContainerForAnimMapsValid);
+	//Assert
+	AddErrorIfFalse(
+		Test.GetData()->CurveContainerIndicesForAnimMaps.Num() == 0 &&
+		Test.GetData()->RigLogicIndicesForAnimMaps.Num() == 0,
+		TEXT("MapMaskMultipliers - EmptyReader ValidAnimatedMap")
+	);
+
+	delete TestReaderEmpty;
+
+	//=== MapMaskMultipliers ValidReader EmptyCurveContainer ===
+
+	//Test
+	Test.Exec_MapMaskMultipliers(TestReaderAnimMapsValid, TestCurveContainerEmpty);
+	//Assert
+	AddErrorIfFalse(
+		Test.GetData()->CurveContainerIndicesForAnimMaps.Num() == 1 &&
+		Test.GetData()->CurveContainerIndicesForAnimMaps[0].Values.Num() == 1 &&
+		Test.GetData()->CurveContainerIndicesForAnimMaps[0].Values[0] == INDEX_NONE &&
+		Test.GetData()->RigLogicIndicesForAnimMaps.Num() == 1 &&
+		Test.GetData()->RigLogicIndicesForAnimMaps[0].Values.Num() == 1 &&
+		Test.GetData()->RigLogicIndicesForAnimMaps[0].Values[0] == 0,
+		TEXT("MapMaskMultipliers - ValidReader EmptyCurveContainer")
+	);
+
+	delete TestCurveContainerEmpty;
+
+	//=== MapMaskMultipliers Valid Inputs ===
+
+	//Prepare
+	//  ---
+	//Test
+	Test.Exec_MapMaskMultipliers(TestReaderAnimMapsValid, TestCurveContainerForAnimMapsValid);
+	//Assert
+	AddErrorIfFalse(
+		Test.GetData()->CurveContainerIndicesForAnimMaps.Num() == 1 &&
+		Test.GetData()->CurveContainerIndicesForAnimMaps[0].Values.Num() == 1 &&
+		Test.GetData()->CurveContainerIndicesForAnimMaps[0].Values[0] == 0 &&
+		Test.GetData()->RigLogicIndicesForAnimMaps.Num() == 1 &&
+		Test.GetData()->RigLogicIndicesForAnimMaps[0].Values.Num() == 1 &&
+		Test.GetData()->RigLogicIndicesForAnimMaps[0].Values[0] == 0,
+		TEXT("MapMaskMultipliers - Valid Inputs")
+	);
+
+	delete TestReaderAnimMapsValid;
+	delete TestCurveContainerForAnimMapsValid;
+
+	Unit.ExecuteContext.Hierarchy = &HierarchyContainer;
+	//BoneHierarchy belongs to HierarchyContainer
+	BoneHierarchy = *TestHierarchyTwoBones;
+	BoneHierarchy.Initialize();
+	Unit.ExecuteContext.Hierarchy = &HierarchyContainer;
+	BoneHierarchy.ResetTransforms();
+	
+	FRigHierarchyContainer* HierarchyContainerPtr = &HierarchyContainer;
+
+
+	//Prepare
+	//-----
+	//create neutral transforms for two bones
+	const uint8 TransformArraySize = 2 * FRigUnit_RigLogic::TestAccessor::MAX_ATTRS_PER_JOINT;
+	float Values[TransformArraySize];  //two bones, nine attributes
+	FTransformArrayView TwoJointNeutralTransforms = Test.CreateTwoJointNeutralTransforms(Values);
+	//create delta transforms
+	float DeltaTransformData[TransformArraySize] = { 0.f };
+	//first bone translation
+	DeltaTransformData[0] = 1.f;
+	DeltaTransformData[1] = 0.f;
+	DeltaTransformData[2] = 0.f;
+	//second bone translation
+	DeltaTransformData[9] = 1.f;
+	DeltaTransformData[10] = 2.f;
+	DeltaTransformData[11] = 7.f;
+	FTransformArrayView DeltaTransforms = FTransformArrayView(DeltaTransformData, sizeof(FTransform));
+	//create variable attribute index arrays for two bones
+	uint16 VariableAttributeIndices_LOD0[TransformArraySize];
+	TArrayView<const uint16> VariableAttributes_LOD0 = Test.CreateTwoJointVariableAttributes(VariableAttributeIndices_LOD0, 0); //LOD0: both bones included
+	FRigUnit_RigLogic_JointUpdateParams TestJointUpdateParamsTwoJoints_LOD0(
+		VariableAttributes_LOD0,
+		TwoJointNeutralTransforms,
+		DeltaTransforms);
+	Test.GetData()->UpdatedJoints.SetNumZeroed(2);
+	//Test
+	Test.Exec_UpdateJoints(HierarchyContainerPtr, TestJointUpdateParamsTwoJoints_LOD0);
+	//Assert
+	//Note that BoneB.GlobalTransform.Z should be zero since the scale.Z is zero. Also, translation Y becomes -Y 
+	AddErrorIfFalse(BoneHierarchy.GetGlobalTransform(0).GetTranslation().Equals(FVector(1.f, 0.f, 0.f)), TEXT("UpdateJoints LOD0 Bone 01 - unexpected transform"));
+	AddErrorIfFalse(BoneHierarchy.GetGlobalTransform(1).GetTranslation().Equals(FVector(2.f, -2.f, 0.f)), TEXT("UpdateJoints LOD0 Bone 02 - unexpected transform"));
+
+
+	//Prepare
+	//-----
+	//Create skeleton, skeletal mesh and skeletal mesh component
+	USkeleton* Skeleton = NewObject<USkeleton>(GetTransientPackage());
+	USkeletalMesh* SkeletalMesh = NewObject<USkeletalMesh>(GetTransientPackage());
+	SkeletalMesh->Skeleton = Skeleton;
+	USkeletalMeshComponent* SkelMeshComponent = NewObject<USkeletalMeshComponent>();
+	SkelMeshComponent->SetSkeletalMesh(SkeletalMesh);
+
+	const FString DNAAssetFileName(TEXT("rl_unit_behavior_test.dna"));
+	const FString DNAFolder = FPackageName::LongPackageNameToFilename(TEXT("/CharacterParts/Face/DNA/Test"));
+	FString FullFolderPath = FPaths::ConvertRelativePathToFull(DNAFolder);
+	FString DNAFilePath = FPaths::Combine(FullFolderPath, DNAAssetFileName);
+	UDNAAsset* MockDNAAsset = NewObject< UDNAAsset >(SkeletalMesh, FName(*DNAAssetFileName)); //SkelMesh has to be its outer, otherwise DNAAsset won't be saved			
+	if (MockDNAAsset->Init(DNAFilePath))  //will set BehaviorReader we need to execute the rig unit
+	{
+		UAssetUserData* DNAAssetUserData = Cast<UAssetUserData>(MockDNAAsset);
+		SkeletalMesh->AddAssetUserData(DNAAssetUserData);
+	}
+
+	Test.GetData()->SkelMeshComponent = SkelMeshComponent;
+	//Test
+	InitAndExecute();
+	//Assert
+	AddErrorIfFalse(
+		// Check rig logic initialized
+		Test.GetData()->RigLogic != nullptr &&
+		Test.GetData()->RigInstance != nullptr &&
+		
+		// Check joints
+		Test.GetData()->HierarchyBoneIndices.Num() > 0 &&
+		
+		// Check input curves
+		Test.GetData()->InputCurveIndices.Num() > 0 &&
+		
+		// Check morph targets
+		Test.GetData()->MorphTargetCurveIndices.Num() > 0 &&
+		Test.GetData()->BlendShapeIndices.Num() > 0 &&
+
+		// Check mask multipliers
+		Test.GetData()->RigLogicIndicesForAnimMaps.Num() > 0 &&
+		Test.GetData()->CurveContainerIndicesForAnimMaps.Num() > 0,
+
+		TEXT("InitAndExecute failed to initialize rig logic.")
+	);
+
+	return true;
+}
+
+#endif
diff --git a/Engine/Plugins/Animation/RigLogic/Source/RigLogicModule/Private/SkelMeshDNAUtils.cpp b/Engine/Plugins/Animation/RigLogic/Source/RigLogicModule/Private/SkelMeshDNAUtils.cpp
new file mode 100644
--- /dev/null
+++ b/Engine/Plugins/Animation/RigLogic/Source/RigLogicModule/Private/SkelMeshDNAUtils.cpp
@@ -0,0 +1,431 @@
+// Copyright 1998-2019 Epic Games, Inc. All Rights Reserved.
+
+#include "SkelMeshDNAUtils.h"
+#include "DNAToSkelMeshMap.h"
+
+#include "RenderResource.h"
+#include "RHICommandList.h"
+#include "Async/ParallelFor.h"
+#include "Components/SkeletalMeshComponent.h"
+#include "Engine/StaticMesh.h"
+#include "Misc/Paths.h"
+#include "ProfilingDebugging/ScopedTimers.h"
+#include "Misc/FileHelper.h"
+#include "UObject/ConstructorHelpers.h"
+#include "Engine/SkeletalMesh.h"
+#include "ComponentReregisterContext.h"
+#include "ReferenceSkeleton.h"
+#include "Rendering/SkeletalMeshModel.h"
+#include "Rendering/SkeletalMeshLODModel.h"
+#include "Rendering/SkeletalMeshRenderData.h"
+#include "Rendering/SkeletalMeshLODRenderData.h"
+
+#include "AnimationRuntime.h"
+
+#include "riglogic/RigLogic.h"
+
+DEFINE_LOG_CATEGORY(LogDNAUtils);
+
+USkelMeshDNAUtils::USkelMeshDNAUtils(const FObjectInitializer& ObjectInitializer)
+	: Super(ObjectInitializer)
+{
+}
+
+FDNAToSkelMeshMap* USkelMeshDNAUtils::CreateMapForUpdatingNeutralMesh(IDNAReader* DNAReader, USkeletalMesh* SkelMesh)
+{
+#if WITH_EDITORONLY_DATA
+	FDNAToSkelMeshMap* DNAToSkelMeshMap = new FDNAToSkelMeshMap();
+
+	//only vertex map is initialized in this pass so we can mix neutral meshes fast (e.g. on slider move);
+	//playing animation on a such a mesh requires updating joints and skin weights
+	//getting full quality animation requires mixing morph targets too
+	DNAToSkelMeshMap->InitBaseMesh(DNAReader, SkelMesh);
+	return DNAToSkelMeshMap;
+#else
+	return nullptr;
+#endif // WITH_EDITORONLY_DATA
+}
+
+#if WITH_EDITORONLY_DATA
+/** Updates bind pose using joint positions from DNA. */
+void USkelMeshDNAUtils::UpdateJoints(USkeletalMesh* SkelMesh, IDNAReader* DNAReader, FDNAToSkelMeshMap* DNAToSkelMeshMap)
+{
+	{	// Scoping of RefSkelModifier
+		FReferenceSkeletonModifier RefSkelModifier(SkelMesh->RefSkeleton, SkelMesh->Skeleton);
+
+		// copy here
+		TArray<FTransform> RawBonePose = SkelMesh->RefSkeleton.GetRawRefBonePose();
+
+		// When we are mounting the head to different bodies than female average, we need to use
+		// component space, as the joint to which the head root is snapped to will be on a different
+		// position than in the head rig!
+
+		// calculate component space ahead of current transform
+		TArray<FTransform> ComponentTransforms;
+		FAnimationRuntime::FillUpComponentSpaceTransforms(SkelMesh->RefSkeleton, RawBonePose, ComponentTransforms);
+
+		const TArray<FMeshBoneInfo>& RawBoneInfo = SkelMesh->RefSkeleton.GetRawRefBoneInfo();
+
+		// Skipping root joint (index 0) to avoid blinking of the mesh due to bounding box issue
+		for (uint16 JointIndex = 0; JointIndex < DNAReader->GetJointCount(); JointIndex++)
+		{
+			int32 BoneIndex = DNAToSkelMeshMap->GetUEBoneIndex(JointIndex);
+
+			FTransform DNATransform = FTransform::Identity;
+
+			// Updating bind pose affects just translations.
+			FVector Translate = DNAReader->GetNeutralJointTranslation(JointIndex);
+			DNATransform.SetTranslation(Translate);
+			FVector RotationVector = DNAReader->GetNeutralJointRotation(JointIndex);
+			FRotator Rotation(RotationVector.X, RotationVector.Y, RotationVector.Z);
+
+			// Joint 0 (spine_04) is a root of GeneSplicer joint hierarchy, and is a special case
+			// 1) it is parent to itself
+			// 2) it is in DNA space, so we need to rotate it 90 degs on x axis to UE4 space
+			// 3) the head joints below it in the skeletal mesh are not spliced, as they are not in DNA,
+			//    so they will retain female average transforms
+
+			if (DNAReader->GetJointParentIndex( JointIndex) == JointIndex )  //parent to itself
+			{
+				Rotation.Pitch += 90;
+				DNATransform.SetRotation(Rotation.Quaternion());
+
+				DNATransform.SetTranslation(FVector(Translate.X, Translate.Z, -Translate.Y));
+
+				ComponentTransforms[BoneIndex] = DNATransform;
+
+			}
+			else
+			{
+				DNATransform.SetRotation(Rotation.Quaternion());
+
+				if (ensure(RawBoneInfo[BoneIndex].ParentIndex != INDEX_NONE))
+				{
+					ComponentTransforms[BoneIndex] = DNATransform * ComponentTransforms[RawBoneInfo[BoneIndex].ParentIndex];
+				}
+			}
+
+			ComponentTransforms[BoneIndex].NormalizeRotation();
+		}
+
+		for (uint16 BoneIndex = 0; BoneIndex < RawBoneInfo.Num(); BoneIndex++)
+		{
+			FTransform LocalTransform;
+			        
+			if (BoneIndex == 0)
+			{
+				LocalTransform = ComponentTransforms[BoneIndex];
+			}
+			else
+			{
+				LocalTransform = ComponentTransforms[BoneIndex].GetRelativeTransform(ComponentTransforms[RawBoneInfo[BoneIndex].ParentIndex]);
+			}
+
+			LocalTransform.NormalizeRotation();
+
+			RefSkelModifier.UpdateRefPoseTransform(BoneIndex, LocalTransform);
+		}
+	}
+
+	SkelMesh->RefBasesInvMatrix.Reset();
+	SkelMesh->CalculateInvRefMatrices(); // Needs to be called after RefSkelModifier is destroyed
+}
+
+/** Updates base mesh vertices using data from DNA. */
+void USkelMeshDNAUtils::UpdateBaseMesh(USkeletalMesh* SkelMesh, IDNAReader* DNAReader, FDNAToSkelMeshMap* DNAToSkelMeshMap, ELodUpdateOption UpdateOption)
+{
+
+	FSkeletalMeshModel* ImportedModel = SkelMesh->GetImportedModel();
+	// Expects vertex map to be initialized beforehand	
+	int32 LODStart;
+	int32 LODEnd;
+	GetLODRange(UpdateOption, ImportedModel->LODModels.Num(), LODStart, LODEnd);
+	for (int32 LODIndex = LODStart; LODIndex < LODEnd; LODIndex++)
+	{
+		FSkeletalMeshLODModel& LODModel = ImportedModel->LODModels[LODIndex];
+		int32 SectionIndex = 0;
+		for (FSkelMeshSection& Section : LODModel.Sections)
+		{
+			int32& DNAMeshIndex = DNAToSkelMeshMap->ImportVtxToDNAMeshIndex[LODIndex][Section.GetVertexBufferIndex()];
+
+			const int32 NumSoftVertices = Section.GetNumVertices();
+			auto& OverlappingMap = DNAToSkelMeshMap->OverlappingVertices[LODIndex][SectionIndex];
+			int32 VertexBufferIndex = Section.GetVertexBufferIndex();
+			for (int32 VertexIndex = 0; VertexIndex < NumSoftVertices; VertexIndex++)
+			{				
+				int32& DNAVertexIndex = DNAToSkelMeshMap->ImportVtxToDNAVtxIndex[LODIndex][VertexBufferIndex];
+
+				if (DNAVertexIndex >= 0)
+				{
+					const FVector Position = DNAReader->GetVertexPosition(DNAMeshIndex, DNAVertexIndex);
+					FSoftSkinVertex& Vertex = Section.SoftVertices[VertexIndex];
+					Vertex.Position = Position;
+
+					// Check if the current vertex has overlapping vertices, and then update them as well.
+					TArray<int32>& OverlappedIndices = OverlappingMap[VertexIndex];
+					int32 OverlappingCount = OverlappedIndices.Num();
+					for (int32 OverlappingIndex = 0; OverlappingIndex < OverlappingCount; ++OverlappingIndex)
+					{
+						int32 OverlappingVertexIndex = OverlappedIndices[OverlappingIndex];
+						FSoftSkinVertex& OverlappingVertex = Section.SoftVertices[OverlappingVertexIndex];
+						OverlappingVertex.Position = Position;
+					}
+				}
+				VertexBufferIndex++;
+			}
+			SectionIndex++;
+		}
+	}
+}
+
+/** Updates Morph Targets using Blend Shapes from DNA.  */
+void USkelMeshDNAUtils::UpdateMorphTargets(USkeletalMesh* SkelMesh, IDNAReader* DNAReader, FDNAToSkelMeshMap* DNAToSkelMeshMap, ELodUpdateOption UpdateOption)
+{
+	TArray<FDNABlendShapeTarget>& MeshBlendShapeTargets = DNAToSkelMeshMap->GetMeshBlendShapeTargets();
+	if (MeshBlendShapeTargets.Num() == 0)
+	{
+		UE_LOG(LogDNAUtils, Warning, TEXT("No morph targets updated!"));
+		return;
+	}
+	TArray<TArray<TArray<int32>>>& BlendShapeVertexDeltas = DNAToSkelMeshMap->GetBlendShapeVertexDeltaIndices();
+	ParallelFor(SkelMesh->MorphTargets.Num(), [&](int32 MorphIndex)
+	{
+		UMorphTarget* MorphTarget = SkelMesh->MorphTargets[MorphIndex];
+		const FDNABlendShapeTarget& MeshTarget = MeshBlendShapeTargets[MorphIndex];
+		// First get all DNA deltas for current Morph Target.
+		TArrayView<const uint32> BlendShapeVertexIndices = DNAReader->GetBlendShapeTargetVertexIndices(MeshTarget.MeshIndex, MeshTarget.TargetIndex);
+		if (BlendShapeVertexIndices.Num() > 0)
+		{
+			int32 LODStart;
+			int32 LODEnd;
+			GetLODRange(UpdateOption, MorphTarget->MorphLODModels.Num(), LODStart, LODEnd);
+			for (int32 LODIndex = 0; LODIndex < LODEnd; LODIndex++)
+			{
+				// MorphTarget vertex indices refer to full vertex index buffer of imported mesh.
+				FMorphTargetLODModel& MorphLODModel = MorphTarget->MorphLODModels[LODIndex];
+				int32 NumDeltas = MorphLODModel.Vertices.Num();
+				FMorphTargetDelta* Deltas = MorphLODModel.Vertices.GetData();
+				for (int32 DeltaIndex = 0; DeltaIndex < NumDeltas; DeltaIndex++)
+				{
+					uint32 DeltaSourceIndex = Deltas[DeltaIndex].SourceIdx;
+					uint32 DNAVertexIndex = DNAToSkelMeshMap->ImportVtxToDNAVtxIndex[LODIndex][DeltaSourceIndex];
+					if (DNAVertexIndex >= 0)
+					{
+						// Retrieve delta index from the mapping.
+						int32 DNADeltaIndex = BlendShapeVertexDeltas[MeshTarget.MeshIndex][DNAVertexIndex][MeshTarget.TargetIndex];
+						if (DNADeltaIndex != INDEX_NONE)
+						{
+							Deltas[DeltaIndex].PositionDelta = DNAReader->GetBlendShapeTargetDelta(MeshTarget.MeshIndex, MeshTarget.TargetIndex, DNADeltaIndex);
+						}
+						else
+						{
+							#ifdef DEBUG
+							UE_LOG(LogDNAUtils, Error, TEXT(" Deltas not found for DNA Vertex %d (UE Vertex %d), Mesh %d and Target Index %d"), DNAVertexIndex, DeltaSourceIndex, MeshTarget.MeshIndex, MeshTarget.TargetIndex);
+							#endif 
+							// Delta not found. Reset it to (0, 0, 0).
+							Deltas[DeltaIndex].PositionDelta = FVector(0, 0, 0);
+						}
+					}
+					else
+					{
+						#ifdef DEBUG
+						UE_LOG(LogDNAUtils, Error, TEXT(" UE Vertex %d not found in DNA for Mesh %d and TargetIndex %d"), DeltaSourceIndex, MeshTarget.MeshIndex, MeshTarget.TargetIndex);
+						#endif 
+					}
+				}
+			}
+		}
+		else
+		{
+			#ifdef DEBUG
+			UE_LOG(LogDNAUtils, Warning, TEXT(" 0 deltas found for mesh %d and blend shape target %d"), MeshTarget.MeshIndex, MeshTarget.TargetIndex);
+			#endif 
+		}
+	});	
+}
+
+/* Updates Bone influences using Skin Weights from DNA. */
+void USkelMeshDNAUtils::UpdateSkinWeights(USkeletalMesh* SkelMesh, IDNAReader* DNAReader, FDNAToSkelMeshMap* DNAToSkelMeshMap, ELodUpdateOption UpdateOption)
+{
+	FSkeletalMeshModel* ImportedModel = SkelMesh->GetImportedModel();
+
+	bool InfluenceMismatch = false;
+
+	int32 LODStart;
+	int32 LODEnd;
+	GetLODRange(UpdateOption, ImportedModel->LODModels.Num(), LODStart, LODEnd);
+	for (int32 LODIndex = LODStart; LODIndex < LODEnd; ++LODIndex)
+	{
+		FSkeletalMeshLODModel& LODModel = ImportedModel->LODModels[LODIndex];
+		for (int32 SectionIndex = 0; SectionIndex < LODModel.Sections.Num(); SectionIndex++)
+		{
+			FSkelMeshSection& Section = LODModel.Sections[SectionIndex];
+			int32 DNAMeshIndex = DNAToSkelMeshMap->ImportVtxToDNAMeshIndex[LODIndex][Section.GetVertexBufferIndex()];
+			/*
+			if (UpdateOption == ELodUpdateOption::LOD0Only && !DNAToSkelMeshMap->LOD0MeshIndices.Contains(DNAMeshIndex)) continue;
+			*/
+
+			//const uint16 DNAMaxInfluencePerVertex = DNAReader->getMaximumInfluencePerVertex(DNAMeshIndex);
+
+			const int32 NumEngineVertices = Section.GetNumVertices();			
+			for (int32 VertexIndex = 0; VertexIndex < NumEngineVertices; VertexIndex++)
+			{
+				const int32 VertexBufferIndex = VertexIndex + Section.GetVertexBufferIndex();
+				int32 DNAVertexIndex = DNAToSkelMeshMap->ImportVtxToDNAVtxIndex[LODIndex][VertexBufferIndex];
+
+				if (DNAVertexIndex < 0) continue; // Skip vertex not in DNA.
+
+				TArrayView<const float> DNASkinWeights = DNAReader->GetSkinWeightsValues(DNAMeshIndex, DNAVertexIndex);
+				TArrayView<const uint16> DNASkinJoints = DNAReader->GetSkinWeightsJointIndices(DNAMeshIndex, DNAVertexIndex);
+				uint16 SkinJointNum = DNASkinJoints.Num();
+
+				uint16 WeightsSum = 0;  // store all influences to vertex to ensure they add up to 255 (fix rounding errors)
+				uint16 MaxInfluenceIndex = 0;
+				uint16 MaxInfluenceWeight = 0;
+
+				FSoftSkinVertex& Vertex = Section.SoftVertices[VertexIndex];
+
+				// First we have to reset all influences that are not covered by DNA data.
+				for (uint16 i = SkinJointNum; i < MAX_TOTAL_INFLUENCES; i++)
+				{
+					Vertex.InfluenceBones[i] = INDEX_NONE;
+					Vertex.InfluenceWeights[i] = 0;
+				}
+
+				for (uint16 InfluenceIndex = 0; InfluenceIndex < FMath::Min(SkinJointNum, static_cast<uint16>(MAX_TOTAL_INFLUENCES)); ++InfluenceIndex)
+				{
+					uint16 EngineWeight = 0;
+					// Find Engine bone for corresponding DNAJoint for the same influence.
+					int32 UpdatedBoneId = DNAToSkelMeshMap->GetUEBoneIndex(DNASkinJoints[InfluenceIndex]);
+					// BoneMap holds subset of bones belonging to current section.
+					int32 BoneMapIndex = Section.BoneMap.Find(UpdatedBoneId);
+					
+					// Update which bone in the subset influences this vertex.
+					Vertex.InfluenceBones[InfluenceIndex] = BoneMapIndex;
+					if (BoneMapIndex != INDEX_NONE)
+					{
+						// Update influence weight.
+						float PreRoundValue = 255.0f * DNASkinWeights[InfluenceIndex];
+						EngineWeight = static_cast<uint16>(FMath::Min(255.0f, round(PreRoundValue)));  // convert RL float weight to 0-255 range
+					}
+					Vertex.InfluenceWeights[InfluenceIndex] = static_cast<uint8>(EngineWeight);
+					WeightsSum += EngineWeight;
+
+					if (EngineWeight > MaxInfluenceWeight)
+					{
+						MaxInfluenceIndex = InfluenceIndex;
+						MaxInfluenceWeight = EngineWeight;
+					}
+				}
+				// Add missing fraction to fill up to 255.
+				int32 ValueToAdd = 255 - WeightsSum;
+				uint8 OldValue = Vertex.InfluenceWeights[MaxInfluenceIndex];
+				int32 NewValue = OldValue + ValueToAdd;
+				Vertex.InfluenceWeights[MaxInfluenceIndex] = static_cast<uint8>(NewValue);
+			}
+		}
+	}
+}
+
+/** Rebuilds render data from LODModel and inits resources. */
+void USkelMeshDNAUtils::RebuildRenderData(USkeletalMesh* SkelMesh)
+{
+	FPlatformTime::InitTiming();
+
+	double StartTime = FPlatformTime::Seconds();
+	{
+		SkelMesh->FlushRenderState();
+	}
+	double TimeToFlush = FPlatformTime::Seconds();
+	{
+		FSkeletalMeshRenderData* RenderData = SkelMesh->GetResourceForRendering();
+		int32 LODIndex = 0;
+
+		for (FSkeletalMeshLODRenderData& LODRenderData : RenderData->LODRenderData)
+		{
+			FSkeletalMeshLODModel& LODModelRef = SkelMesh->GetImportedModel()->LODModels[LODIndex];
+			for (int32 i = 0; i < LODModelRef.Sections.Num(); i++)
+			{
+				FSkelMeshSection& ModelSection = LODModelRef.Sections[i];
+				ModelSection.CalcMaxBoneInfluences();
+				ModelSection.CalcUse16BitBoneIndex();
+			}
+
+			const FSkeletalMeshLODModel* LODModelPtr = &LODModelRef;
+			LODRenderData.BuildFromLODModel(LODModelPtr, 0);
+			LODIndex++;
+		}
+	}
+	double TimeToRebuildModel = FPlatformTime::Seconds();
+	{
+		if (FApp::CanEverRender())
+		{
+			// Reinitialize the static mesh's resources.
+			SkelMesh->InitResources();
+		}
+	}
+	double TimeToInitResources = FPlatformTime::Seconds();
+	{
+		// Re-register scope
+		TArray<UActorComponent*> ComponentsToReregister;
+		for (TObjectIterator<USkeletalMeshComponent> It; It; ++It)
+		{
+			USkeletalMeshComponent* MeshComponent = *It;
+			if (MeshComponent && !MeshComponent->IsTemplate() && MeshComponent->SkeletalMesh == SkelMesh)
+			{
+				ComponentsToReregister.Add(*It);
+			}
+		}
+		FMultiComponentReregisterContext ReregisterContext(ComponentsToReregister);
+	}
+}
+
+void USkelMeshDNAUtils::RebuildRenderData_VertexPosition(USkeletalMesh* SkelMesh)
+{
+	if (!FApp::CanEverRender())
+	{
+		return;
+	}
+
+	{
+		FSkeletalMeshModel* MeshModel = SkelMesh->GetImportedModel();
+		FSkeletalMeshRenderData* RenderData = SkelMesh->GetResourceForRendering();
+
+		for (int32 LODIdx = 0; LODIdx < RenderData->LODRenderData.Num(); ++LODIdx)
+		{
+			FSkeletalMeshLODModel& LODModel = MeshModel->LODModels[LODIdx];
+			FSkeletalMeshLODRenderData& LODRenderData = RenderData->LODRenderData[LODIdx];
+
+			ENQUEUE_RENDER_COMMAND(FSkelMeshDNAUpdatePositions)
+				([&LODModel, &LODRenderData](FRHICommandListImmediate& RHICmdList)
+			{
+				TArray<FSoftSkinVertex> Vertices;
+				LODModel.GetVertices(Vertices);
+				check(Vertices.Num() == LODRenderData.StaticVertexBuffers.PositionVertexBuffer.GetNumVertices());
+				LODRenderData.StaticVertexBuffers.PositionVertexBuffer.Init(Vertices.Num());
+				for (int32 i = 0; i < Vertices.Num(); i++)
+				{
+					LODRenderData.StaticVertexBuffers.PositionVertexBuffer.VertexPosition(i) = Vertices[i].Position;
+				}
+
+				auto& VertexBuffer = LODRenderData.StaticVertexBuffers.PositionVertexBuffer;
+				void* VertexBufferData = RHILockVertexBuffer(VertexBuffer.VertexBufferRHI, 0, VertexBuffer.GetNumVertices() * VertexBuffer.GetStride(), RLM_WriteOnly);
+				FMemory::Memcpy(VertexBufferData, VertexBuffer.GetVertexData(), VertexBuffer.GetNumVertices() * VertexBuffer.GetStride());
+				RHIUnlockVertexBuffer(VertexBuffer.VertexBufferRHI);
+			});
+		}
+	}
+}
+
+void USkelMeshDNAUtils::UpdateJointBehavior(USkeletalMeshComponent* SkelMeshComponent)
+{
+    //DNAAsset->SetBehaviorReader is called before calling this method
+    //it is not here to avoid having DNAAsset in the API, as in future we might want
+	//to generalize SkelMeshUpdate to be dna-independent
+
+	//the rig behavior has changed, we need to force re-initializing of RigLogic
+	//this will set RigLogic RigUnit to initial state
+	SkelMeshComponent->InitAnim(true);
+}
+
+#endif // WITH_EDITORONLY_DATA
diff --git a/Engine/Plugins/Animation/RigLogic/Source/RigLogicModule/Public/DNAAsset.h b/Engine/Plugins/Animation/RigLogic/Source/RigLogicModule/Public/DNAAsset.h
new file mode 100644
--- /dev/null
+++ b/Engine/Plugins/Animation/RigLogic/Source/RigLogicModule/Public/DNAAsset.h
@@ -0,0 +1,92 @@
+// Copyright 1998-2019 Epic Games, Inc. All Rights Reserved.
+
+#pragma once
+
+#include "CoreMinimal.h"
+
+#include "Engine/AssetUserData.h"
+#include "Interfaces/Interface_AssetUserData.h"
+
+#include "RigLogic.h"
+
+#include "DNAAsset.generated.h"
+
+DECLARE_LOG_CATEGORY_EXTERN(LogDNAAsset, Log, All);
+
+class IDNAReader;
+class IBehaviorReader;
+class IGeometryReader;
+class FRigLogicMemoryStream;
+class UAssetUserData;
+enum class EDNADataLayer: uint8;
+
+
+ /** An asset holding the data needed to generate/update/animate a RigLogic character
+  * It is imported from character's DNA file as a bit stream, and separated out it into runtime (behavior) and design-time chunks;
+  * Currently, the design-time part still loads the geometry, as it is needed for the skeletal mesh update; once SkeletalMeshDNAReader is
+  * fully implemented, it will be able to read the geometry directly from the SkeletalMesh and won't load it into this asset 
+  **/
+UCLASS(NotBlueprintable, hidecategories = (Object))
+class RIGLOGICMODULE_API UDNAAsset : public UAssetUserData
+{
+	GENERATED_BODY()
+
+public:
+
+#if WITH_EDITORONLY_DATA
+	UPROPERTY(VisibleAnywhere, Instanced, Category = ImportSettings)
+	class UAssetImportData* AssetImportData;
+#endif
+
+	/* Behavior part of the .dna file, used run-time */
+	UPROPERTY()
+	TArray<uint8> BehaviorData; 
+
+	TSharedPtr<IBehaviorReader> GetBehaviorReader()
+	{
+		return BehaviorStreamReader;
+	}
+
+
+#if WITH_EDITORONLY_DATA
+	/* The design-time part of the .dna file, used in-editor */
+	UPROPERTY()
+	TArray<uint8> DesignTimeData; 
+
+	TSharedPtr<IGeometryReader> GetGeometryReader()
+	{
+		return GeometryStreamReader;
+	}
+#endif
+
+	UPROPERTY()
+	FString DNAFileName; 
+
+	bool Init(const FString Filename);
+    void PostLoad() override;
+
+	/** Used when importing behavior into archetype SkelMesh in the editor, 
+	  * and when updating SkeletalMesh runtime with GeneSplicer; it will split
+	  * GeneSplicerDNAReader will overwrite the behavior-only reader
+	  * allowing the RigLogic RigUnit to read behavior part directly from it
+	**/
+	void SetBehaviorReader(const TSharedPtr<IDNAReader> SourceDNAReader);
+	void SetGeometryReader(const TSharedPtr<IDNAReader> SourceDNAReader);
+
+private:
+	/** Part of the .dna file needed for run-time execution of RigLogic;
+	  * Serialized using BehaviorData property
+	 **/
+    TSharedPtr<IBehaviorReader> BehaviorStreamReader = nullptr;
+
+#if WITH_EDITORONLY_DATA
+	/** Part of the .dna file used design-time for updating SkeletalMesh geometry
+	  * Serialized using GeometryData property
+	 **/
+	TSharedPtr<IGeometryReader> GeometryStreamReader = nullptr;
+#endif
+
+	static TSharedPtr<IDNAReader> CopyDNALayer(const IDNAReader* Source, EDNADataLayer DNADataLayer, uint32 PredictedSize );
+	static void WriteToBuffer(IDNAReader* SourcePartialDNAReader, TArray<uint8>* WriteBuffer, uint32 PredictedSize);
+
+};
diff --git a/Engine/Plugins/Animation/RigLogic/Source/RigLogicModule/Public/DNAReaderAdapterImpl.h b/Engine/Plugins/Animation/RigLogic/Source/RigLogicModule/Public/DNAReaderAdapterImpl.h
new file mode 100644
--- /dev/null
+++ b/Engine/Plugins/Animation/RigLogic/Source/RigLogicModule/Public/DNAReaderAdapterImpl.h
@@ -0,0 +1,622 @@
+// Copyright 1998-2019 Epic Games, Inc. All Rights Reserved.
+#pragma once
+
+template <class TWrappedReader>
+FDNAReader<TWrappedReader>::FDNAReader(TWrappedReader* Source) :
+	ReaderPtr{Source}
+{
+}
+
+template <class TWrappedReader>
+dna::Reader* FDNAReader<TWrappedReader>::Unwrap() const
+{
+	return ReaderPtr.Get();
+}
+
+template <class TWrappedReader>
+void FDNAReader<TWrappedReader>::FWrappedReaderDeleter::operator()(TWrappedReader* Pointer)
+{
+	if (Pointer != nullptr)
+	{
+		TWrappedReader::destroy(Pointer);
+	}
+}
+
+template <class TWrappedReader>
+FString FDNAReader<TWrappedReader>::GetName() const
+{
+	return FString(ANSI_TO_TCHAR(ReaderPtr->getName().data()));
+}
+
+template <class TWrappedReader>
+EArchetype FDNAReader<TWrappedReader>::GetArchetype() const
+{
+	return static_cast<EArchetype>(ReaderPtr->getArchetype());
+}
+
+template <class TWrappedReader>
+EGender FDNAReader<TWrappedReader>::GetGender() const
+{
+	return static_cast<EGender>(ReaderPtr->getGender());
+}
+
+template <class TWrappedReader>
+uint16 FDNAReader<TWrappedReader>::GetAge() const
+{
+	return ReaderPtr->getAge();
+}
+
+template <class TWrappedReader>
+uint32 FDNAReader<TWrappedReader>::GetMetaDataCount() const
+{
+	return ReaderPtr->getMetaDataCount();
+}
+
+template <class TWrappedReader>
+FString FDNAReader<TWrappedReader>::GetMetaDataKey(uint32 Index) const
+{
+	return FString(ANSI_TO_TCHAR(ReaderPtr->getMetaDataKey(Index).data()));
+}
+
+template <class TWrappedReader>
+FString FDNAReader<TWrappedReader>::GetMetaDataValue(const FString& Key) const
+{
+	return FString(ANSI_TO_TCHAR(ReaderPtr->getMetaDataValue(TCHAR_TO_ANSI(*Key)).data()));
+}
+
+template <class TWrappedReader>
+ETranslationUnit FDNAReader<TWrappedReader>::GetTranslationUnit() const
+{
+	return static_cast<ETranslationUnit>(ReaderPtr->getTranslationUnit());
+}
+
+template <class TWrappedReader>
+ERotationUnit FDNAReader<TWrappedReader>::GetRotationUnit() const
+{
+	return static_cast<ERotationUnit>(ReaderPtr->getRotationUnit());
+}
+
+template <class TWrappedReader>
+FCoordinateSystem FDNAReader<TWrappedReader>::GetCoordinateSystem() const
+{
+	const auto System = ReaderPtr->getCoordinateSystem();
+	return FCoordinateSystem{
+		static_cast<EDirection>(System.xAxis),
+		static_cast<EDirection>(System.yAxis),
+		static_cast<EDirection>(System.zAxis)
+	};
+}
+
+template <class TWrappedReader>
+uint16 FDNAReader<TWrappedReader>::GetLODCount() const
+{
+	return ReaderPtr->getLODCount();
+}
+
+template <class TWrappedReader>
+uint16 FDNAReader<TWrappedReader>::GetDBMaxLOD() const
+{
+	return ReaderPtr->getDBMaxLOD();
+}
+
+template <class TWrappedReader>
+FString FDNAReader<TWrappedReader>::GetDBComplexity() const
+{
+	return FString(ANSI_TO_TCHAR(ReaderPtr->getDBComplexity().data()));
+}
+
+template <class TWrappedReader>
+FString FDNAReader<TWrappedReader>::GetDBName() const
+{
+	return FString(ANSI_TO_TCHAR(ReaderPtr->getDBName().data()));
+}
+
+
+template <class TWrappedReader>
+uint16 FDNAReader<TWrappedReader>::GetGUIControlCount() const
+{
+	return ReaderPtr->getGUIControlCount();
+}
+
+template <class TWrappedReader>
+FString FDNAReader<TWrappedReader>::GetGUIControlName(uint16 Index) const
+{
+	return FString(ANSI_TO_TCHAR(ReaderPtr->getGUIControlName(Index).data()));
+}
+
+template <class TWrappedReader>
+uint16 FDNAReader<TWrappedReader>::GetRawControlCount() const
+{
+	return ReaderPtr->getRawControlCount();
+}
+
+template <class TWrappedReader>
+FString FDNAReader<TWrappedReader>::GetRawControlName(uint16 Index) const
+{
+	return FString(ANSI_TO_TCHAR(ReaderPtr->getRawControlName(Index).data()));
+}
+
+template <class TWrappedReader>
+uint16 FDNAReader<TWrappedReader>::GetJointCount() const
+{
+	return ReaderPtr->getJointCount();
+}
+
+template <class TWrappedReader>
+FString FDNAReader<TWrappedReader>::GetJointName(uint16 Index) const
+{
+	return FString(ANSI_TO_TCHAR(ReaderPtr->getJointName(Index).data()));
+}
+
+template <class TWrappedReader>
+TArrayView<const uint16> FDNAReader<TWrappedReader>::GetJointIndicesForLOD(uint16 LOD) const
+{
+	const auto Indices = ReaderPtr->getJointIndicesForLOD(LOD);
+	return TArrayView<const uint16>(Indices.data(), static_cast<int32>(Indices.size()));
+}
+
+template <class TWrappedReader>
+uint16 FDNAReader<TWrappedReader>::GetJointParentIndex(uint16 Index) const
+{
+	return ReaderPtr->getJointParentIndex(Index);
+}
+
+template <class TWrappedReader>
+uint16 FDNAReader<TWrappedReader>::GetBlendShapeChannelCount() const
+{
+	return ReaderPtr->getBlendShapeChannelCount();
+}
+
+template <class TWrappedReader>
+FString FDNAReader<TWrappedReader>::GetBlendShapeChannelName(uint16 Index) const
+{
+	return FString(ANSI_TO_TCHAR(ReaderPtr->getBlendShapeChannelName(Index).data()));
+}
+
+template <class TWrappedReader>
+TArrayView<const uint16> FDNAReader<TWrappedReader>::GetBlendShapeChannelIndicesForLOD(uint16 LOD) const
+{
+	const auto Indices = ReaderPtr->getBlendShapeChannelIndicesForLOD(LOD);
+	return TArrayView<const uint16>(Indices.data(), static_cast<int32>(Indices.size()));
+}
+
+template <class TWrappedReader>
+uint16 FDNAReader<TWrappedReader>::GetAnimatedMapCount() const
+{
+	return ReaderPtr->getAnimatedMapCount();
+}
+
+template <class TWrappedReader>
+FString FDNAReader<TWrappedReader>::GetAnimatedMapName(uint16 Index) const
+{
+	return FString(ANSI_TO_TCHAR(ReaderPtr->getAnimatedMapName(Index).data()));
+}
+
+template <class TWrappedReader>
+TArrayView<const uint16> FDNAReader<TWrappedReader>::GetAnimatedMapIndicesForLOD(uint16 LOD) const
+{
+	const auto Indices = ReaderPtr->getAnimatedMapIndicesForLOD(LOD);
+	return TArrayView<const uint16>(Indices.data(), static_cast<int32>(Indices.size()));
+}
+
+template <class TWrappedReader>
+uint16 FDNAReader<TWrappedReader>::GetMeshCount() const
+{
+	return ReaderPtr->getMeshCount();
+}
+
+template <class TWrappedReader>
+FString FDNAReader<TWrappedReader>::GetMeshName(uint16 Index) const
+{
+	return FString(ANSI_TO_TCHAR(ReaderPtr->getMeshName(Index).data()));
+}
+
+template <class TWrappedReader>
+TArrayView<const uint16> FDNAReader<TWrappedReader>::GetMeshIndicesForLOD(uint16 LOD) const
+{
+	const auto Indices = ReaderPtr->getMeshIndicesForLOD(LOD);
+	return TArrayView<const uint16>(Indices.data(), static_cast<int32>(Indices.size()));
+}
+
+template <class TWrappedReader>
+uint16 FDNAReader<TWrappedReader>::GetMeshBlendShapeChannelMappingCount() const
+{
+	return ReaderPtr->getMeshBlendShapeChannelMappingCount();
+}
+
+template <class TWrappedReader>
+FMeshBlendShapeChannelMapping FDNAReader<TWrappedReader>::GetMeshBlendShapeChannelMapping(uint16 Index) const
+{
+	const auto Mapping = ReaderPtr->getMeshBlendShapeChannelMapping(Index);
+	return FMeshBlendShapeChannelMapping{ Mapping.meshIndex, Mapping.blendShapeChannelIndex };
+}
+
+template <class TWrappedReader>
+TArrayView<const uint16> FDNAReader<TWrappedReader>::GetMeshBlendShapeChannelMappingIndicesForLOD(uint16 LOD) const
+{
+	const auto Indices = ReaderPtr->getMeshBlendShapeChannelMappingIndicesForLOD(LOD);
+	return TArrayView<const uint16>(Indices.data(), static_cast<int32>(Indices.size()));
+}
+
+template <class TWrappedReader>
+FVector FDNAReader<TWrappedReader>::GetNeutralJointTranslation(uint16 Index) const
+{
+	const auto Translation = ReaderPtr->getNeutralJointTranslation(Index);
+	// X = X, Y = -Y, Z = Z
+	return FVector(Translation.x, -Translation.y, Translation.z);
+}
+
+template <class TWrappedReader>
+FVector FDNAReader<TWrappedReader>::GetNeutralJointRotation(uint16 Index) const
+{
+	const auto Translation = ReaderPtr->getNeutralJointRotation(Index);
+	// X = -Y, Y = -Z, Z = X
+	return FVector(-Translation.y, -Translation.z, Translation.x);
+}
+
+template <class TWrappedReader>
+TArrayView<const uint16> FDNAReader<TWrappedReader>::GetGUIToRawInputIndices() const
+{
+	const auto Indices = ReaderPtr->getGUIToRawInputIndices();
+	return TArrayView<const uint16>(Indices.data(), static_cast<int32>(Indices.size()));
+}
+
+template <class TWrappedReader>
+TArrayView<const uint16> FDNAReader<TWrappedReader>::GetGUIToRawOutputIndices() const
+{
+	const auto Indices = ReaderPtr->getGUIToRawOutputIndices();
+	return TArrayView<const uint16>(Indices.data(), static_cast<int32>(Indices.size()));
+}
+
+template <class TWrappedReader>
+TArrayView<const float> FDNAReader<TWrappedReader>::GetGUIToRawFromValues() const
+{
+	const auto Values = ReaderPtr->getGUIToRawFromValues();
+	return TArrayView<const float>(Values.data(), static_cast<int32>(Values.size()));
+}
+
+template <class TWrappedReader>
+TArrayView<const float> FDNAReader<TWrappedReader>::GetGUIToRawToValues() const
+{
+	const auto Values = ReaderPtr->getGUIToRawToValues();
+	return TArrayView<const float>(Values.data(), static_cast<int32>(Values.size()));
+}
+
+template <class TWrappedReader>
+TArrayView<const float> FDNAReader<TWrappedReader>::GetGUIToRawSlopeValues() const
+{
+	const auto Values = ReaderPtr->getGUIToRawSlopeValues();
+	return TArrayView<const float>(Values.data(), static_cast<int32>(Values.size()));
+}
+
+template <class TWrappedReader>
+TArrayView<const float> FDNAReader<TWrappedReader>::GetGUIToRawCutValues() const
+{
+	const auto Values = ReaderPtr->getGUIToRawCutValues();
+	return TArrayView<const float>(Values.data(), static_cast<int32>(Values.size()));
+}
+
+template <class TWrappedReader>
+uint16 FDNAReader<TWrappedReader>::GetPSDCount() const
+{
+	return ReaderPtr->getPSDCount();
+}
+
+template <class TWrappedReader>
+TArrayView<const uint16> FDNAReader<TWrappedReader>::GetPSDRowIndices() const
+{
+	const auto Indices = ReaderPtr->getPSDRowIndices();
+	return TArrayView<const uint16>(Indices.data(), static_cast<int32>(Indices.size()));
+}
+
+template <class TWrappedReader>
+TArrayView<const uint16> FDNAReader<TWrappedReader>::GetPSDColumnIndices() const
+{
+	const auto Indices = ReaderPtr->getPSDColumnIndices();
+	return TArrayView<const uint16>(Indices.data(), static_cast<int32>(Indices.size()));
+}
+
+template <class TWrappedReader>
+TArrayView<const float> FDNAReader<TWrappedReader>::GetPSDValues() const
+{
+	const auto Values = ReaderPtr->getPSDValues();
+	return TArrayView<const float>(Values.data(), static_cast<int32>(Values.size()));
+}
+
+template <class TWrappedReader>
+uint16 FDNAReader<TWrappedReader>::GetJointRowCount() const
+{
+	return ReaderPtr->getJointRowCount();
+}
+
+template <class TWrappedReader>
+uint16 FDNAReader<TWrappedReader>::GetJointColumnCount() const
+{
+	return ReaderPtr->getJointColumnCount();
+}
+
+template <class TWrappedReader>
+TArrayView<const uint16> FDNAReader<TWrappedReader>::GetJointVariableAttributeIndices(uint16 LOD) const
+{
+	const auto Indices = ReaderPtr->getJointVariableAttributeIndices(LOD);
+	return TArrayView<const uint16>(Indices.data(), static_cast<int32>(Indices.size()));
+}
+
+template <class TWrappedReader>
+uint16 FDNAReader<TWrappedReader>::GetJointGroupCount() const
+{
+	return ReaderPtr->getJointGroupCount();
+}
+
+template <class TWrappedReader>
+TArrayView<const uint16> FDNAReader<TWrappedReader>::GetJointGroupLODs(uint16 JointGroupIndex) const
+{
+	const auto LODs = ReaderPtr->getJointGroupLODs(JointGroupIndex);
+	return TArrayView<const uint16>(LODs.data(), static_cast<int32>(LODs.size()));
+}
+
+template <class TWrappedReader>
+TArrayView<const uint16> FDNAReader<TWrappedReader>::GetJointGroupInputIndices(uint16 JointGroupIndex) const
+{
+	const auto Indices = ReaderPtr->getJointGroupInputIndices(JointGroupIndex);
+	return TArrayView<const uint16>(Indices.data(), static_cast<int32>(Indices.size()));
+}
+
+template <class TWrappedReader>
+TArrayView<const uint16> FDNAReader<TWrappedReader>::GetJointGroupOutputIndices(uint16 JointGroupIndex) const
+{
+	const auto Indices = ReaderPtr->getJointGroupOutputIndices(JointGroupIndex);
+	return TArrayView<const uint16>(Indices.data(), static_cast<int32>(Indices.size()));
+}
+
+template <class TWrappedReader>
+TArrayView<const float> FDNAReader<TWrappedReader>::GetJointGroupValues(uint16 JointGroupIndex) const
+{
+	const auto Values = ReaderPtr->getJointGroupValues(JointGroupIndex);
+	return TArrayView<const float>(Values.data(), static_cast<int32>(Values.size()));
+}
+
+template <class TWrappedReader>
+TArrayView<const uint16> FDNAReader<TWrappedReader>::GetJointGroupJointIndices(uint16 JointGroupIndex) const
+{
+	const auto Indices = ReaderPtr->getJointGroupJointIndices(JointGroupIndex);
+	return TArrayView<const uint16>(Indices.data(), static_cast<int32>(Indices.size()));
+}
+
+template <class TWrappedReader>
+TArrayView<const uint16> FDNAReader<TWrappedReader>::GetBlendShapeChannelLODs() const
+{
+	const auto LODs = ReaderPtr->getBlendShapeChannelLODs();
+	return TArrayView<const uint16>(LODs.data(), static_cast<int32>(LODs.size()));
+}
+
+template <class TWrappedReader>
+TArrayView<const uint16> FDNAReader<TWrappedReader>::GetBlendShapeChannelOutputIndices() const
+{
+	const auto Indices = ReaderPtr->getBlendShapeChannelOutputIndices();
+	return TArrayView<const uint16>(Indices.data(), static_cast<int32>(Indices.size()));
+}
+
+template <class TWrappedReader>
+TArrayView<const uint16> FDNAReader<TWrappedReader>::GetBlendShapeChannelInputIndices() const
+{
+	const auto Indices = ReaderPtr->getBlendShapeChannelInputIndices();
+	return TArrayView<const uint16>(Indices.data(), static_cast<int32>(Indices.size()));
+}
+
+template <class TWrappedReader>
+TArrayView<const uint16> FDNAReader<TWrappedReader>::GetAnimatedMapLODs() const
+{
+	const auto LODs = ReaderPtr->getAnimatedMapLODs();
+	return TArrayView<const uint16>(LODs.data(), static_cast<int32>(LODs.size()));
+}
+
+template <class TWrappedReader>
+TArrayView<const uint16> FDNAReader<TWrappedReader>::GetAnimatedMapInputIndices() const
+{
+	const auto Indices = ReaderPtr->getAnimatedMapInputIndices();
+	return TArrayView<const uint16>(Indices.data(), static_cast<int32>(Indices.size()));
+}
+
+template <class TWrappedReader>
+TArrayView<const uint16> FDNAReader<TWrappedReader>::GetAnimatedMapOutputIndices() const
+{
+	const auto Indices = ReaderPtr->getAnimatedMapOutputIndices();
+	return TArrayView<const uint16>(Indices.data(), static_cast<int32>(Indices.size()));
+}
+
+template <class TWrappedReader>
+TArrayView<const float> FDNAReader<TWrappedReader>::GetAnimatedMapFromValues() const
+{
+	const auto Values = ReaderPtr->getAnimatedMapFromValues();
+	return TArrayView<const float>(Values.data(), static_cast<int32>(Values.size()));
+}
+
+template <class TWrappedReader>
+TArrayView<const float> FDNAReader<TWrappedReader>::GetAnimatedMapToValues() const
+{
+	const auto Values = ReaderPtr->getAnimatedMapToValues();
+	return TArrayView<const float>(Values.data(), static_cast<int32>(Values.size()));
+}
+
+template <class TWrappedReader>
+TArrayView<const float> FDNAReader<TWrappedReader>::GetAnimatedMapSlopeValues() const
+{
+	const auto Values = ReaderPtr->getAnimatedMapSlopeValues();
+	return TArrayView<const float>(Values.data(), static_cast<int32>(Values.size()));
+}
+
+template <class TWrappedReader>
+TArrayView<const float> FDNAReader<TWrappedReader>::GetAnimatedMapCutValues() const
+{
+	const auto Values = ReaderPtr->getAnimatedMapCutValues();
+	return TArrayView<const float>(Values.data(), static_cast<int32>(Values.size()));
+}
+
+
+template <class TWrappedReader>
+uint32 FDNAReader<TWrappedReader>::GetVertexPositionCount(uint16 MeshIndex) const
+{
+	return ReaderPtr->getVertexPositionCount(MeshIndex);
+}
+
+template <class TWrappedReader>
+FVector FDNAReader<TWrappedReader>::GetVertexPosition(uint16 MeshIndex, uint32 VertexIndex) const
+{
+	const auto Position = ReaderPtr->getVertexPosition(MeshIndex, VertexIndex);
+	// X = X, Y = Z, Z = Y
+	return FVector(Position.x, Position.z, Position.y);
+}
+
+template <class TWrappedReader>
+TArrayView<const float> FDNAReader<TWrappedReader>::GetVertexPositionXs(uint16 MeshIndex) const
+{
+	// X = X
+	const auto Values = ReaderPtr->getVertexPositionXs(MeshIndex);
+	return TArrayView<const float>(Values.data(), static_cast<int32>(Values.size()));
+}
+
+template <class TWrappedReader>
+TArrayView<const float> FDNAReader<TWrappedReader>::GetVertexPositionYs(uint16 MeshIndex) const
+{
+	// Y = Z
+	const auto Values = ReaderPtr->getVertexPositionZs(MeshIndex);
+	return TArrayView<const float>(Values.data(), static_cast<int32>(Values.size()));
+}
+
+template <class TWrappedReader>
+TArrayView<const float> FDNAReader<TWrappedReader>::GetVertexPositionZs(uint16 MeshIndex) const
+{
+	// Z = Y
+	const auto Values = ReaderPtr->getVertexPositionYs(MeshIndex);
+	return TArrayView<const float>(Values.data(), static_cast<int32>(Values.size()));
+}
+
+template <class TWrappedReader>
+uint32 FDNAReader<TWrappedReader>::GetVertexTextureCoordinateCount(uint16 MeshIndex) const
+{
+	return ReaderPtr->getVertexTextureCoordinateCount(MeshIndex);
+}
+
+template <class TWrappedReader>
+FTextureCoordinate FDNAReader<TWrappedReader>::GetVertexTextureCoordinate(uint16 MeshIndex, uint32 TextureCoordinateIndex) const
+{
+	const auto Coordinate = ReaderPtr->getVertexTextureCoordinate(MeshIndex, TextureCoordinateIndex);
+	return FTextureCoordinate{Coordinate.u, Coordinate.v};
+}
+
+template <class TWrappedReader>
+uint32 FDNAReader<TWrappedReader>::GetVertexNormalCount(uint16 MeshIndex) const
+{
+	return ReaderPtr->getVertexNormalCount(MeshIndex);
+}
+
+template <class TWrappedReader>
+FVector FDNAReader<TWrappedReader>::GetVertexNormal(uint16 MeshIndex, uint32 NormalIndex) const
+{
+	const auto Normal = ReaderPtr->getVertexNormal(MeshIndex, NormalIndex);
+	return FVector(Normal.x, Normal.y, Normal.z);
+}
+
+template <class TWrappedReader>
+uint32 FDNAReader<TWrappedReader>::GetFaceCount(uint16 MeshIndex) const
+{
+	return ReaderPtr->getFaceCount(MeshIndex);
+}
+
+template <class TWrappedReader>
+TArrayView<const uint32> FDNAReader<TWrappedReader>::GetFaceVertexLayoutIndices(uint16 MeshIndex, uint32 FaceIndex) const
+{
+	const auto Indices = ReaderPtr->getFaceVertexLayoutIndices(MeshIndex, FaceIndex);
+	return TArrayView<const uint32>(Indices.data(), static_cast<int32>(Indices.size()));
+}
+
+template <class TWrappedReader>
+uint32 FDNAReader<TWrappedReader>::GetVertexLayoutCount(uint16 MeshIndex) const
+{
+	return ReaderPtr->getVertexLayoutCount(MeshIndex);
+}
+
+template <class TWrappedReader>
+FVertexLayout FDNAReader<TWrappedReader>::GetVertexLayout(uint16 MeshIndex, uint32 LayoutIndex) const
+{
+	const auto Layout = ReaderPtr->getVertexLayout(MeshIndex, LayoutIndex);
+	return FVertexLayout{static_cast<int32>(Layout.position), static_cast<int32>(Layout.textureCoordinate), static_cast<int32>(Layout.normal)};
+}
+
+template <class TWrappedReader>
+uint16 FDNAReader<TWrappedReader>::GetMaximumInfluencePerVertex(uint16 MeshIndex) const
+{
+	return ReaderPtr->getMaximumInfluencePerVertex(MeshIndex);
+}
+
+template <class TWrappedReader>
+TArrayView<const float> FDNAReader<TWrappedReader>::GetSkinWeightsValues(uint16 MeshIndex, uint32 VertexIndex) const
+{
+	const auto Values = ReaderPtr->getSkinWeightsValues(MeshIndex, VertexIndex);
+	return TArrayView<const float>(Values.data(), static_cast<int32>(Values.size()));
+}
+
+template <class TWrappedReader>
+TArrayView<const uint16> FDNAReader<TWrappedReader>::GetSkinWeightsJointIndices(uint16 MeshIndex, uint32 VertexIndex) const
+{
+	const auto Indices = ReaderPtr->getSkinWeightsJointIndices(MeshIndex, VertexIndex);
+	return TArrayView<const uint16>(Indices.data(), static_cast<int32>(Indices.size()));
+}
+
+template <class TWrappedReader>
+uint16 FDNAReader<TWrappedReader>::GetBlendShapeTargetCount(uint16 MeshIndex) const
+{
+	return ReaderPtr->getBlendShapeTargetCount(MeshIndex);
+}
+
+template <class TWrappedReader>
+uint16 FDNAReader<TWrappedReader>::GetBlendShapeChannelIndex(uint16 MeshIndex, uint16 BlendShapeTargetIndex) const
+{
+	return ReaderPtr->getBlendShapeChannelIndex(MeshIndex, BlendShapeTargetIndex);
+}
+
+template <class TWrappedReader>
+uint32 FDNAReader<TWrappedReader>::GetBlendShapeTargetDeltaCount(uint16 MeshIndex, uint16 BlendShapeTargetIndex) const
+{
+	return ReaderPtr->getBlendShapeTargetDeltaCount(MeshIndex, BlendShapeTargetIndex);
+}
+
+template <class TWrappedReader>
+FVector FDNAReader<TWrappedReader>::GetBlendShapeTargetDelta(uint16 MeshIndex, uint16 BlendShapeTargetIndex, uint32 DeltaIndex) const
+{
+	const auto Delta = ReaderPtr->getBlendShapeTargetDelta(MeshIndex, BlendShapeTargetIndex, DeltaIndex);
+	// X = X, Y = Z, Z = Y
+	return FVector(Delta.x, Delta.z, Delta.y);
+}
+
+template <class TWrappedReader>
+TArrayView<const float> FDNAReader<TWrappedReader>::GetBlendShapeTargetDeltaXs(uint16 MeshIndex, uint16 BlendShapeTargetIndex) const
+{
+	// X = X
+	const auto Values = ReaderPtr->getBlendShapeTargetDeltaXs(MeshIndex, BlendShapeTargetIndex);
+	return TArrayView<const float>(Values.data(), static_cast<int32>(Values.size()));
+}
+
+template <class TWrappedReader>
+TArrayView<const float> FDNAReader<TWrappedReader>::GetBlendShapeTargetDeltaYs(uint16 MeshIndex, uint16 BlendShapeTargetIndex) const
+{
+	// Y = Z
+	const auto Values = ReaderPtr->getBlendShapeTargetDeltaZs(MeshIndex, BlendShapeTargetIndex);
+	return TArrayView<const float>(Values.data(), static_cast<int32>(Values.size()));
+}
+
+template <class TWrappedReader>
+TArrayView<const float> FDNAReader<TWrappedReader>::GetBlendShapeTargetDeltaZs(uint16 MeshIndex, uint16 BlendShapeTargetIndex) const
+{
+	// Z = Y
+	const auto Values = ReaderPtr->getBlendShapeTargetDeltaYs(MeshIndex, BlendShapeTargetIndex);
+	return TArrayView<const float>(Values.data(), static_cast<int32>(Values.size()));
+}
+
+template <class TWrappedReader>
+TArrayView<const uint32> FDNAReader<TWrappedReader>::GetBlendShapeTargetVertexIndices(uint16 MeshIndex, uint16 BlendShapeTargetIndex) const
+{
+	const auto Indices = ReaderPtr->getBlendShapeTargetVertexIndices(MeshIndex, BlendShapeTargetIndex);
+	return TArrayView<const uint32>(Indices.data(), static_cast<int32>(Indices.size()));
+}
diff --git a/Engine/Plugins/Animation/RigLogic/Source/RigLogicModule/Public/DNAUtils.h b/Engine/Plugins/Animation/RigLogic/Source/RigLogicModule/Public/DNAUtils.h
new file mode 100644
--- /dev/null
+++ b/Engine/Plugins/Animation/RigLogic/Source/RigLogicModule/Public/DNAUtils.h
@@ -0,0 +1,15 @@
+// Copyright 1998-2019 Epic Games, Inc. All Rights Reserved.
+
+#pragma once
+
+#include "DNAReader.h"
+
+#include "CoreMinimal.h"
+#include "UObject/NoExportTypes.h"
+#include "UObject/ObjectMacros.h"
+
+RIGLOGICMODULE_API TSharedPtr<IDNAReader> ReadDNAFromFile(const FString& Path, EDNADataLayer Layer = EDNADataLayer::All, uint16_t MaxLOD = 0);
+RIGLOGICMODULE_API TSharedPtr<IDNAReader> ReadDNAFromFile(const FString& Path, EDNADataLayer Layer, TArrayView<uint16_t> LODs);
+RIGLOGICMODULE_API TSharedPtr<IDNAReader> ReadDNALayerFromBuffer(TArray<uint8>* DNABuffer, EDNADataLayer Layer = EDNADataLayer::All, uint16_t MaxLOD = 0);
+RIGLOGICMODULE_API TSharedPtr<IDNAReader> ReadDNALayerFromBuffer(TArray<uint8>* DNABuffer, EDNADataLayer Layer, TArrayView<uint16_t> LODs);
+RIGLOGICMODULE_API void WriteDNAToFile(const IDNAReader* Reader, EDNADataLayer Layer, const FString& Path);
\ No newline at end of file
diff --git a/Engine/Plugins/Animation/RigLogic/Source/RigLogicModule/Public/RigInstance.h b/Engine/Plugins/Animation/RigLogic/Source/RigLogicModule/Public/RigInstance.h
new file mode 100644
--- /dev/null
+++ b/Engine/Plugins/Animation/RigLogic/Source/RigLogicModule/Public/RigInstance.h
@@ -0,0 +1,58 @@
+// Copyright 1998-2019 Epic Games, Inc. All Rights Reserved.
+
+#pragma once
+
+#include "Containers/ArrayView.h"
+#include "CoreMinimal.h"
+#include "Math/Vector.h"
+#include "UObject/NoExportTypes.h"
+#include "UObject/ObjectMacros.h"
+
+#include "TransformArrayView.h"
+
+class FRigLogic;
+
+namespace rl4
+{
+
+class RigInstance;
+
+}  // namespace rl4
+
+class RIGLOGICMODULE_API FRigInstance
+{
+public:
+	FRigInstance(FRigLogic* RigLogic);
+	~FRigInstance();
+
+	FRigInstance(const FRigInstance&) = delete;
+	FRigInstance& operator=(const FRigInstance&) = delete;
+
+	FRigInstance(FRigInstance&&) = default;
+	FRigInstance& operator=(FRigInstance&&) = default;
+
+	uint16 GetGUIControlCount() const;
+	void SetGUIControl(uint16 Index, float Value);
+	void SetGUIControlValues(const float* Values);
+
+	uint16 GetRawControlCount() const;
+	void SetRawControl(uint16 Index, float Value);
+	void SetRawControlValues(const float* Values);
+	TArrayView<const float> GetRawJointOutputs() const;
+	FTransformArrayView GetJointOutputs() const;
+	TArrayView<const float> GetBlendShapeOutputs() const;
+	TArrayView<const float> GetAnimatedMapOutputs() const;
+	uint16 GetLOD() const;
+	void SetLOD(uint16 LOD);
+
+private:
+	friend FRigLogic;
+	rl4::RigInstance* Unwrap() const;
+
+private:
+	struct FRigInstanceDeleter
+	{
+		void operator()(rl4::RigInstance* Pointer);
+	};
+	TUniquePtr<rl4::RigInstance, FRigInstanceDeleter> RigInstance;
+};
diff --git a/Engine/Plugins/Animation/RigLogic/Source/RigLogicModule/Public/RigLogic.h b/Engine/Plugins/Animation/RigLogic/Source/RigLogicModule/Public/RigLogic.h
new file mode 100644
--- /dev/null
+++ b/Engine/Plugins/Animation/RigLogic/Source/RigLogicModule/Public/RigLogic.h
@@ -0,0 +1,66 @@
+// Copyright 1998-2019 Epic Games, Inc. All Rights Reserved.
+
+#pragma once
+
+#include "Containers/ArrayView.h"
+#include "CoreMinimal.h"
+#include "Math/Vector.h"
+#include "UObject/NoExportTypes.h"
+#include "UObject/ObjectMacros.h"
+
+#include "TransformArrayView.h"
+
+class FRigInstance;
+class IBehaviorReader;
+
+namespace rl4
+{
+
+class RigLogic;
+
+}  // namespace rl4
+
+UENUM(BlueprintType)
+enum class ERigLogicCalculationType: uint8
+{
+	Scalar,
+	SSE,
+	AVX
+};
+
+class RIGLOGICMODULE_API FRigLogic
+{
+public:
+	FRigLogic(const IBehaviorReader* Reader, ERigLogicCalculationType CalculationType = ERigLogicCalculationType::SSE);
+	~FRigLogic();
+
+	FRigLogic(const FRigLogic&) = delete;
+	FRigLogic& operator=(const FRigLogic&) = delete;
+
+	FRigLogic(FRigLogic&&) = default;
+	FRigLogic& operator=(FRigLogic&&) = default;
+
+	uint16 GetLODCount() const;
+	FTransformArrayView GetNeutralJointValues() const;
+	uint16 GetJointGroupCount() const;
+	TArrayView<const uint16> GetJointVariableAttributeIndices(uint16 LOD) const;
+	void MapGUIToRawControls(FRigInstance* Instance) const;
+	void CalculateControls(FRigInstance* Instance) const;
+	void CalculateJoints(FRigInstance* Instance) const;
+	void CalculateJoints(FRigInstance* Instance, uint16 JointGroupIndex) const;
+	void CalculateBlendShapes(FRigInstance* Instance) const;
+	void CalculateAnimatedMaps(FRigInstance* Instance) const;
+	void Calculate(FRigInstance* Instance) const;
+
+private:
+	friend FRigInstance;
+	rl4::RigLogic* Unwrap() const;
+
+private:
+	struct FRigLogicDeleter
+	{
+		void operator()(rl4::RigLogic* Pointer);
+	};
+	TUniquePtr<rl4::RigLogic, FRigLogicDeleter> RigLogic;
+};
+
diff --git a/Engine/Plugins/Animation/RigLogic/Source/RigLogicModule/Public/RigLogicModule.h b/Engine/Plugins/Animation/RigLogic/Source/RigLogicModule/Public/RigLogicModule.h
new file mode 100644
--- /dev/null
+++ b/Engine/Plugins/Animation/RigLogic/Source/RigLogicModule/Public/RigLogicModule.h
@@ -0,0 +1,15 @@
+// Copyright 1998-2019 Epic Games, Inc. All Rights Reserved.
+
+#pragma once
+
+#include "CoreMinimal.h"
+#include "Modules/ModuleInterface.h"
+
+DECLARE_LOG_CATEGORY_EXTERN(LogRigLogic, Log, All);
+
+struct RIGLOGICMODULE_API FRigLogicModule: IModuleInterface
+{
+public:
+	virtual void StartupModule() override;
+	virtual void ShutdownModule() override;
+};
diff --git a/Engine/Plugins/Animation/RigLogic/Source/RigLogicModule/Public/RigUnit_RigLogic.h b/Engine/Plugins/Animation/RigLogic/Source/RigLogicModule/Public/RigUnit_RigLogic.h
new file mode 100644
--- /dev/null
+++ b/Engine/Plugins/Animation/RigLogic/Source/RigLogicModule/Public/RigUnit_RigLogic.h
@@ -0,0 +1,183 @@
+// Copyright 1998-2019 Epic Games, Inc. All Rights Reserved.
+
+#pragma once
+
+#include "RigInstance.h"
+#include "RigLogic.h"
+
+#include "ControlRig.h"
+#include "../Public/DNAAsset.h"
+#include "Runtime/Engine/Classes/Materials/Material.h"
+#include "Units/RigUnit.h"
+
+#include "RigUnit_RigLogic.generated.h"
+
+class IBehaviorReader;
+class FTransformArrayView;
+
+DECLARE_LOG_CATEGORY_EXTERN(LogRigLogicUnit, Log, All);
+
+/* A helper struct used inside of the RigUnit_RigLogic to store arrays of arrays of integers. */
+USTRUCT(meta = (DocumentationPolicy = "Strict"))
+struct FRigUnit_RigLogic_IntArray
+{
+	GENERATED_BODY()
+
+	// The values stored within this array.
+	UPROPERTY(transient)
+	TArray<int32> Values;
+};
+
+/**
+ * A helper struct used inside of the RigUnit_RigLogic holding params for updating joints.
+ * Note that these params are array views - they don't own the memory they're pointing to.
+ */
+struct FRigUnit_RigLogic_JointUpdateParams
+{
+	FRigUnit_RigLogic_JointUpdateParams(
+		TArrayView<const uint16> InVariableAttributes,
+		FTransformArrayView InNeutralJointTransforms,
+		FTransformArrayView InDeltaTransforms
+	)
+		: VariableAttributes(InVariableAttributes)
+		, NeutralJointTransforms(InNeutralJointTransforms)
+		, DeltaTransforms(InDeltaTransforms)
+	{ }
+
+public:
+	TArrayView<const uint16> VariableAttributes;
+	FTransformArrayView NeutralJointTransforms;
+	FTransformArrayView DeltaTransforms; //the result of rig logic calculations
+};
+
+/* The work data used by the FRigUnit_RigLogic */
+USTRUCT(meta = (DocumentationPolicy = "Strict"))
+struct FRigUnit_RigLogic_Data
+{
+	GENERATED_BODY()
+
+	FRigUnit_RigLogic_Data();
+ 	~FRigUnit_RigLogic_Data();
+ 	FRigUnit_RigLogic_Data(const FRigUnit_RigLogic_Data& Other);
+ 	FRigUnit_RigLogic_Data& operator = (const FRigUnit_RigLogic_Data& Other);
+
+	/** Cached Skeletal Mesh Component **/
+	UPROPERTY(transient)
+	TWeakObjectPtr<USkeletalMeshComponent> SkelMeshComponent;
+
+	/** LOD for which the model is rendered **/
+	UPROPERTY(transient)
+	uint32 CurrentLOD;
+
+	// Using TSharedPtr so the struct itself remains copy-able, but still cleans
+	// up these members properly when needed
+	FRigLogic* RigLogic;
+	FRigInstance* RigInstance;
+
+	/** Mapping RL indices to UE indices
+	  * Note: we use int32 instead of uint32 to allow storing INDEX_NONE for missing elements
+	  * if value is valid, it is cast to appropriate uint type
+    **/
+
+    /** RL input index to ControlRig's input curve index for each LOD **/
+	UPROPERTY(transient)
+	TArray<int32> InputCurveIndices;  
+
+	/** RL joint index to ControlRig's hierarchy bone index **/
+	UPROPERTY(transient)
+	TArray<int32> HierarchyBoneIndices;
+
+	/** RL mesh blend shape index to ControlRig's output blendshape curve index for each LOD **/
+	UPROPERTY(transient)
+	TArray<FRigUnit_RigLogic_IntArray> MorphTargetCurveIndices;  
+
+	/** RL mesh+blend shape array index to RL blend shape index for each LOD **/
+	UPROPERTY(transient)
+	TArray<FRigUnit_RigLogic_IntArray> BlendShapeIndices;
+
+	/** RL animated map index to ControlRig's output anim map curve index for each LOD **/
+	UPROPERTY(transient)
+	TArray<FRigUnit_RigLogic_IntArray> CurveContainerIndicesForAnimMaps;
+
+	/** RL animated map index to RL anim map curve index for each LOD **/
+	UPROPERTY(transient)
+	TArray<FRigUnit_RigLogic_IntArray> RigLogicIndicesForAnimMaps;
+
+	/** Stores joints already updated, to avoid iterating through them multiple times **/
+	UPROPERTY(transient)
+	TArray<bool> UpdatedJoints;
+
+	/** Size of the blendshape map, cached from BehaviorReader **/
+	UPROPERTY(transient)
+	uint32 BlendShapeMappingCount = 0;
+
+	/** The index of RL input control for turning on the corrective expression for neck on average female body */
+	UPROPERTY(transient)
+	int32 NeckFemaleAverageCorCurveIndex = -1;
+
+	/** The index of RL input control for turning on the corrective expression for neck on mascular male body */
+	UPROPERTY(transient)
+	int32 NeckMaleMuscularCorExpCurveIndex = -1;
+
+	static const uint8 MAX_ATTRS_PER_JOINT;
+
+	void InitializeRigLogic(IBehaviorReader* DNABehavior);
+	bool IsRigLogicInitialized();
+	void ChangeRigLogicLODIfNeeded();
+
+	/** Constructs curve name from nameToSplit using formatString of form x<obj>y<attr>z **/
+	FString ConstructCurveName(const FString& nameToSplit, const FString& formatString);
+
+	/** Makes a map of input curve indices from DNA file to the control rig curves **/
+	void MapInputCurveIndices(const IBehaviorReader* DNABehavior, const FRigCurveContainer* CurveContainer);
+	/** Uses names to map joint indices from DNA file to the indices of bones in control rig hierarchy **/
+	void MapJoints(const IBehaviorReader* DNABehavior, const FRigBoneHierarchy* Hierarchy);
+	/** Uses names of blend shapes and meshes in DNA file, for all LODs, to map their indices to the indices of
+	  * morph target curves in the curve container; curve name format is <mesh>__<blendshape> **/
+	void MapMorphTargets(const IBehaviorReader* DNABehavior, const FRigCurveContainer* CurveContainer );
+	/** Uses names to map mask multiplier indices from DNA file, for all LODs, to the indices of curves in the
+	  * control rig's curve container **/
+	void MapMaskMultipliers(const IBehaviorReader* DNABehavior, const FRigCurveContainer* CurveContainer );
+
+	/** Calculates joint positions, orientation and scale based on inputs curves of the control rig **/
+	void CalculateRigLogic(FControlRigExecuteContext& ExecuteContext);
+	/** Updates joint positions in the hierarchy based on inputs curves of the control rig **/
+	void UpdateJoints(FRigHierarchyContainer* Hierarchy, const FRigUnit_RigLogic_JointUpdateParams& JointUpdateParams);
+	/** Updates morph target curve values based on values of input curves of the control rig **/
+	void UpdateBlendShapeCurves(FRigCurveContainer* CurveContainer, TArrayView<const float>& BlendShapeValues);
+	/** Updates anim map curve values based on values of input curves of the control rig **/
+	void UpdateAnimMapCurves(FRigCurveContainer* CurveContainer, TArrayView<const float>& AnimMapOutputs);
+
+	IBehaviorReader* FetchBehaviorReaderFromOwner();
+};
+
+
+/** RigLogic is used to translate control input curves into bone transforms and values for blend shape and
+  *  animated map multiplier curves */
+
+USTRUCT(meta = (DisplayName = "RigLogic", Category = "RigLogic", DocumentationPolicy = "Strict", Keywords = "Rig,RigLogic"))
+struct FRigUnit_RigLogic: public FRigUnitMutable
+{
+	GENERATED_BODY()
+
+	public:
+#if WITH_DEV_AUTOMATION_TESTS
+	/** Allows accessing private Data property from unit tests **/
+	struct TestAccessor;
+	friend TestAccessor;
+#endif
+
+	FRigUnit_RigLogic()
+	{
+	}
+
+	virtual FString GetUnitLabel() const override;
+	RIGVM_METHOD()
+	virtual void Execute(const FRigUnitContext& Context) override;
+
+private:
+
+	// internal work data for the unit
+	UPROPERTY(transient)
+	FRigUnit_RigLogic_Data Data;
+};
