Editgrating RigLogic plugin from Dev-CharacterTech to UE5/Main @ CL 13891876.

diff --git a/Engine/Plugins/Animation/RigLogic/Source/RigLogicModule/Public/RigUnit_RigLogic_Test.h b/Engine/Plugins/Animation/RigLogic/Source/RigLogicModule/Public/RigUnit_RigLogic_Test.h
new file mode 100644
--- /dev/null
+++ b/Engine/Plugins/Animation/RigLogic/Source/RigLogicModule/Public/RigUnit_RigLogic_Test.h
@@ -0,0 +1,519 @@
+#pragma once
+
+#include "CoreTypes.h"
+#include "CoreMinimal.h"
+
+#include "DNAReader.h"
+
+#if WITH_DEV_AUTOMATION_TESTS
+	#include "RigUnit_RigLogic.h"
+#endif
+
+struct FRigUnit_RigLogic_JointUpdateParams;
+
+//** A unit test-only helper class to prevent FString from being deleted as it is converted from FString*//
+
+class TestBehaviorReader: public IBehaviorReader
+{
+public:
+	// DescriptorReader methods start
+	FString GetName() const override
+	{
+		return {};
+	}
+
+	EArchetype GetArchetype() const override
+	{
+		return {};
+	}
+
+	EGender GetGender() const override
+	{
+		return {};
+	}
+
+	uint16 GetAge() const override
+	{
+		return {};
+	}
+
+	uint32 GetMetaDataCount() const override
+	{
+		return {};
+	}
+
+	FString GetMetaDataKey(uint32 Index) const override
+	{
+		return {};
+	}
+
+	FString GetMetaDataValue(const FString& Key) const override
+	{
+		return {};
+	}
+
+	ETranslationUnit GetTranslationUnit() const override
+	{
+		return {};
+	}
+
+	ERotationUnit GetRotationUnit() const override
+	{
+		return {};
+	}
+
+	FCoordinateSystem GetCoordinateSystem() const override
+	{
+		return {};
+	}
+
+	uint16 GetLODCount() const override
+	{
+		return LODCount;
+	}
+
+	uint16 GetDBMaxLOD() const override
+	{
+		return {};
+	}
+
+	FString GetDBComplexity() const override
+	{
+		return {};
+	}
+
+	FString GetDBName() const override
+	{
+		return {};
+	}
+
+	// DefinitionReader methods start
+	uint16 GetGUIControlCount() const override
+	{
+		return {};
+	}
+
+	FString GetGUIControlName(uint16 Index) const override
+	{
+		return {};
+	}
+
+	uint16 GetRawControlCount() const override
+	{
+		return rawControls.Num();
+	}
+
+	FString GetRawControlName(uint16 Index) const override
+	{
+		return rawControls[Index];
+	}
+
+	uint16 GetAnimatedMapCount() const override
+	{
+		return animatedMaps.Num();
+	}
+
+	FString GetAnimatedMapName(uint16 Index) const override
+	{
+		return animatedMaps[Index];
+	}
+
+	TArrayView<const uint16> GetAnimatedMapIndicesForLOD(uint16 lod) const override
+	{
+		if (lod == 1)
+		{
+			return TArrayView<const uint16_t>(animatedMapIndicesLOD1.GetData(), animatedMapIndicesLOD1.Num());
+		}
+		else
+		{
+			return TArrayView<const uint16_t>(animatedMapIndicesLOD0.GetData(), animatedMapIndicesLOD0.Num());
+		}
+	}
+
+	uint16 GetJointCount() const override
+	{
+		return static_cast<uint16>(jointNames.Num());
+	}
+
+	FString GetJointName(uint16 Index) const override
+	{
+		return jointNames[Index];
+	}
+
+	TArrayView<const uint16> GetJointIndicesForLOD(uint16 lod) const override
+	{
+		if (lod == 1)
+		{
+			return TArrayView<const uint16_t>(jointIndicesLOD1.GetData(), jointIndicesLOD1.Num());
+		}
+		else
+		{
+			return TArrayView<const uint16_t>(jointIndicesLOD0.GetData(), jointIndicesLOD0.Num());
+		}
+	}
+
+	uint16 GetJointParentIndex(uint16 Index) const override
+	{
+		return {};
+	}
+
+	uint16 GetBlendShapeChannelCount() const override
+	{
+		return static_cast<uint16>(blendShapeChannelNames.Num());
+	}
+
+	FString GetBlendShapeChannelName(uint16 Index) const override
+	{
+		return blendShapeChannelNames[Index];
+	}
+
+	TArrayView<const uint16> GetBlendShapeChannelIndicesForLOD(uint16 lod) const override
+	{
+		if (lod == 0)
+		{
+			return TArrayView<const uint16_t>(blendShapeIndicesForLOD0.GetData(), blendShapeIndicesForLOD0.Num());
+		}
+		else
+		{
+			return TArrayView<const uint16_t>(blendShapeIndicesForLOD1.GetData(), blendShapeIndicesForLOD1.Num());
+		}
+	}
+
+	uint16 GetMeshCount() const override
+	{
+		return {};
+	}
+
+	FString GetMeshName(uint16 Index) const override
+	{		
+		return meshNames[Index];
+	}
+
+	uint16 GetMeshBlendShapeChannelMappingCount() const override
+	{
+		return meshBlendShapeChannelMappings.Num();
+	}
+
+	FMeshBlendShapeChannelMapping GetMeshBlendShapeChannelMapping(uint16 Index) const override
+	{
+		return meshBlendShapeChannelMappings[Index];
+	}
+
+	TArrayView<const uint16> GetMeshBlendShapeChannelMappingIndicesForLOD(uint16 lod) const override
+	{
+		if (lod == 0)
+		{
+			return TArrayView<const uint16_t>(blendShapeMappingIndicesLOD0.GetData(), blendShapeMappingIndicesLOD0.Num());
+		}
+		else
+		{
+			return TArrayView<const uint16_t>(blendShapeMappingIndicesLOD1.GetData(), blendShapeMappingIndicesLOD1.Num());
+		}
+	}
+
+	TArrayView<const uint16> GetMeshIndicesForLOD(uint16 lod) const override
+	{
+		if (lod == 0)
+		{
+			return TArrayView<const uint16_t>(meshIndicesLOD0.GetData(), meshIndicesLOD0.Num());
+		}
+		else
+		{
+			return TArrayView<const uint16_t>(meshIndicesLOD1.GetData(), meshIndicesLOD1.Num());
+		}
+	}
+
+	FVector GetNeutralJointTranslation(uint16 Index) const override
+	{
+		return {};
+	}
+
+	FVector GetNeutralJointRotation(uint16 Index) const override
+	{
+		return {};
+	}
+
+	// BehaviorReader methods start
+	TArrayView<const uint16> GetGUIToRawInputIndices() const override
+	{
+		return {};
+	}
+
+	TArrayView<const uint16> GetGUIToRawOutputIndices() const override
+	{
+		return {};
+	}
+
+	TArrayView<const float> GetGUIToRawFromValues() const override
+	{
+		return {};
+	}
+
+	TArrayView<const float> GetGUIToRawToValues() const override
+	{
+		return {};
+	}
+
+	TArrayView<const float> GetGUIToRawSlopeValues() const override
+	{
+		return {};
+	}
+
+	TArrayView<const float> GetGUIToRawCutValues() const override
+	{
+		return {};
+	}
+
+	uint16 GetPSDCount() const override
+	{
+		return {};
+	}
+
+	TArrayView<const uint16> GetPSDRowIndices() const override
+	{
+		return {};
+	}
+
+	TArrayView<const uint16> GetPSDColumnIndices() const override
+	{
+		return {};
+	}
+
+	TArrayView<const float> GetPSDValues() const override
+	{
+		return {};
+	}
+
+	TArrayView<const uint16> GetJointVariableAttributeIndices(uint16 lod) const override
+	{
+		return {};
+	}
+
+	uint16 GetJointRowCount() const override
+	{
+		return 0;
+	}
+
+	uint16 GetJointColumnCount() const override
+	{
+		return 0;
+	}
+
+	uint16 GetJointGroupCount() const override
+	{
+		return 0;
+	}
+
+	TArrayView<const uint16> GetJointGroupLODs(uint16 jointGroupIndex) const override
+	{
+		return {};
+	}
+
+	TArrayView<const uint16> GetJointGroupInputIndices(uint16 jointGroupIndex) const override
+	{
+		return {};
+	}
+
+	TArrayView<const uint16> GetJointGroupOutputIndices(uint16 jointGroupIndex) const override
+	{
+		return {};
+	}
+
+	TArrayView<const float> GetJointGroupValues(uint16 jointGroupIndex) const override
+	{
+		return {};
+	}
+
+	TArrayView<const uint16> GetJointGroupJointIndices(uint16 jointGroupIndex) const override
+	{
+		return {};
+	}
+
+	TArrayView<const uint16> GetBlendShapeChannelLODs() const override
+	{
+		return {};
+	}
+
+	TArrayView<const uint16> GetBlendShapeChannelInputIndices() const override
+	{
+		return {};
+	}
+
+	TArrayView<const uint16> GetBlendShapeChannelOutputIndices() const override
+	{
+		return {};
+	}
+
+	TArrayView<const uint16> GetAnimatedMapLODs() const override
+	{
+		return {};
+	}
+
+	TArrayView<const uint16> GetAnimatedMapInputIndices() const override
+	{
+		return {};
+	}
+
+	TArrayView<const uint16> GetAnimatedMapOutputIndices() const override
+	{
+		return {};
+	}
+
+	TArrayView<const float> GetAnimatedMapFromValues() const override
+	{
+		return {};
+	}
+
+	TArrayView<const float> GetAnimatedMapToValues() const override
+	{
+		return {};
+	}
+
+	TArrayView<const float> GetAnimatedMapSlopeValues() const override
+	{
+		return {};
+	}
+
+	TArrayView<const float> GetAnimatedMapCutValues() const override
+	{
+		return {};
+	}
+
+	uint16 LODCount;
+
+	TArray<FString> rawControls;
+
+	TArray<FString> jointNames;
+	TArray<uint16> jointIndicesLOD0;
+	TArray<uint16> jointIndicesLOD1;
+
+	TArray<FString> meshNames;
+	TArray<uint16> meshIndicesLOD0; //indices into previous array
+	TArray<uint16> meshIndicesLOD1; //indices into previous array
+
+	TArray<FString> blendShapeChannelNames;
+	TArray<uint16> blendShapeIndicesForLOD0; //indices into previous array
+	TArray<uint16> blendShapeIndicesForLOD1;
+
+	TArray<FMeshBlendShapeChannelMapping> meshBlendShapeChannelMappings; //contains all mappings from blendshapes to meshes
+	TArray<uint16> blendShapeMappingIndicesLOD0; //indices into previous array
+	TArray<uint16> blendShapeMappingIndicesLOD1; 
+
+	TArray<FString> animatedMaps;
+	TArray<uint16> animatedMapIndicesLOD0;
+	TArray<uint16> animatedMapIndicesLOD1;
+
+	void addJoint(FString newJointName)
+	{
+		jointNames.Add(newJointName);
+	}
+
+	void addBlendShapeChannelName( FString newBlendShapeName )
+	{
+		blendShapeChannelNames.Add(newBlendShapeName);
+	}
+
+	void addMeshName(FString newMeshName)
+	{
+		meshNames.Add(newMeshName);
+	}
+
+	void addBlendShapeMapping( uint16 meshIndex, uint16 blendShapeChannelIndex )
+	{
+		//{ (0,0), (0, 1), (0, 2), (1, 3), (1, 4), .... (5, 666) }
+		//  ^^^^^^^^^^^^^^^^^^^^^   ^^^^^^^^^^^^^       ^^^^^^^^^
+		//     mesh 1                 mesh 2             mesh n
+		FMeshBlendShapeChannelMapping mapping;
+		mapping.MeshIndex = meshIndex;
+		mapping.BlendShapeChannelIndex = blendShapeChannelIndex;
+		meshBlendShapeChannelMappings.Add(mapping);
+	}
+
+	void addBlendShapeMappingIndicesToLOD(uint16 MappingIndex, uint16 LODIndex)
+	{
+		if (LODIndex == 0)
+		{
+			blendShapeMappingIndicesLOD0.Add(MappingIndex);
+		}
+		else 
+		{
+			blendShapeMappingIndicesLOD1.Add(MappingIndex);
+		}
+	}
+
+	void addAnimatedMapIndicesToLOD(uint16 MappingIndex, uint16 LODIndex)
+	{
+		if (LODIndex == 0)
+		{
+			animatedMapIndicesLOD0.Add(MappingIndex);
+		}
+		else
+		{
+			animatedMapIndicesLOD1.Add(MappingIndex);
+		}
+	}
+
+private:
+	dna::Reader* Unwrap() const
+	{
+		return nullptr;  // Unused
+	}
+};
+
+struct FRigUnit_RigLogic::TestAccessor
+{
+	TestAccessor(FRigUnit_RigLogic* Unit);
+
+	 /** MapInputCurve Tests **/
+
+	TestBehaviorReader* CreateBehaviorReaderEmpty();
+	TestBehaviorReader* CreateBehaviorReaderOneCurve(FString ControlNameStr);
+
+	FRigCurveContainer* CreateCurveContainerEmpty();
+	FRigCurveContainer* CreateCurveContainerOneCurve(FString CurveNameStr);
+
+	void Exec_MapInputCurve(TestBehaviorReader* TestReader, FRigCurveContainer* TestCurveContainer);
+
+	/** MapJoints Tests **/
+
+	TestBehaviorReader* CreateBehaviorReaderTwoJoints(FString Joint1NameStr, FString Joint2NameStr);
+	FRigBoneHierarchy* CreateBoneHierarchyEmpty();
+	FRigBoneHierarchy* CreateBoneHierarchyTwoBones(FString Bone1NameStr, FString Bone2NameStr);
+
+	void Exec_MapJoints(TestBehaviorReader* TestReader, FRigBoneHierarchy* MockHierachy);
+
+	/** MapBlendShapes Tests **/
+
+	TestBehaviorReader* CreateBehaviorReaderNoBlendshapes(FString MeshNameStr);
+	TestBehaviorReader* CreateBehaviorReaderOneBlendShape(FString MeshNameStr, FString BlendShapeNameStr);
+	FRigCurveContainer* CreateCurveContainerOneMorphTarget(FString MorphTargetStr);
+	TestBehaviorReader* CreateBehaviorReaderTwoBlendShapes(FString MeshNameStr, FString BlendShape1Str, FString BlendShape2Str);
+	FRigCurveContainer* CreateCurveContainerTwoMorphTargets(FString MorphTarget1Str, FString MorphTarget2Str);
+
+	void Exec_MapMorphTargets(TestBehaviorReader* TestReader, FRigCurveContainer* TestCurveContainer);
+
+	/** MapMaskMultipliers Tests **/
+	TestBehaviorReader* CreateBehaviorReaderOneAnimatedMap(FString AnimatedMapNameStr);
+
+	void Exec_MapMaskMultipliers(TestBehaviorReader* TestReader, FRigCurveContainer* TestCurveContainer);
+
+
+	static const uint8 MAX_ATTRS_PER_JOINT;
+
+	/** UpdateJoints Tests **/
+	void AddToTransformArray(float* InArray, FTransform& Transform);
+	FTransformArrayView CreateTwoJointNeutralTransforms(float *InValueArray);
+	TArrayView<const uint16> CreateTwoJointVariableAttributes(uint16* InVariableAttributeIndices, uint8 LOD);
+
+	void Exec_UpdateJoints(FRigHierarchyContainer* TestHierarchyContainer, FRigUnit_RigLogic_JointUpdateParams& JointUpdateParams);
+
+
+private:
+	FName RawCtrlName;
+	FRigCurveContainer CurveContainer;
+	FRigBoneHierarchy Hierarchy;
+	FRigUnit_RigLogic* Unit;
+
+public:
+	FRigUnit_RigLogic_Data* GetData() { return &(Unit->Data); };
+};
\ No newline at end of file
diff --git a/Engine/Plugins/Animation/RigLogic/Source/RigLogicModule/Public/TransformArrayView.h b/Engine/Plugins/Animation/RigLogic/Source/RigLogicModule/Public/TransformArrayView.h
new file mode 100644
--- /dev/null
+++ b/Engine/Plugins/Animation/RigLogic/Source/RigLogicModule/Public/TransformArrayView.h
@@ -0,0 +1,50 @@
+// Copyright 1998-2019 Epic Games, Inc. All Rights Reserved.
+
+#pragma once
+
+#include "CoreMinimal.h"
+#include "Math/Transform.h"
+#include "UObject/NoExportTypes.h"
+#include "UObject/ObjectMacros.h"
+
+class RIGLOGICMODULE_API FTransformArrayView
+{
+private:
+	enum
+	{
+		TransformationSize = 9ul
+	};
+
+public:
+	FTransformArrayView(const float* Values, size_t Size) :
+		Values{Values},
+		Size{Size} {
+	}
+
+	size_t Count() const
+	{
+		return Size / TransformationSize;
+	}
+
+	FTransform operator[](size_t Index) const
+	{
+		return ConvertCoordinateSystem(Values + (Index * TransformationSize));
+	}
+
+private:
+	FTransform ConvertCoordinateSystem(const float* Source) const
+	{
+		return FTransform{
+			// Rotation: X = -Y, Y = -Z, Z = X
+			FRotator(-Source[4], -Source[5], Source[3]),
+			// Translation: X = X, Y = -Y, Z = Z
+			FVector(Source[0], -Source[1], Source[2]),
+			// Scale: X = X, Y = Y, Z = Z
+			FVector(Source[6], Source[7], Source[8])
+		};
+	}
+
+private:
+	const float* Values;
+	size_t Size;
+};
diff --git a/Engine/Plugins/Animation/RigLogic/Source/RigLogicModule/RigLogicModule.Build.cs b/Engine/Plugins/Animation/RigLogic/Source/RigLogicModule/RigLogicModule.Build.cs
new file mode 100644
--- /dev/null
+++ b/Engine/Plugins/Animation/RigLogic/Source/RigLogicModule/RigLogicModule.Build.cs
@@ -0,0 +1,50 @@
+ // Copyright 1998-2019 Epic Games, Inc. All Rights Reserved.
+
+using System.IO;
+
+namespace UnrealBuildTool.Rules
+{
+    public class RigLogicModule : ModuleRules
+    {
+        private string ModulePath
+        {
+            get { return ModuleDirectory; }
+        }
+
+        public RigLogicModule(ReadOnlyTargetRules Target) : base(Target)
+        {
+            PublicDefinitions.Add("RL_SHARED=1"); //used instead of #define so cpp files can do conditional compilation
+
+
+            PublicDependencyModuleNames.AddRange(
+                new string[]
+                {
+                    "Core",
+                    "CoreUObject",
+                    "Engine",
+					"RigVM",
+                    "ControlRig",
+                    "RigLogicLib",
+                    "MessageLog",
+                    "RigVM",
+					"Projects"
+                } 
+                );
+
+            if (Target.Type == TargetType.Editor)
+            {
+                PublicDependencyModuleNames.Add("UnrealEd");
+            }
+
+            PrivateDependencyModuleNames.AddRange(
+                new string[]
+                {
+					"AnimationCore",
+                    "ControlRig",
+					"RenderCore",
+					"RHI"
+				}
+             );
+        }
+    }
+}
diff --git a/Engine/Plugins/Animation/RigLogic/Source/RigLogicModule/RigLogicModule_UPL.xml b/Engine/Plugins/Animation/RigLogic/Source/RigLogicModule/RigLogicModule_UPL.xml
new file mode 100644
--- /dev/null
+++ b/Engine/Plugins/Animation/RigLogic/Source/RigLogicModule/RigLogicModule_UPL.xml
@@ -0,0 +1,32 @@
+<?xml version="1.0" encoding="utf-8"?>
+<!--SpeechGraphics plugin additions-->
+<root xmlns:android="http://schemas.android.com/apk/res/android">
+	<!-- init section is always evaluated once per architecture -->
+	<init>
+		<log text="RigLogic init"/>
+		<setBool result="bSupported" value="false"/>
+		<isArch arch="arm64-v8a">
+			<setBool result="bSupported" value="true"/>
+		</isArch>
+		<if condition="bSupported">
+			<false>
+				<log text="Unsupported architecture: $S(Architecture)"/>
+			</false>
+		</if>
+	</init>
+
+	<!-- optional files or directories to copy to Intermediate/Android/APK -->
+	<resourceCopies>
+		<if condition="bSupported">
+			<true>					
+				<copyFile src="$S(AbsPluginDir)/ThirdParty/UFC/Lib/Android/arm64/libufc.so"
+							dst="$S(BuildDir)/libs/$S(Architecture)/libufc.so" />						
+				<copyFile src="$S(AbsPluginDir)/ThirdParty/UFC/Lib/Android/arm64/libc++_shared.so"
+							dst="$S(BuildDir)/libs/$S(Architecture)/libc++_shared.so" />
+				<copyFile src="$S(AbsPluginDir)/ThirdParty/UFC/Lib/Android/arm64/liblog.so"
+							dst="$S(BuildDir)/libs/$S(Architecture)/liblog.so" />																			
+				<log text="Copying libufc.so"/>
+				</true>
+		</if>
+	</resourceCopies>
+</root>
