RigLogic: Fix build issues with non-unity builds and unknown defs on Linux.

#rnx

diff --git a/Engine/Plugins/Animation/RigLogic/Source/RigLogicEditor/Private/DNAAssetImportFactory.cpp b/Engine/Plugins/Animation/RigLogic/Source/RigLogicEditor/Private/DNAAssetImportFactory.cpp
--- a/Engine/Plugins/Animation/RigLogic/Source/RigLogicEditor/Private/DNAAssetImportFactory.cpp
+++ b/Engine/Plugins/Animation/RigLogic/Source/RigLogicEditor/Private/DNAAssetImportFactory.cpp
@@ -3,23 +3,25 @@
 #include "DNAAssetImportFactory.h"
 #include "DNAAsset.h"
 
-#include "Runtime/Engine/Classes/Engine/SkeletalMesh.h"
-#include "Runtime/Engine/Public/ReferenceSkeleton.h"
-#include "Runtime/Engine/Public/Rendering/SkeletalMeshModel.h"
-#include "Runtime/Engine/Public/Rendering/SkeletalMeshLODModel.h"
-#include "Runtime/Engine/Public/Rendering/SkeletalMeshRenderData.h"
-#include "Runtime/Engine/Public/Rendering/SkeletalMeshLODRenderData.h"
-#include "Runtime/Engine/Public/ComponentReregisterContext.h"
-#include "Runtime/CoreUObject/Public/UObject/ConstructorHelpers.h"
+#include "Engine/SkeletalMesh.h"
+#include "ReferenceSkeleton.h"
+#include "Rendering/SkeletalMeshModel.h"
+#include "Rendering/SkeletalMeshLODModel.h"
+#include "Rendering/SkeletalMeshRenderData.h"
+#include "Rendering/SkeletalMeshLODRenderData.h"
+#include "ComponentReregisterContext.h"
+#include "UObject/ConstructorHelpers.h"
+#include "Editor.h"
 
 #include "ActorFactories/ActorFactoryStaticMesh.h"
 #include "ProfilingDebugging/ScopedTimers.h"
 
 #include "Engine/StaticMesh.h"
+#include "Misc/MessageDialog.h"
 #include "Misc/Paths.h"
 
 #include "AssetImportTask.h"
-#include "Runtime/Core/Public/Misc/FileHelper.h"
+#include "Misc/FileHelper.h"
 
 #include "SkelMeshDNAUtils.h"
 #include "DNAAssetImportUI.h"
diff --git a/Engine/Plugins/Animation/RigLogic/Source/RigLogicLibTest/Private/dnatests/Defs.h b/Engine/Plugins/Animation/RigLogic/Source/RigLogicLibTest/Private/dnatests/Defs.h
--- a/Engine/Plugins/Animation/RigLogic/Source/RigLogicLibTest/Private/dnatests/Defs.h
+++ b/Engine/Plugins/Animation/RigLogic/Source/RigLogicLibTest/Private/dnatests/Defs.h
@@ -9,6 +9,7 @@
 #ifdef __clang__
     #pragma clang diagnostic push
     #pragma clang diagnostic ignored "-Wglobal-constructors"
+	#pragma clang diagnostic ignored "-Wundef"
 #endif
 #include <gtest/gtest.h>
 #ifdef __clang__
diff --git a/Engine/Plugins/Animation/RigLogic/Source/RigLogicModule/Private/DNAToSkelMeshMap.cpp b/Engine/Plugins/Animation/RigLogic/Source/RigLogicModule/Private/DNAToSkelMeshMap.cpp
--- a/Engine/Plugins/Animation/RigLogic/Source/RigLogicModule/Private/DNAToSkelMeshMap.cpp
+++ b/Engine/Plugins/Animation/RigLogic/Source/RigLogicModule/Private/DNAToSkelMeshMap.cpp
@@ -5,14 +5,16 @@
 #include "DNAAsset.h"
 #include "SkelMeshDNAReader.h"
 
-#include "Runtime/Core/Public/Serialization/BufferArchive.h"
-#include "Runtime/Core/Public/Serialization/MemoryReader.h"
-#include "Runtime/Core/Public/Misc/Paths.h"
-#include "Runtime/Core/Public/Misc/FileHelper.h"
-#include "Runtime/CoreUObject/Public/UObject/UObjectGlobals.h"
-#include "Runtime/Engine/Public/Rendering/SkeletalMeshModel.h"
-#include "Runtime/Engine/Public/Rendering/SkeletalMeshLODModel.h"
-#include "Runtime/Engine/Classes/Animation/SmartName.h"
+#include "Serialization/BufferArchive.h"
+#include "Serialization/MemoryReader.h"
+#include "Misc/Paths.h"
+#include "Misc/FileHelper.h"
+#include "UObject/UObjectGlobals.h"
+#include "Rendering/SkeletalMeshModel.h"
+#include "Rendering/SkeletalMeshLODModel.h"
+#include "Animation/SmartName.h"
+#include "Engine/SkeletalMesh.h"
+
 
 #include "riglogic/RigLogic.h"
 
