Editgrating RigLogic plugin up to CL 13904728.

#rnx

diff --git a/Engine/Plugins/Animation/RigLogic/Source/RigLogicModule/Private/DNAAsset.cpp b/Engine/Plugins/Animation/RigLogic/Source/RigLogicModule/Private/DNAAsset.cpp
--- a/Engine/Plugins/Animation/RigLogic/Source/RigLogicModule/Private/DNAAsset.cpp
+++ b/Engine/Plugins/Animation/RigLogic/Source/RigLogicModule/Private/DNAAsset.cpp
@@ -70,7 +70,7 @@ bool UDNAAsset::Init(const FString DNAFilename)
 	BehaviorStreamReader = DNAReaderBehavior;
 	WriteToBuffer(DNAReaderBehavior.Get(), &BehaviorData, AVG_BEHAVIOR_SIZE);
 
-#if WITH_EDITOR
+#if WITH_EDITORONLY_DATA
 	//We use geometry part of the data in MHC only (for updating the SkeletalMesh with
 	//result of GeneSplicer), so we can drop geometry part when cooking for runtime
 	TSharedPtr<IDNAReader> DNAReaderGeometry = ReadDNALayerFromBuffer(&TempFileBuffer, EDNADataLayer::Geometry, 0u); //0u = MaxLOD
@@ -80,7 +80,7 @@ bool UDNAAsset::Init(const FString DNAFilename)
 	//to enable GeneSplicer to read geometry directly from SkeletalMeshes, as
 	//a way to save memory, as on consoles the "database" will be exactly the set of characters
 	//used in the game
-#endif
+#endif // #if WITH_EDITORONLY_DATA
 
 	return true;
 }
@@ -102,7 +102,7 @@ void UDNAAsset::PostLoad()
 		return;
 	}
 	GeometryStreamReader = ReadDNALayerFromBuffer(&DesignTimeData, EDNADataLayer::Geometry, 0u); //0u = max LOD
-#endif
+#endif // #if WITH_EDITORONLY_DATA
 }
 
 TSharedPtr<IDNAReader> UDNAAsset::CopyDNALayer(const IDNAReader* Source, EDNADataLayer DNADataLayer, uint32 PredictedSize )
@@ -131,5 +131,7 @@ void UDNAAsset::SetBehaviorReader(const TSharedPtr<IDNAReader> SourceDNAReader)
 
 void UDNAAsset::SetGeometryReader(const TSharedPtr<IDNAReader> SourceDNAReader)
 {
+#if WITH_EDITORONLY_DATA
 	GeometryStreamReader = CopyDNALayer(SourceDNAReader.Get(), EDNADataLayer::Geometry, AVG_GEOMETRY_SIZE);
+#endif // #if WITH_EDITORONLY_DATA
 }
\ No newline at end of file
