Editgrated RigLogic up to CL 13923318. Includes several platform fixes.

#rnx

diff --git a/Engine/Plugins/Animation/RigLogic/Config/FilterPlugin.ini b/Engine/Plugins/Animation/RigLogic/Config/FilterPlugin.ini
deleted file mode 100644
diff --git a/Engine/Plugins/Animation/RigLogic/LICENSE.txt b/Engine/Plugins/Animation/RigLogic/LICENSE.txt
new file mode 100644
--- /dev/null
+++ b/Engine/Plugins/Animation/RigLogic/LICENSE.txt
@@ -0,0 +1,11 @@
+Copyright (c) 2017, 3Lateral. All rights reserved.
+Owner: 3Lateral - www.3lateral.com.
+License:
+    The use of this source code and its functionality is regulated by the licensing terms defined by 3Lateral.
+    The source code has been delivered for the use only for this project and under strict conditions that were defined by 3Lateral.
+
+    Please note that this infrastructure, code and its functionality are subject to Confidentiality Agreement and that you
+    should not do or let anyone else do any decompilation, reverse engineering or any other attempt to learn the inner workings
+    of the system. In case you wish to apply this infrastructure, code and its functionality on any other project
+    (both commercial or non-commercial) please consult 3Lateral regarding the license terms and fees. Furthermore,
+    please consult and confirm with 3Lateral when it comes to allowing any other third party to access this infrastructure or code.
\ No newline at end of file
diff --git a/Engine/Plugins/Animation/RigLogic/RigLogic.uplugin b/Engine/Plugins/Animation/RigLogic/RigLogic.uplugin
--- a/Engine/Plugins/Animation/RigLogic/RigLogic.uplugin
+++ b/Engine/Plugins/Animation/RigLogic/RigLogic.uplugin
@@ -1,9 +1,9 @@
 {
 	"FileVersion": 3,
 	"Version": 5,
-	"VersionName": "5.1.0",
-	"FriendlyName": "RigLogic Plugin v5.1.0",
-	"Description": "3Lateral RigLogic Plugin for Facial Animation v5.1.0",
+	"VersionName": "5.2.0",
+	"FriendlyName": "RigLogic Plugin v5.2.0",
+	"Description": "3Lateral RigLogic Plugin for Facial Animation v5.2.0",
 	"Category": "Animation",
 	"CreatedBy": "3Lateral",
 	"CreatedByURL": "http://www.3lateral.com/",
@@ -17,15 +17,7 @@
 		{
 			"Name": "RigLogicLib",
 			"Type": "Runtime",
-			"LoadingPhase": "PreDefault",
-			"WhitelistPlatforms": [
-				"Win32",
-				"Win64",
-				"Linux",
-				"Mac",
-				"PS4",
-				"XboxOne"
-			]
+			"LoadingPhase": "PreDefault"
 		},
 		{
 			"Name": "RigLogicLibTest",
@@ -41,10 +33,7 @@
 		{
 			"Name": "RigLogicModule",
 			"Type": "Runtime",
-			"LoadingPhase": "Default",
-			"WhitelistPlatforms": [
-				"Win64"
-			]
+			"LoadingPhase": "Default"
 		},
 		{
 			"Name": "RigLogicEditor",
diff --git a/Engine/Plugins/Animation/RigLogic/Source/RigLogicEditor/RigLogicEditor.Build.cs b/Engine/Plugins/Animation/RigLogic/Source/RigLogicEditor/RigLogicEditor.Build.cs
--- a/Engine/Plugins/Animation/RigLogic/Source/RigLogicEditor/RigLogicEditor.Build.cs
+++ b/Engine/Plugins/Animation/RigLogic/Source/RigLogicEditor/RigLogicEditor.Build.cs
@@ -1,4 +1,4 @@
- // Copyright 1998-2019 Epic Games, Inc. All Rights Reserved.
+// Copyright 1998-2019 Epic Games, Inc. All Rights Reserved.
 
 using System.IO;
 
@@ -18,7 +18,7 @@ namespace UnrealBuildTool.Rules
 
         public RigLogicEditor(ReadOnlyTargetRules Target) : base(Target)
         {
-            PublicDefinitions.Add("RL_SHARED=1"); //used instead of #define so cpp files can do conditional compilation
+            PublicDefinitions.Add("RL_SHARED=1");
 
             PublicDependencyModuleNames.AddRange(
                 new string[]
@@ -27,26 +27,26 @@ namespace UnrealBuildTool.Rules
                     "CoreUObject",
                     "Engine",
                     "ControlRig",
-					"UnrealEd",
-					"MainFrame",
-					"RigLogicModule",
-					"RigLogicLib",
-					"PropertyEditor",
-					"SlateCore",
-					"ApplicationCore",
-					"Slate",
-					"EditorStyle",
-					"InputCore"
+                    "UnrealEd",
+                    "MainFrame",
+                    "RigLogicModule",
+                    "RigLogicLib",
+                    "PropertyEditor",
+                    "SlateCore",
+                    "ApplicationCore",
+                    "Slate",
+                    "EditorStyle",
+                    "InputCore"
                 }
             );
 
             PrivateIncludePathModuleNames.AddRange(
-				new string[]
-				{
+                new string[]
+                {
                     "PropertyEditor",
                     "AssetTools"
-				}
-			);
+                }
+            );
 
         }
     }
diff --git a/Engine/Plugins/Animation/RigLogic/Source/RigLogicLib/Private/dna/WriterImpl.hpp b/Engine/Plugins/Animation/RigLogic/Source/RigLogicLib/Private/dna/WriterImpl.hpp
--- a/Engine/Plugins/Animation/RigLogic/Source/RigLogicLib/Private/dna/WriterImpl.hpp
+++ b/Engine/Plugins/Animation/RigLogic/Source/RigLogicLib/Private/dna/WriterImpl.hpp
@@ -60,7 +60,7 @@ void WriterImpl<TWriterBase>::setMetaData(const char* key, const char* value) {
             return (std::strlen(key) == k.size() && std::strncmp(k.data(), key, k.size()) == 0);
         });
     if (it == dna.descriptor.metadata.end()) {
-        dna.descriptor.metadata.push_back({String<char>{key, memRes}, String<char>{value, memRes}});
+        dna.descriptor.metadata.emplace_back(String<char>{key, memRes}, String<char>{value, memRes});
     } else {
         std::get<1>(*it) = value;
     }
diff --git a/Engine/Plugins/Animation/RigLogic/Source/RigLogicLib/Private/riglogic/joints/bpcm/strategies/AVX.h b/Engine/Plugins/Animation/RigLogic/Source/RigLogicLib/Private/riglogic/joints/bpcm/strategies/AVX.h
--- a/Engine/Plugins/Animation/RigLogic/Source/RigLogicLib/Private/riglogic/joints/bpcm/strategies/AVX.h
+++ b/Engine/Plugins/Animation/RigLogic/Source/RigLogicLib/Private/riglogic/joints/bpcm/strategies/AVX.h
@@ -8,7 +8,7 @@ namespace rl4 {
 namespace bpcm {
 
 template<typename T>
-using AVXJointCalculationStrategy = Block8JointCalculationStrategy<simd::avx::F256, T>;
+using AVXJointCalculationStrategy = Block8JointCalculationStrategy<trimd::avx::F256, T>;
 
 }  // namespace bpcm
 
diff --git a/Engine/Plugins/Animation/RigLogic/Source/RigLogicLib/Private/riglogic/joints/bpcm/strategies/SSE.h b/Engine/Plugins/Animation/RigLogic/Source/RigLogicLib/Private/riglogic/joints/bpcm/strategies/SSE.h
--- a/Engine/Plugins/Animation/RigLogic/Source/RigLogicLib/Private/riglogic/joints/bpcm/strategies/SSE.h
+++ b/Engine/Plugins/Animation/RigLogic/Source/RigLogicLib/Private/riglogic/joints/bpcm/strategies/SSE.h
@@ -8,7 +8,7 @@ namespace rl4 {
 namespace bpcm {
 
 template<typename T>
-using SSEJointCalculationStrategy = Block4JointCalculationStrategy<simd::sse::F128, T>;
+using SSEJointCalculationStrategy = Block4JointCalculationStrategy<trimd::sse::F128, T>;
 
 }  // namespace bpcm
 
diff --git a/Engine/Plugins/Animation/RigLogic/Source/RigLogicLib/Private/riglogic/joints/bpcm/strategies/Scalar.h b/Engine/Plugins/Animation/RigLogic/Source/RigLogicLib/Private/riglogic/joints/bpcm/strategies/Scalar.h
--- a/Engine/Plugins/Animation/RigLogic/Source/RigLogicLib/Private/riglogic/joints/bpcm/strategies/Scalar.h
+++ b/Engine/Plugins/Animation/RigLogic/Source/RigLogicLib/Private/riglogic/joints/bpcm/strategies/Scalar.h
@@ -7,7 +7,7 @@ namespace rl4 {
 
 namespace bpcm {
 
-using ScalarJointCalculationStrategy = Block4JointCalculationStrategy<simd::scalar::F128, float>;
+using ScalarJointCalculationStrategy = Block4JointCalculationStrategy<trimd::scalar::F128, float>;
 
 }  // namespace bpcm
 
diff --git a/Engine/Plugins/Animation/RigLogic/Source/RigLogicLib/Private/riglogic/system/simd/SIMD.h b/Engine/Plugins/Animation/RigLogic/Source/RigLogicLib/Private/riglogic/system/simd/SIMD.h
--- a/Engine/Plugins/Animation/RigLogic/Source/RigLogicLib/Private/riglogic/system/simd/SIMD.h
+++ b/Engine/Plugins/Animation/RigLogic/Source/RigLogicLib/Private/riglogic/system/simd/SIMD.h
@@ -1,9 +1,15 @@
 #pragma once
 
+#ifdef RL_USE_HALF_FLOATS
+    #define TRIMD_ENABLE_F16C
+#endif  // RL_USE_HALF_FLOATS
+
 #ifdef RL_BUILD_WITH_AVX
-    #include "riglogic/system/simd/AVX.h"
-#endif
+    #define TRIMD_ENABLE_AVX
+#endif  // RL_BUILD_WITH_AVX
+
 #ifdef RL_BUILD_WITH_SSE
-    #include "riglogic/system/simd/SSE.h"
-#endif
-#include "riglogic/system/simd/Scalar.h"
+    #define TRIMD_ENABLE_SSE
+#endif  // RL_BUILD_WITH_SSE
+
+#include <trimd/TRiMD.h>
diff --git a/Engine/Plugins/Animation/RigLogic/Source/RigLogicLib/Private/riglogic/system/simd/Scalar.h b/Engine/Plugins/Animation/RigLogic/Source/RigLogicLib/Private/riglogic/system/simd/Scalar.h
deleted file mode 100644
diff --git a/Engine/Plugins/Animation/RigLogic/Source/RigLogicLib/Private/status/Storage.cpp b/Engine/Plugins/Animation/RigLogic/Source/RigLogicLib/Private/status/Storage.cpp
--- a/Engine/Plugins/Animation/RigLogic/Source/RigLogicLib/Private/status/Storage.cpp
+++ b/Engine/Plugins/Animation/RigLogic/Source/RigLogicLib/Private/status/Storage.cpp
@@ -27,7 +27,7 @@ void strcopy(char* destination, const char* source, std::size_t size) {
 
 }  // namespace
 
-thread_local static StatusCodeStorage currentStatus{OK.code, "Ok"};
+thread_local static StatusCodeStorage currentStatus{0, "Ok"};
 
 void StatusStorage::set(StatusCode status) {
     // The Release build will eliminate this call, as it's really just a sanity check
diff --git a/Engine/Plugins/Animation/RigLogic/Source/RigLogicLib/Public/dna/version/Version.h b/Engine/Plugins/Animation/RigLogic/Source/RigLogicLib/Public/dna/version/Version.h
--- a/Engine/Plugins/Animation/RigLogic/Source/RigLogicLib/Public/dna/version/Version.h
+++ b/Engine/Plugins/Animation/RigLogic/Source/RigLogicLib/Public/dna/version/Version.h
@@ -2,5 +2,5 @@
 
 #define DNA_MAJOR_VERSION 1
 #define DNA_MINOR_VERSION 1
-#define DNA_PATCH_VERSION 0
-#define DNA_VERSION_STRING "1.1.0"
+#define DNA_PATCH_VERSION 1
+#define DNA_VERSION_STRING "1.1.1"
diff --git a/Engine/Plugins/Animation/RigLogic/Source/RigLogicLib/Public/pma/version/Version.h b/Engine/Plugins/Animation/RigLogic/Source/RigLogicLib/Public/pma/version/Version.h
--- a/Engine/Plugins/Animation/RigLogic/Source/RigLogicLib/Public/pma/version/Version.h
+++ b/Engine/Plugins/Animation/RigLogic/Source/RigLogicLib/Public/pma/version/Version.h
@@ -2,5 +2,5 @@
 
 #define PMA_MAJOR_VERSION 1
 #define PMA_MINOR_VERSION 0
-#define PMA_PATCH_VERSION 1
-#define PMA_VERSION_STRING "1.0.1"
+#define PMA_PATCH_VERSION 2
+#define PMA_VERSION_STRING "1.0.2"
diff --git a/Engine/Plugins/Animation/RigLogic/Source/RigLogicLib/Public/riglogic/version/Version.h b/Engine/Plugins/Animation/RigLogic/Source/RigLogicLib/Public/riglogic/version/Version.h
--- a/Engine/Plugins/Animation/RigLogic/Source/RigLogicLib/Public/riglogic/version/Version.h
+++ b/Engine/Plugins/Animation/RigLogic/Source/RigLogicLib/Public/riglogic/version/Version.h
@@ -1,6 +1,6 @@
 #pragma once
 
 #define RL_MAJOR_VERSION 5
-#define RL_MINOR_VERSION 1
+#define RL_MINOR_VERSION 2
 #define RL_PATCH_VERSION 0
-#define RL_VERSION_STRING "5.1.0"
+#define RL_VERSION_STRING "5.2.0"
diff --git a/Engine/Plugins/Animation/RigLogic/Source/RigLogicLib/Public/status/version/Version.h b/Engine/Plugins/Animation/RigLogic/Source/RigLogicLib/Public/status/version/Version.h
--- a/Engine/Plugins/Animation/RigLogic/Source/RigLogicLib/Public/status/version/Version.h
+++ b/Engine/Plugins/Animation/RigLogic/Source/RigLogicLib/Public/status/version/Version.h
@@ -2,5 +2,5 @@
 
 #define SC_MAJOR_VERSION 1
 #define SC_MINOR_VERSION 0
-#define SC_PATCH_VERSION 0
-#define SC_VERSION_STRING "1.0.0"
+#define SC_PATCH_VERSION 1
+#define SC_VERSION_STRING "1.0.1"
diff --git a/Engine/Plugins/Animation/RigLogic/Source/RigLogicLib/Public/terse/archives/binary/InputArchive.h b/Engine/Plugins/Animation/RigLogic/Source/RigLogicLib/Public/terse/archives/binary/InputArchive.h
--- a/Engine/Plugins/Animation/RigLogic/Source/RigLogicLib/Public/terse/archives/binary/InputArchive.h
+++ b/Engine/Plugins/Animation/RigLogic/Source/RigLogicLib/Public/terse/archives/binary/InputArchive.h
@@ -64,7 +64,7 @@ struct ValueFactory {
         const ParentAllocator& alloc) {
         using K = typename T::first_type;
         using V = typename T::second_type;
-        return {ValueFactory<K>::create(alloc), ValueFactory<V>::create(alloc)};
+        return T{ValueFactory<K>::create(alloc), ValueFactory<V>::create(alloc)};
     }
 
     template<class ParentAllocator, bool IsPrimitive = IsPrimitive>
@@ -72,7 +72,7 @@ struct ValueFactory {
         const ParentAllocator& alloc) {
         using K = typename std::tuple_element<0, T>::type;
         using V = typename std::tuple_element<0, T>::type;
-        return {ValueFactory<K>::create(alloc), ValueFactory<V>::create(alloc)};
+        return T{ValueFactory<K>::create(alloc), ValueFactory<V>::create(alloc)};
     }
 
 };
diff --git a/Engine/Plugins/Animation/RigLogic/Source/RigLogicLib/Public/terse/version/Version.h b/Engine/Plugins/Animation/RigLogic/Source/RigLogicLib/Public/terse/version/Version.h
--- a/Engine/Plugins/Animation/RigLogic/Source/RigLogicLib/Public/terse/version/Version.h
+++ b/Engine/Plugins/Animation/RigLogic/Source/RigLogicLib/Public/terse/version/Version.h
@@ -2,5 +2,5 @@
 
 #define TERSE_MAJOR_VERSION 1
 #define TERSE_MINOR_VERSION 1
-#define TERSE_PATCH_VERSION 0
-#define TERSE_VERSION_STRING "1.1.0"
+#define TERSE_PATCH_VERSION 1
+#define TERSE_VERSION_STRING "1.1.1"
diff --git a/Engine/Plugins/Animation/RigLogic/Source/RigLogicLib/Public/trimd/AVX.h b/Engine/Plugins/Animation/RigLogic/Source/RigLogicLib/Public/trimd/AVX.h
new file mode 100644
--- /dev/null
+++ b/Engine/Plugins/Animation/RigLogic/Source/RigLogicLib/Public/trimd/AVX.h
@@ -0,0 +1,184 @@
+#pragma once
+
+#ifdef TRIMD_ENABLE_AVX
+#include "trimd/Polyfill.h"
+
+#include <immintrin.h>
+
+namespace trimd {
+
+namespace avx {
+
+class F256 {
+    public:
+        F256() : data{_mm256_setzero_ps()} {
+        }
+
+        explicit F256(__m256 value) : data{value} {
+        }
+
+        explicit F256(float value) : F256{_mm256_set1_ps(value)} {
+        }
+
+        F256(float v1, float v2, float v3, float v4, float v5, float v6, float v7, float v8) :
+            data{_mm256_set_ps(v8, v7, v6, v5, v4, v3, v2, v1)} {
+        }
+
+        static F256 fromAlignedSource(const float* source) {
+            return F256{_mm256_load_ps(source)};
+        }
+
+        static F256 fromUnalignedSource(const float* source) {
+            return F256{_mm256_loadu_ps(source)};
+        }
+
+        static F256 loadSingleValue(const float* source) {
+            const __m256i mask = _mm256_set_epi32(-1, 0, 0, 0, 0, 0, 0, 0);
+            return F256{_mm256_maskload_ps(source, mask)};
+        }
+
+        #ifdef TRIMD_ENABLE_F16C
+            static F256 fromAlignedSource(const std::uint16_t* source) {
+                __m128i halfFloats = _mm_load_si128(reinterpret_cast<const __m128i*>(source));
+                return F256{_mm256_cvtph_ps(halfFloats)};
+            }
+
+            static F256 fromUnalignedSource(const std::uint16_t* source) {
+                __m128i halfFloats = _mm_loadu_si128(reinterpret_cast<const __m128i*>(source));
+                return F256{_mm256_cvtph_ps(halfFloats)};
+            }
+
+            static F256 loadSingleValue(const std::uint16_t* source) {
+                __m128i halfFloats = _mm_loadu_si16(source);
+                return F256{_mm256_cvtph_ps(halfFloats)};
+            }
+
+        #endif  // TRIMD_ENABLE_F16C
+
+        template<typename T>
+        static void prefetchT0(const T* source) {
+            #if defined(__clang__) || defined(__GNUC__)
+                #pragma GCC diagnostic push
+                #pragma GCC diagnostic ignored "-Wold-style-cast"
+            #endif
+            _mm_prefetch(reinterpret_cast<const char*>(source), _MM_HINT_T0);
+            #if defined(__clang__) || defined(__GNUC__)
+                #pragma GCC diagnostic pop
+            #endif
+        }
+
+        template<typename T>
+        static void prefetchT1(const T* source) {
+            #if defined(__clang__) || defined(__GNUC__)
+                #pragma GCC diagnostic push
+                #pragma GCC diagnostic ignored "-Wold-style-cast"
+            #endif
+            _mm_prefetch(reinterpret_cast<const char*>(source), _MM_HINT_T1);
+            #if defined(__clang__) || defined(__GNUC__)
+                #pragma GCC diagnostic pop
+            #endif
+        }
+
+        template<typename T>
+        static void prefetchT2(const T* source) {
+            #if defined(__clang__) || defined(__GNUC__)
+                #pragma GCC diagnostic push
+                #pragma GCC diagnostic ignored "-Wold-style-cast"
+            #endif
+            _mm_prefetch(reinterpret_cast<const char*>(source), _MM_HINT_T2);
+            #if defined(__clang__) || defined(__GNUC__)
+                #pragma GCC diagnostic pop
+            #endif
+        }
+
+        template<typename T>
+        static void prefetchNTA(const T* source) {
+            #if defined(__clang__) || defined(__GNUC__)
+                #pragma GCC diagnostic push
+                #pragma GCC diagnostic ignored "-Wold-style-cast"
+            #endif
+            _mm_prefetch(reinterpret_cast<const char*>(source), _MM_HINT_NTA);
+            #if defined(__clang__) || defined(__GNUC__)
+                #pragma GCC diagnostic pop
+            #endif
+        }
+
+        void alignedLoad(const float* source) {
+            data = _mm256_load_ps(source);
+        }
+
+        void unalignedLoad(const float* source) {
+            data = _mm256_loadu_ps(source);
+        }
+
+        void alignedStore(float* dest) const {
+            _mm256_store_ps(dest, data);
+        }
+
+        void unalignedStore(float* dest) const {
+            _mm256_storeu_ps(dest, data);
+        }
+
+        float sum() const {
+            // (data[3] + data[7], data[2] + data[6], data[1] + data[5], data[0] + data[4])
+            const __m128 x128 = _mm_add_ps(_mm256_extractf128_ps(data, 1), _mm256_castps256_ps128(data));
+            // (-, -, data[1] + data[3] + data[5] + data[7], data[0] + data[2] + data[4] + data[6])
+            const __m128 x64 = _mm_add_ps(x128, _mm_movehl_ps(x128, x128));
+            // (-, -, -, data[0] + data[1] + data[2] + data[3] + data[4] + data[5] + data[6] + data[7])
+            const __m128 x32 = _mm_add_ss(x64, _mm_shuffle_ps(x64, x64, 0x55));
+            return _mm_cvtss_f32(x32);
+        }
+
+        F256& operator+=(const F256& rhs) {
+            data = _mm256_add_ps(data, rhs.data);
+            return *this;
+        }
+
+        F256& operator-=(const F256& rhs) {
+            data = _mm256_sub_ps(data, rhs.data);
+            return *this;
+        }
+
+        F256& operator*=(const F256& rhs) {
+            data = _mm256_mul_ps(data, rhs.data);
+            return *this;
+        }
+
+        F256& operator/=(const F256& rhs) {
+            data = _mm256_div_ps(data, rhs.data);
+            return *this;
+        }
+
+        static constexpr std::size_t size() {
+            return sizeof(decltype(data)) / sizeof(float);
+        }
+
+        static constexpr std::size_t alignment() {
+            return alignof(decltype(data));
+        }
+
+    private:
+        __m256 data;
+};
+
+inline F256 operator+(const F256& lhs, const F256& rhs) {
+    return F256(lhs) += rhs;
+}
+
+inline F256 operator-(const F256& lhs, const F256& rhs) {
+    return F256(lhs) -= rhs;
+}
+
+inline F256 operator*(const F256& lhs, const F256& rhs) {
+    return F256(lhs) *= rhs;
+}
+
+inline F256 operator/(const F256& lhs, const F256& rhs) {
+    return F256(lhs) /= rhs;
+}
+
+}  // namespace avx
+
+}  // namespace trimd
+
+#endif  // TRIMD_ENABLE_AVX
diff --git a/Engine/Plugins/Animation/RigLogic/Source/RigLogicLib/Public/trimd/Polyfill.h b/Engine/Plugins/Animation/RigLogic/Source/RigLogicLib/Public/trimd/Polyfill.h
new file mode 100644
--- /dev/null
+++ b/Engine/Plugins/Animation/RigLogic/Source/RigLogicLib/Public/trimd/Polyfill.h
@@ -0,0 +1,21 @@
+#pragma once
+
+#if defined(TRIMD_ENABLE_AVX) || defined(TRIMD_ENABLE_SSE)
+#include <immintrin.h>
+
+#include <cstdint>
+
+#if (!defined(__clang__) && defined(_MSC_VER) && _MSC_VER <= 1900) || (!defined(__clang__) && defined(__GNUC__) && __GNUC__ < 9)
+    #define _mm_loadu_si64 _mm_loadl_epi64
+#endif
+
+#if (!defined(__clang__) && defined(_MSC_VER) && _MSC_VER <= 1900) || \
+    (defined(__clang__) && __clang_major__ < 8) || \
+    (!defined(__clang__) && defined(__GNUC__))
+    inline __m128i _mm_loadu_si16(const void* source) {
+        return _mm_insert_epi16(_mm_setzero_si128(), *reinterpret_cast<const std::int16_t*>(source), 0);
+    }
+
+#endif
+
+#endif  // defined(TRIMD_ENABLE_AVX) || defined(TRIMD_ENABLE_SSE)
diff --git a/Engine/Plugins/Animation/RigLogic/Source/RigLogicLib/Public/trimd/SSE.h b/Engine/Plugins/Animation/RigLogic/Source/RigLogicLib/Public/trimd/SSE.h
new file mode 100644
--- /dev/null
+++ b/Engine/Plugins/Animation/RigLogic/Source/RigLogicLib/Public/trimd/SSE.h
@@ -0,0 +1,180 @@
+#pragma once
+
+#ifdef TRIMD_ENABLE_SSE
+#include "trimd/Polyfill.h"
+
+#include <immintrin.h>
+
+namespace trimd {
+
+namespace sse {
+
+class F128 {
+    public:
+        F128() : data{_mm_setzero_ps()} {
+        }
+
+        explicit F128(__m128 value) : data{value} {
+        }
+
+        explicit F128(float value) : F128{_mm_set1_ps(value)} {
+        }
+
+        F128(float v1, float v2, float v3, float v4) : data{_mm_set_ps(v4, v3, v2, v1)} {
+        }
+
+        static F128 fromAlignedSource(const float* source) {
+            return F128{_mm_load_ps(source)};
+        }
+
+        static F128 fromUnalignedSource(const float* source) {
+            return F128{_mm_loadu_ps(source)};
+        }
+
+        static F128 loadSingleValue(const float* source) {
+            return F128{_mm_load_ss(source)};
+        }
+
+        #ifdef TRIMD_ENABLE_F16C
+            static F128 fromAlignedSource(const std::uint16_t* source) {
+                __m128i halfFloats = _mm_loadl_epi64(reinterpret_cast<const __m128i*>(source));
+                return F128{_mm_cvtph_ps(halfFloats)};
+            }
+
+            static F128 fromUnalignedSource(const std::uint16_t* source) {
+                __m128i halfFloats = _mm_loadu_si64(reinterpret_cast<const __m128i*>(source));
+                return F128{_mm_cvtph_ps(halfFloats)};
+            }
+
+            static F128 loadSingleValue(const std::uint16_t* source) {
+                __m128i halfFloats = _mm_loadu_si16(source);
+                return F128{_mm_cvtph_ps(halfFloats)};
+            }
+
+        #endif  // TRIMD_ENABLE_F16C
+
+        template<typename T>
+        static void prefetchT0(const T* source) {
+            #if defined(__clang__) || defined(__GNUC__)
+                #pragma GCC diagnostic push
+                #pragma GCC diagnostic ignored "-Wold-style-cast"
+            #endif
+            _mm_prefetch(reinterpret_cast<const char*>(source), _MM_HINT_T0);
+            #if defined(__clang__) || defined(__GNUC__)
+                #pragma GCC diagnostic pop
+            #endif
+        }
+
+        template<typename T>
+        static void prefetchT1(const T* source) {
+            #if defined(__clang__) || defined(__GNUC__)
+                #pragma GCC diagnostic push
+                #pragma GCC diagnostic ignored "-Wold-style-cast"
+            #endif
+            _mm_prefetch(reinterpret_cast<const char*>(source), _MM_HINT_T1);
+            #if defined(__clang__) || defined(__GNUC__)
+                #pragma GCC diagnostic pop
+            #endif
+        }
+
+        template<typename T>
+        static void prefetchT2(const T* source) {
+            #if defined(__clang__) || defined(__GNUC__)
+                #pragma GCC diagnostic push
+                #pragma GCC diagnostic ignored "-Wold-style-cast"
+            #endif
+            _mm_prefetch(reinterpret_cast<const char*>(source), _MM_HINT_T2);
+            #if defined(__clang__) || defined(__GNUC__)
+                #pragma GCC diagnostic pop
+            #endif
+        }
+
+        template<typename T>
+        static void prefetchNTA(const T* source) {
+            #if defined(__clang__) || defined(__GNUC__)
+                #pragma GCC diagnostic push
+                #pragma GCC diagnostic ignored "-Wold-style-cast"
+            #endif
+            _mm_prefetch(reinterpret_cast<const char*>(source), _MM_HINT_NTA);
+            #if defined(__clang__) || defined(__GNUC__)
+                #pragma GCC diagnostic pop
+            #endif
+        }
+
+        void alignedLoad(const float* source) {
+            data = _mm_load_ps(source);
+        }
+
+        void unalignedLoad(const float* source) {
+            data = _mm_loadu_ps(source);
+        }
+
+        void alignedStore(float* dest) const {
+            _mm_store_ps(dest, data);
+        }
+
+        void unalignedStore(float* dest) const {
+            _mm_storeu_ps(dest, data);
+        }
+
+        float sum() const {
+            __m128 temp = _mm_movehdup_ps(data);
+            __m128 result = _mm_add_ps(data, temp);
+            temp = _mm_movehl_ps(temp, result);
+            result = _mm_add_ss(result, temp);
+            return _mm_cvtss_f32(result);
+        }
+
+        F128& operator+=(const F128& rhs) {
+            data = _mm_add_ps(data, rhs.data);
+            return *this;
+        }
+
+        F128& operator-=(const F128& rhs) {
+            data = _mm_sub_ps(data, rhs.data);
+            return *this;
+        }
+
+        F128& operator*=(const F128& rhs) {
+            data = _mm_mul_ps(data, rhs.data);
+            return *this;
+        }
+
+        F128& operator/=(const F128& rhs) {
+            data = _mm_div_ps(data, rhs.data);
+            return *this;
+        }
+
+        static constexpr std::size_t size() {
+            return sizeof(decltype(data)) / sizeof(float);
+        }
+
+        static constexpr std::size_t alignment() {
+            return alignof(decltype(data));
+        }
+
+    private:
+        __m128 data;
+};
+
+inline F128 operator+(const F128& lhs, const F128& rhs) {
+    return F128(lhs) += rhs;
+}
+
+inline F128 operator-(const F128& lhs, const F128& rhs) {
+    return F128(lhs) -= rhs;
+}
+
+inline F128 operator*(const F128& lhs, const F128& rhs) {
+    return F128(lhs) *= rhs;
+}
+
+inline F128 operator/(const F128& lhs, const F128& rhs) {
+    return F128(lhs) /= rhs;
+}
+
+}  // namespace sse
+
+}  // namespace trimd
+
+#endif  // TRIMD_ENABLE_SSE
diff --git a/Engine/Plugins/Animation/RigLogic/Source/RigLogicLib/Public/trimd/Scalar.h b/Engine/Plugins/Animation/RigLogic/Source/RigLogicLib/Public/trimd/Scalar.h
new file mode 100644
--- /dev/null
+++ b/Engine/Plugins/Animation/RigLogic/Source/RigLogicLib/Public/trimd/Scalar.h
@@ -0,0 +1,322 @@
+#pragma once
+
+#include <array>
+#include <cstddef>
+
+namespace trimd {
+
+namespace scalar {
+
+template<typename T>
+class T128 {
+    public:
+        using value_type = typename std::remove_cv<T>::type;
+
+    public:
+        T128() : data{} {
+        }
+
+        T128(value_type v1, value_type v2, value_type v3, value_type v4) : data({v1, v2, v3, v4}) {
+        }
+
+        explicit T128(value_type value) : T128(value, value, value, value) {
+        }
+
+        static T128 fromAlignedSource(const value_type* source) {
+            return T128{source[0], source[1], source[2], source[3]};
+        }
+
+        static T128 fromUnalignedSource(const value_type* source) {
+            return T128::fromAlignedSource(source);
+        }
+
+        static T128 loadSingleValue(const value_type* source) {
+            return T128{source[0], value_type{}, value_type{}, value_type{}};
+        }
+
+        template<typename U>
+        static void prefetchT0(const U*  /*unused*/) {
+            // Intentionally noop
+        }
+
+        template<typename U>
+        static void prefetchT1(const U*  /*unused*/) {
+            // Intentionally noop
+        }
+
+        template<typename U>
+        static void prefetchT2(const U*  /*unused*/) {
+            // Intentionally noop
+        }
+
+        template<typename U>
+        static void prefetchNTA(const U*  /*unused*/) {
+            // Intentionally noop
+        }
+
+        void alignedLoad(const value_type* source) {
+            data[0] = source[0];
+            data[1] = source[1];
+            data[2] = source[2];
+            data[3] = source[3];
+        }
+
+        void unalignedLoad(const value_type* source) {
+            alignedLoad(source);
+        }
+
+        void alignedStore(value_type* dest) const {
+            dest[0] = data[0];
+            dest[1] = data[1];
+            dest[2] = data[2];
+            dest[3] = data[3];
+        }
+
+        void unalignedStore(value_type* dest) const {
+            alignedStore(dest);
+        }
+
+        float sum() const {
+            return data[0] + data[1] + data[2] + data[3];
+        }
+
+        T128& operator+=(const T128& rhs) {
+            data[0] += rhs.data[0];
+            data[1] += rhs.data[1];
+            data[2] += rhs.data[2];
+            data[3] += rhs.data[3];
+            return *this;
+        }
+
+        T128& operator-=(const T128& rhs) {
+            data[0] -= rhs.data[0];
+            data[1] -= rhs.data[1];
+            data[2] -= rhs.data[2];
+            data[3] -= rhs.data[3];
+            return *this;
+        }
+
+        T128& operator*=(const T128& rhs) {
+            data[0] *= rhs.data[0];
+            data[1] *= rhs.data[1];
+            data[2] *= rhs.data[2];
+            data[3] *= rhs.data[3];
+            return *this;
+        }
+
+        T128& operator/=(const T128& rhs) {
+            data[0] /= rhs.data[0];
+            data[1] /= rhs.data[1];
+            data[2] /= rhs.data[2];
+            data[3] /= rhs.data[3];
+            return *this;
+        }
+
+        static constexpr std::size_t size() {
+            return sizeof(decltype(data)) / sizeof(float);
+        }
+
+        static constexpr std::size_t alignment() {
+            #if defined(__arm__) || defined(__aarch64__) || defined(_M_ARM) || defined(_M_ARM64)
+                return std::alignment_of<std::max_align_t>::value;
+            #else
+                return sizeof(decltype(data));
+            #endif
+        }
+
+    private:
+        std::array<value_type, 4> data;
+};
+
+template<typename T>
+inline T128<T> operator+(const T128<T>& lhs, const T128<T>& rhs) {
+    return T128<T>(lhs) += rhs;
+}
+
+template<typename T>
+inline T128<T> operator-(const T128<T>& lhs, const T128<T>& rhs) {
+    return T128<T>(lhs) -= rhs;
+}
+
+template<typename T>
+inline T128<T> operator*(const T128<T>& lhs, const T128<T>& rhs) {
+    return T128<T>(lhs) *= rhs;
+}
+
+template<typename T>
+inline T128<T> operator/(const T128<T>& lhs, const T128<T>& rhs) {
+    return T128<T>(lhs) /= rhs;
+}
+
+template<typename T>
+class T256 {
+    public:
+        using value_type = typename std::remove_cv<T>::type;
+
+    public:
+        T256() : data{} {
+        }
+
+        T256(value_type v1, value_type v2, value_type v3, value_type v4,
+             value_type v5, value_type v6, value_type v7, value_type v8) : data({v1, v2, v3, v4, v5, v6, v7, v8}) {
+        }
+
+        explicit T256(value_type value) : T256(value, value, value, value, value, value, value, value) {
+        }
+
+        static T256 fromAlignedSource(const value_type* source) {
+            return T256{source[0], source[1], source[2], source[3], source[4], source[5], source[6], source[7]};
+        }
+
+        static T256 fromUnalignedSource(const value_type* source) {
+            return T256::fromAlignedSource(source);
+        }
+
+        static T256 loadSingleValue(const value_type* source) {
+            return T256{source[0], value_type{}, value_type{}, value_type{}, value_type{}, value_type{}, value_type{}, value_type{}};
+        }
+
+        template<typename U>
+        static void prefetchT0(const U*  /*unused*/) {
+            // Intentionally noop
+        }
+
+        template<typename U>
+        static void prefetchT1(const U*  /*unused*/) {
+            // Intentionally noop
+        }
+
+        template<typename U>
+        static void prefetchT2(const U*  /*unused*/) {
+            // Intentionally noop
+        }
+
+        template<typename U>
+        static void prefetchNTA(const U*  /*unused*/) {
+            // Intentionally noop
+        }
+
+        void alignedLoad(const value_type* source) {
+            data[0] = source[0];
+            data[1] = source[1];
+            data[2] = source[2];
+            data[3] = source[3];
+            data[4] = source[4];
+            data[5] = source[5];
+            data[6] = source[6];
+            data[7] = source[7];
+        }
+
+        void unalignedLoad(const value_type* source) {
+            alignedLoad(source);
+        }
+
+        void alignedStore(value_type* dest) const {
+            dest[0] = data[0];
+            dest[1] = data[1];
+            dest[2] = data[2];
+            dest[3] = data[3];
+            dest[4] = data[4];
+            dest[5] = data[5];
+            dest[6] = data[6];
+            dest[7] = data[7];
+        }
+
+        void unalignedStore(value_type* dest) const {
+            alignedStore(dest);
+        }
+
+        float sum() const {
+            return data[0] + data[1] + data[2] + data[3] + data[4] + data[5] + data[6] + data[7];
+        }
+
+        T256& operator+=(const T256& rhs) {
+            data[0] += rhs.data[0];
+            data[1] += rhs.data[1];
+            data[2] += rhs.data[2];
+            data[3] += rhs.data[3];
+            data[4] += rhs.data[4];
+            data[5] += rhs.data[5];
+            data[6] += rhs.data[6];
+            data[7] += rhs.data[7];
+            return *this;
+        }
+
+        T256& operator-=(const T256& rhs) {
+            data[0] -= rhs.data[0];
+            data[1] -= rhs.data[1];
+            data[2] -= rhs.data[2];
+            data[3] -= rhs.data[3];
+            data[4] -= rhs.data[4];
+            data[5] -= rhs.data[5];
+            data[6] -= rhs.data[6];
+            data[7] -= rhs.data[7];
+            return *this;
+        }
+
+        T256& operator*=(const T256& rhs) {
+            data[0] *= rhs.data[0];
+            data[1] *= rhs.data[1];
+            data[2] *= rhs.data[2];
+            data[3] *= rhs.data[3];
+            data[4] *= rhs.data[4];
+            data[5] *= rhs.data[5];
+            data[6] *= rhs.data[6];
+            data[7] *= rhs.data[7];
+            return *this;
+        }
+
+        T256& operator/=(const T256& rhs) {
+            data[0] /= rhs.data[0];
+            data[1] /= rhs.data[1];
+            data[2] /= rhs.data[2];
+            data[3] /= rhs.data[3];
+            data[4] /= rhs.data[4];
+            data[5] /= rhs.data[5];
+            data[6] /= rhs.data[6];
+            data[7] /= rhs.data[7];
+            return *this;
+        }
+
+        static constexpr std::size_t size() {
+            return sizeof(decltype(data)) / sizeof(float);
+        }
+
+        static constexpr std::size_t alignment() {
+            #if defined(__arm__) || defined(__aarch64__) || defined(_M_ARM) || defined(_M_ARM64)
+                return std::alignment_of<std::max_align_t>::value;
+            #else
+                return sizeof(decltype(data));
+            #endif
+        }
+
+    private:
+        std::array<value_type, 8> data;
+};
+
+template<typename T>
+inline T256<T> operator+(const T256<T>& lhs, const T256<T>& rhs) {
+    return T256<T>(lhs) += rhs;
+}
+
+template<typename T>
+inline T256<T> operator-(const T256<T>& lhs, const T256<T>& rhs) {
+    return T256<T>(lhs) -= rhs;
+}
+
+template<typename T>
+inline T256<T> operator*(const T256<T>& lhs, const T256<T>& rhs) {
+    return T256<T>(lhs) *= rhs;
+}
+
+template<typename T>
+inline T256<T> operator/(const T256<T>& lhs, const T256<T>& rhs) {
+    return T256<T>(lhs) /= rhs;
+}
+
+using F128 = T128<float>;
+using F256 = T256<float>;
+
+}  // namespace scalar
+
+}  // namespace trimd
diff --git a/Engine/Plugins/Animation/RigLogic/Source/RigLogicLib/Public/trimd/TRiMD.h b/Engine/Plugins/Animation/RigLogic/Source/RigLogicLib/Public/trimd/TRiMD.h
new file mode 100644
--- /dev/null
+++ b/Engine/Plugins/Animation/RigLogic/Source/RigLogicLib/Public/trimd/TRiMD.h
@@ -0,0 +1,5 @@
+#pragma once
+
+#include "trimd/AVX.h"
+#include "trimd/SSE.h"
+#include "trimd/Scalar.h"
diff --git a/Engine/Plugins/Animation/RigLogic/Source/RigLogicLib/Public/trimd/version/Version.h b/Engine/Plugins/Animation/RigLogic/Source/RigLogicLib/Public/trimd/version/Version.h
new file mode 100644
--- /dev/null
+++ b/Engine/Plugins/Animation/RigLogic/Source/RigLogicLib/Public/trimd/version/Version.h
@@ -0,0 +1,6 @@
+#pragma once
+
+#define TRIMD_MAJOR_VERSION 1
+#define TRIMD_MINOR_VERSION 0
+#define TRIMD_PATCH_VERSION 0
+#define TRIMD_VERSION_STRING "1.0.0"
diff --git a/Engine/Plugins/Animation/RigLogic/Source/RigLogicLib/Public/trio/version/Version.h b/Engine/Plugins/Animation/RigLogic/Source/RigLogicLib/Public/trio/version/Version.h
--- a/Engine/Plugins/Animation/RigLogic/Source/RigLogicLib/Public/trio/version/Version.h
+++ b/Engine/Plugins/Animation/RigLogic/Source/RigLogicLib/Public/trio/version/Version.h
@@ -2,5 +2,5 @@
 
 #define TRIO_MAJOR_VERSION 1
 #define TRIO_MINOR_VERSION 0
-#define TRIO_PATCH_VERSION 1
-#define TRIO_VERSION_STRING "1.0.1"
+#define TRIO_PATCH_VERSION 2
+#define TRIO_VERSION_STRING "1.0.2"
diff --git a/Engine/Plugins/Animation/RigLogic/Source/RigLogicLib/RigLogicLib.Build.cs b/Engine/Plugins/Animation/RigLogic/Source/RigLogicLib/RigLogicLib.Build.cs
--- a/Engine/Plugins/Animation/RigLogic/Source/RigLogicLib/RigLogicLib.Build.cs
+++ b/Engine/Plugins/Animation/RigLogic/Source/RigLogicLib/RigLogicLib.Build.cs
@@ -1,7 +1,8 @@
 // Copyright 1998-2019 Epic Games, Inc. All Rights Reserved.
 
-using UnrealBuildTool;
+using System;
 using System.IO;
+using UnrealBuildTool;
 
 public class RigLogicLib : ModuleRules
 {
@@ -29,7 +30,7 @@ public class RigLogicLib : ModuleRules
 
         Type = ModuleType.CPlusPlus;
 
-        PrivateDefinitions.Add("RL_BUILD_SHARED=1");  // used instead of #define so cpp files can do conditional compilation
+        PrivateDefinitions.Add("RL_BUILD_SHARED=1");
 
         if (Target.Platform == UnrealTargetPlatform.Win32 ||
                 Target.Platform == UnrealTargetPlatform.Win64 ||
diff --git a/Engine/Plugins/Animation/RigLogic/Source/RigLogicLibTest/Private/dnatests/Defs.h b/Engine/Plugins/Animation/RigLogic/Source/RigLogicLibTest/Private/dnatests/Defs.h
--- a/Engine/Plugins/Animation/RigLogic/Source/RigLogicLibTest/Private/dnatests/Defs.h
+++ b/Engine/Plugins/Animation/RigLogic/Source/RigLogicLibTest/Private/dnatests/Defs.h
@@ -1,24 +1,21 @@
 #pragma once
 
 #ifdef _MSC_VER
-    #pragma warning(disable:4668)
-    #pragma warning(disable:4996)
     #pragma warning(push)
-    #pragma warning(disable : 4365 4987)
+    #pragma warning(disable : 4365 4987 4668 4996)
 #endif
 #ifdef __clang__
     #pragma clang diagnostic push
-    #pragma clang diagnostic ignored "-Wglobal-constructors"
-	#pragma clang diagnostic ignored "-Wundef"
+    #pragma clang diagnostic ignored "-Wundef"
 #endif
 #include <gtest/gtest.h>
 #ifdef __clang__
     #pragma clang diagnostic pop
+    // This cannot be applied to the scope of the header inclusion only, as the TEST macros will trigger it
+    #pragma clang diagnostic ignored "-Wglobal-constructors"
 #endif
 #ifdef _MSC_VER
     #pragma warning(pop)
-    #pragma warning(default:4668)
-    #pragma warning(default:4996)
 #endif
 
 #ifndef INSTANTIATE_TEST_SUITE_P
diff --git a/Engine/Plugins/Animation/RigLogic/Source/RigLogicLibTest/Private/pmatests/Defs.h b/Engine/Plugins/Animation/RigLogic/Source/RigLogicLibTest/Private/pmatests/Defs.h
--- a/Engine/Plugins/Animation/RigLogic/Source/RigLogicLibTest/Private/pmatests/Defs.h
+++ b/Engine/Plugins/Animation/RigLogic/Source/RigLogicLibTest/Private/pmatests/Defs.h
@@ -1,23 +1,21 @@
 #pragma once
 
 #ifdef _MSC_VER
-    #pragma warning(disable:4668)
-    #pragma warning(disable:4996)
     #pragma warning(push)
-    #pragma warning(disable : 4365 4987)
+    #pragma warning(disable : 4365 4987 4668 4996)
 #endif
 #ifdef __clang__
     #pragma clang diagnostic push
-    #pragma clang diagnostic ignored "-Wglobal-constructors"
+    #pragma clang diagnostic ignored "-Wundef"
 #endif
 #include <gtest/gtest.h>
 #ifdef __clang__
     #pragma clang diagnostic pop
+    // This cannot be applied to the scope of the header inclusion only, as the TEST macros will trigger it
+    #pragma clang diagnostic ignored "-Wglobal-constructors"
 #endif
 #ifdef _MSC_VER
     #pragma warning(pop)
-    #pragma warning(default:4668)
-    #pragma warning(default:4996)
 #endif
 
 #ifndef INSTANTIATE_TEST_SUITE_P
diff --git a/Engine/Plugins/Animation/RigLogic/Source/RigLogicLibTest/Private/rltests/Defs.h b/Engine/Plugins/Animation/RigLogic/Source/RigLogicLibTest/Private/rltests/Defs.h
--- a/Engine/Plugins/Animation/RigLogic/Source/RigLogicLibTest/Private/rltests/Defs.h
+++ b/Engine/Plugins/Animation/RigLogic/Source/RigLogicLibTest/Private/rltests/Defs.h
@@ -1,23 +1,21 @@
 #pragma once
 
 #ifdef _MSC_VER
-    #pragma warning(disable:4668)
-    #pragma warning(disable:4996)
     #pragma warning(push)
-    #pragma warning(disable : 4365 4987)
+    #pragma warning(disable : 4365 4987 4668 4996)
 #endif
 #ifdef __clang__
     #pragma clang diagnostic push
-    #pragma clang diagnostic ignored "-Wglobal-constructors"
+    #pragma clang diagnostic ignored "-Wundef"
 #endif
 #include <gtest/gtest.h>
 #ifdef __clang__
     #pragma clang diagnostic pop
+    // This cannot be applied to the scope of the header inclusion only, as the TEST macros will trigger it
+    #pragma clang diagnostic ignored "-Wglobal-constructors"
 #endif
 #ifdef _MSC_VER
     #pragma warning(pop)
-    #pragma warning(default:4668)
-    #pragma warning(default:4996)
 #endif
 
 #ifndef INSTANTIATE_TEST_SUITE_P
diff --git a/Engine/Plugins/Animation/RigLogic/Source/RigLogicLibTest/Private/tersetests/Defs.h b/Engine/Plugins/Animation/RigLogic/Source/RigLogicLibTest/Private/tersetests/Defs.h
--- a/Engine/Plugins/Animation/RigLogic/Source/RigLogicLibTest/Private/tersetests/Defs.h
+++ b/Engine/Plugins/Animation/RigLogic/Source/RigLogicLibTest/Private/tersetests/Defs.h
@@ -1,23 +1,21 @@
 #pragma once
 
 #ifdef _MSC_VER
-    #pragma warning(disable:4668)
-    #pragma warning(disable:4996)
     #pragma warning(push)
-    #pragma warning(disable : 4365 4987)
+    #pragma warning(disable : 4365 4987 4668 4996)
 #endif
 #ifdef __clang__
     #pragma clang diagnostic push
-    #pragma clang diagnostic ignored "-Wglobal-constructors"
+    #pragma clang diagnostic ignored "-Wundef"
 #endif
 #include <gtest/gtest.h>
 #ifdef __clang__
     #pragma clang diagnostic pop
+    // This cannot be applied to the scope of the header inclusion only, as the TEST macros will trigger it
+    #pragma clang diagnostic ignored "-Wglobal-constructors"
 #endif
 #ifdef _MSC_VER
     #pragma warning(pop)
-    #pragma warning(default:4668)
-    #pragma warning(default:4996)
 #endif
 
 #ifndef INSTANTIATE_TEST_SUITE_P
diff --git a/Engine/Plugins/Animation/RigLogic/Source/RigLogicModule/Private/RigUnit_RigLogic.cpp b/Engine/Plugins/Animation/RigLogic/Source/RigLogicModule/Private/RigUnit_RigLogic.cpp
--- a/Engine/Plugins/Animation/RigLogic/Source/RigLogicModule/Private/RigUnit_RigLogic.cpp
+++ b/Engine/Plugins/Animation/RigLogic/Source/RigLogicModule/Private/RigUnit_RigLogic.cpp
@@ -295,11 +295,6 @@ void FRigUnit_RigLogic_Data::UpdateAnimMapCurves(FRigCurveContainer* CurveContai
 	}	
 }
 
-FString FRigUnit_RigLogic::GetUnitLabel() const
-{
-	return FString::Printf(TEXT("RigUnit"));
-}
-
 IBehaviorReader* FRigUnit_RigLogic_Data::FetchBehaviorReaderFromOwner() 
 {
 	if (SkelMeshComponent.IsValid())
diff --git a/Engine/Plugins/Animation/RigLogic/Source/RigLogicModule/Public/RigUnit_RigLogic.h b/Engine/Plugins/Animation/RigLogic/Source/RigLogicModule/Public/RigUnit_RigLogic.h
--- a/Engine/Plugins/Animation/RigLogic/Source/RigLogicModule/Public/RigUnit_RigLogic.h
+++ b/Engine/Plugins/Animation/RigLogic/Source/RigLogicModule/Public/RigUnit_RigLogic.h
@@ -171,7 +171,6 @@ struct FRigUnit_RigLogic: public FRigUnitMutable
 	{
 	}
 
-	virtual FString GetUnitLabel() const override;
 	RIGVM_METHOD()
 	virtual void Execute(const FRigUnitContext& Context) override;
 
diff --git a/Engine/Plugins/Animation/RigLogic/Source/RigLogicModule/RigLogicModule.Build.cs b/Engine/Plugins/Animation/RigLogic/Source/RigLogicModule/RigLogicModule.Build.cs
--- a/Engine/Plugins/Animation/RigLogic/Source/RigLogicModule/RigLogicModule.Build.cs
+++ b/Engine/Plugins/Animation/RigLogic/Source/RigLogicModule/RigLogicModule.Build.cs
@@ -1,6 +1,8 @@
- // Copyright 1998-2019 Epic Games, Inc. All Rights Reserved.
+// Copyright 1998-2019 Epic Games, Inc. All Rights Reserved.
 
+using System;
 using System.IO;
+using UnrealBuildTool;
 
 namespace UnrealBuildTool.Rules
 {
@@ -13,8 +15,7 @@ namespace UnrealBuildTool.Rules
 
         public RigLogicModule(ReadOnlyTargetRules Target) : base(Target)
         {
-            PublicDefinitions.Add("RL_SHARED=1"); //used instead of #define so cpp files can do conditional compilation
-
+            PublicDefinitions.Add("RL_SHARED=1");
 
             PublicDependencyModuleNames.AddRange(
                 new string[]
@@ -22,13 +23,13 @@ namespace UnrealBuildTool.Rules
                     "Core",
                     "CoreUObject",
                     "Engine",
-					"RigVM",
+                    "RigVM",
                     "ControlRig",
                     "RigLogicLib",
                     "MessageLog",
                     "RigVM",
-					"Projects"
-                } 
+                    "Projects"
+                }
                 );
 
             if (Target.Type == TargetType.Editor)
@@ -39,11 +40,11 @@ namespace UnrealBuildTool.Rules
             PrivateDependencyModuleNames.AddRange(
                 new string[]
                 {
-					"AnimationCore",
+                    "AnimationCore",
                     "ControlRig",
-					"RenderCore",
-					"RHI"
-				}
+                    "RenderCore",
+                    "RHI"
+                }
              );
         }
     }
