[CIS] Redoing change 13925048 (disabling Mac, Linux RigLogicLibTest module since it fails to link)

diff --git a/Engine/Plugins/Animation/RigLogic/RigLogic.uplugin b/Engine/Plugins/Animation/RigLogic/RigLogic.uplugin
--- a/Engine/Plugins/Animation/RigLogic/RigLogic.uplugin
+++ b/Engine/Plugins/Animation/RigLogic/RigLogic.uplugin
@@ -25,9 +25,7 @@
 			"LoadingPhase": "PreDefault",
 			"WhitelistPlatforms": [
 				"Win32",
-				"Win64",
-				"Linux",
-				"Mac"
+				"Win64"
 			]
 		},
 		{
@@ -50,4 +48,4 @@
 			"Enabled": true
 		}
 	]
-}
\ No newline at end of file
+}
