Merging RigLogic up to CL 13924890

#rnx

diff --git a/Engine/Plugins/Animation/RigLogic/Source/RigLogicEditor/RigLogicEditor.Build.cs b/Engine/Plugins/Animation/RigLogic/Source/RigLogicEditor/RigLogicEditor.Build.cs
--- a/Engine/Plugins/Animation/RigLogic/Source/RigLogicEditor/RigLogicEditor.Build.cs
+++ b/Engine/Plugins/Animation/RigLogic/Source/RigLogicEditor/RigLogicEditor.Build.cs
@@ -11,14 +11,12 @@ namespace UnrealBuildTool.Rules
             get { return ModuleDirectory; }
         }
 
-        private string ThirdPartyPath
-        {
-            get { return Path.GetFullPath(Path.Combine(ModulePath, "../Source/ThirdParty/")); }
-        }
-
         public RigLogicEditor(ReadOnlyTargetRules Target) : base(Target)
         {
-            PublicDefinitions.Add("RL_SHARED=1");
+            if (Target.LinkType != TargetLinkType.Monolithic)
+            {
+                PublicDefinitions.Add("RL_SHARED=1");
+            }
 
             PublicDependencyModuleNames.AddRange(
                 new string[]
diff --git a/Engine/Plugins/Animation/RigLogic/Source/RigLogicLib/Private/pma/resources/AlignedMemoryResource.cpp b/Engine/Plugins/Animation/RigLogic/Source/RigLogicLib/Private/pma/resources/AlignedMemoryResource.cpp
--- a/Engine/Plugins/Animation/RigLogic/Source/RigLogicLib/Private/pma/resources/AlignedMemoryResource.cpp
+++ b/Engine/Plugins/Animation/RigLogic/Source/RigLogicLib/Private/pma/resources/AlignedMemoryResource.cpp
@@ -1,5 +1,6 @@
 #include "pma/resources/AlignedMemoryResource.h"
 
+#include <cassert>
 #include <cstddef>
 #include <cstdlib>
 
@@ -24,7 +25,54 @@
     #define ALIGNED_ALLOC(ptr, alignment, size) ptr = ::memalign(alignment, size)
     #define ALIGNED_FREE(ptr) ::free(ptr)
 #else
-    #error No aligned allocation is possible.
+    inline std::uintptr_t alignAddress(std::uintptr_t address, std::size_t alignment) {
+        const std::size_t mask = alignment - 1ul;
+        assert((alignment & mask) == 0ul);
+        return (address + mask) & ~mask;
+    }
+
+    template<typename T>
+    inline T* alignPointer(T* ptr, std::size_t alignment) {
+        const std::uintptr_t address = reinterpret_cast<std::uintptr_t>(ptr);
+        const std::uintptr_t aligned = alignAddress(address, alignment);
+        return reinterpret_cast<T*>(aligned);
+    }
+
+    void* alignedAlloc(std::size_t size, std::size_t alignment) {
+        const std::size_t actualSize = size + alignment;
+        char* pUnalignedChunk = new char[actualSize];
+        char* pAlignedChunk = alignPointer(pUnalignedChunk, alignment);
+        if (pAlignedChunk == pUnalignedChunk) {
+            // There was no need for additional alignment, so an artifical gap
+            // must be added to make room for the shift value
+            pAlignedChunk += alignment;
+        }
+        // Compute and write the shift value into the buffer (it is needed during
+        // deallocation to reconstruct by which amount was the given pointer
+        // offset to meet the alignment requirement)
+        std::ptrdiff_t shift = pAlignedChunk - pUnalignedChunk;
+        // The guaranteed storage size for the shift value is only one byte
+        assert(shift > 0ul && shift <= 256ul);
+        pAlignedChunk[-1] = static_cast<char>(shift & 0xFF);
+        return static_cast<void*>(pAlignedChunk);
+    }
+
+    void alignedFree(void* ptr) {
+        if (ptr) {
+            auto pAlignedChunk = reinterpret_cast<char*>(ptr);
+            std::ptrdiff_t shift = pAlignedChunk[-1];
+            // There is no 0 alignment, so the value is reused to actually denote alignment by 256
+            if (shift == 0) {
+                shift = 256;
+            }
+            // Compute the actual start address of the chunk, prior to aligning it's address
+            auto pUnalignedChunk = pAlignedChunk - shift;
+            delete[] pUnalignedChunk;
+        }
+    }
+
+    #define ALIGNED_ALLOC(ptr, alignment, size) ptr = alignedAlloc(size, alignment)
+    #define ALIGNED_FREE(ptr) alignedFree(ptr)
 #endif
 
 namespace pma {
diff --git a/Engine/Plugins/Animation/RigLogic/Source/RigLogicLib/RigLogicLib.Build.cs b/Engine/Plugins/Animation/RigLogic/Source/RigLogicLib/RigLogicLib.Build.cs
--- a/Engine/Plugins/Animation/RigLogic/Source/RigLogicLib/RigLogicLib.Build.cs
+++ b/Engine/Plugins/Animation/RigLogic/Source/RigLogicLib/RigLogicLib.Build.cs
@@ -30,7 +30,10 @@ public class RigLogicLib : ModuleRules
 
         Type = ModuleType.CPlusPlus;
 
-        PrivateDefinitions.Add("RL_BUILD_SHARED=1");
+        if (Target.LinkType != TargetLinkType.Monolithic)
+        {
+            PrivateDefinitions.Add("RL_BUILD_SHARED=1");
+        }
 
         if (Target.Platform == UnrealTargetPlatform.Win32 ||
                 Target.Platform == UnrealTargetPlatform.Win64 ||
diff --git a/Engine/Plugins/Animation/RigLogic/Source/RigLogicLibTest/RigLogicLibTest.Build.cs b/Engine/Plugins/Animation/RigLogic/Source/RigLogicLibTest/RigLogicLibTest.Build.cs
--- a/Engine/Plugins/Animation/RigLogic/Source/RigLogicLibTest/RigLogicLibTest.Build.cs
+++ b/Engine/Plugins/Animation/RigLogic/Source/RigLogicLibTest/RigLogicLibTest.Build.cs
@@ -19,8 +19,16 @@ public class RigLogicLibTest : ModuleRules
 
         string RigLogicLibPath = Path.GetFullPath(Path.Combine(ModuleDirectory, "../RigLogicLib"));
 
-        PrivateDefinitions.Add("RIGLOGIC_MODULE_DISCARD");
-        ConditionalAddModuleDirectory(new DirectoryReference(RigLogicLibPath));
+        if (Target.LinkType == TargetLinkType.Monolithic)
+        {
+            PublicDependencyModuleNames.Add("RigLogicLib");
+            PrivateIncludePaths.Add(Path.Combine(RigLogicLibPath, "Private"));
+        }
+        else
+        {
+            PrivateDefinitions.Add("RIGLOGIC_MODULE_DISCARD");
+            ConditionalAddModuleDirectory(new DirectoryReference(RigLogicLibPath));
+        }
 
         PublicDependencyModuleNames.AddRange(
             new string[]
diff --git a/Engine/Plugins/Animation/RigLogic/Source/RigLogicModule/RigLogicModule.Build.cs b/Engine/Plugins/Animation/RigLogic/Source/RigLogicModule/RigLogicModule.Build.cs
--- a/Engine/Plugins/Animation/RigLogic/Source/RigLogicModule/RigLogicModule.Build.cs
+++ b/Engine/Plugins/Animation/RigLogic/Source/RigLogicModule/RigLogicModule.Build.cs
@@ -15,7 +15,10 @@ namespace UnrealBuildTool.Rules
 
         public RigLogicModule(ReadOnlyTargetRules Target) : base(Target)
         {
-            PublicDefinitions.Add("RL_SHARED=1");
+            if (Target.LinkType != TargetLinkType.Monolithic)
+            {
+                PublicDefinitions.Add("RL_SHARED=1");
+            }
 
             PublicDependencyModuleNames.AddRange(
                 new string[]
