Editgrated RigLogic up to CL 13936592. Fixes nonunity compile.

#rnx

diff --git a/Engine/Plugins/Animation/RigLogic/Source/RigLogicLib/Private/pma/resources/AlignedMemoryResource.cpp b/Engine/Plugins/Animation/RigLogic/Source/RigLogicLib/Private/pma/resources/AlignedMemoryResource.cpp
--- a/Engine/Plugins/Animation/RigLogic/Source/RigLogicLib/Private/pma/resources/AlignedMemoryResource.cpp
+++ b/Engine/Plugins/Animation/RigLogic/Source/RigLogicLib/Private/pma/resources/AlignedMemoryResource.cpp
@@ -2,6 +2,7 @@
 
 #include <cassert>
 #include <cstddef>
+#include <cstdint>
 #include <cstdlib>
 
 #if defined(__linux__) || defined(__APPLE__)
