Merging RigLogic from Dev-CharacterTech to UE5Main up to CL 13956739

diff --git a/Engine/Plugins/Animation/RigLogic/Source/RigLogicLib/Private/dna/BaseImpl.h b/Engine/Plugins/Animation/RigLogic/Source/RigLogicLib/Private/dna/BaseImpl.h
--- a/Engine/Plugins/Animation/RigLogic/Source/RigLogicLib/Private/dna/BaseImpl.h
+++ b/Engine/Plugins/Animation/RigLogic/Source/RigLogicLib/Private/dna/BaseImpl.h
@@ -1,3 +1,5 @@
+// Copyright Epic Games, Inc. All Rights Reserved.
+
 #pragma once
 
 #include "dna/DNA.h"
diff --git a/Engine/Plugins/Animation/RigLogic/Source/RigLogicLib/Private/riglogic/system/simd/SIMD.h b/Engine/Plugins/Animation/RigLogic/Source/RigLogicLib/Private/riglogic/system/simd/SIMD.h
--- a/Engine/Plugins/Animation/RigLogic/Source/RigLogicLib/Private/riglogic/system/simd/SIMD.h
+++ b/Engine/Plugins/Animation/RigLogic/Source/RigLogicLib/Private/riglogic/system/simd/SIMD.h
@@ -1,14 +1,16 @@
+// Copyright Epic Games, Inc. All Rights Reserved.
+
 #pragma once
 
-#ifdef RL_USE_HALF_FLOATS
+#if defined(RL_USE_HALF_FLOATS) && !defined(TRIMD_ENABLE_F16C)
     #define TRIMD_ENABLE_F16C
 #endif  // RL_USE_HALF_FLOATS
 
-#ifdef RL_BUILD_WITH_AVX
+#if defined(RL_BUILD_WITH_AVX) && !defined(TRIMD_ENABLE_AVX)
     #define TRIMD_ENABLE_AVX
 #endif  // RL_BUILD_WITH_AVX
 
-#ifdef RL_BUILD_WITH_SSE
+#if defined(RL_BUILD_WITH_SSE) && !defined(TRIMD_ENABLE_SSE)
     #define TRIMD_ENABLE_SSE
 #endif  // RL_BUILD_WITH_SSE
 
diff --git a/Engine/Plugins/Animation/RigLogic/Source/RigLogicLib/Public/dna/version/Version.h b/Engine/Plugins/Animation/RigLogic/Source/RigLogicLib/Public/dna/version/Version.h
--- a/Engine/Plugins/Animation/RigLogic/Source/RigLogicLib/Public/dna/version/Version.h
+++ b/Engine/Plugins/Animation/RigLogic/Source/RigLogicLib/Public/dna/version/Version.h
@@ -1,6 +1,8 @@
+// Copyright Epic Games, Inc. All Rights Reserved.
+
 #pragma once
 
 #define DNA_MAJOR_VERSION 1
 #define DNA_MINOR_VERSION 1
-#define DNA_PATCH_VERSION 1
-#define DNA_VERSION_STRING "1.1.1"
+#define DNA_PATCH_VERSION 2
+#define DNA_VERSION_STRING "1.1.2"
diff --git a/Engine/Plugins/Animation/RigLogic/Source/RigLogicLib/Public/pma/version/Version.h b/Engine/Plugins/Animation/RigLogic/Source/RigLogicLib/Public/pma/version/Version.h
--- a/Engine/Plugins/Animation/RigLogic/Source/RigLogicLib/Public/pma/version/Version.h
+++ b/Engine/Plugins/Animation/RigLogic/Source/RigLogicLib/Public/pma/version/Version.h
@@ -1,6 +1,8 @@
+// Copyright Epic Games, Inc. All Rights Reserved.
+
 #pragma once
 
 #define PMA_MAJOR_VERSION 1
 #define PMA_MINOR_VERSION 0
-#define PMA_PATCH_VERSION 2
-#define PMA_VERSION_STRING "1.0.2"
+#define PMA_PATCH_VERSION 3
+#define PMA_VERSION_STRING "1.0.3"
diff --git a/Engine/Plugins/Animation/RigLogic/Source/RigLogicLib/Public/status/Provider.h b/Engine/Plugins/Animation/RigLogic/Source/RigLogicLib/Public/status/Provider.h
--- a/Engine/Plugins/Animation/RigLogic/Source/RigLogicLib/Public/status/Provider.h
+++ b/Engine/Plugins/Animation/RigLogic/Source/RigLogicLib/Public/status/Provider.h
@@ -1,3 +1,5 @@
+// Copyright Epic Games, Inc. All Rights Reserved.
+
 #pragma once
 
 #include "status/Defs.h"
@@ -7,9 +9,10 @@
     #pragma warning(push)
     #pragma warning(disable : 4365 4987)
 #endif
+#include <stdio.h>
+
 #include <algorithm>
 #include <array>
-#include <cstdio>
 #ifdef _MSC_VER
     #pragma warning(pop)
 #endif
@@ -33,9 +36,9 @@ class SCAPI StatusProvider {
                 #pragma GCC diagnostic ignored "-Wformat-security"
             #endif
             // The returned number of bytes to be written does not include the null terminator
-            const auto neededSize = std::snprintf(nullptr, 0ul, status.message, args ...) + 1;
+            const auto neededSize = snprintf(nullptr, 0ul, status.message, args ...) + 1;
             const auto size = std::min(buffer.size(), static_cast<std::size_t>(neededSize));
-            std::snprintf(buffer.data(), size, status.message, args ...);
+            snprintf(buffer.data(), size, status.message, args ...);
             #if !defined(__clang__) && defined(__GNUC__)
                 #pragma GCC diagnostic pop
             #endif
diff --git a/Engine/Plugins/Animation/RigLogic/Source/RigLogicLib/Public/status/version/Version.h b/Engine/Plugins/Animation/RigLogic/Source/RigLogicLib/Public/status/version/Version.h
--- a/Engine/Plugins/Animation/RigLogic/Source/RigLogicLib/Public/status/version/Version.h
+++ b/Engine/Plugins/Animation/RigLogic/Source/RigLogicLib/Public/status/version/Version.h
@@ -1,6 +1,8 @@
+// Copyright Epic Games, Inc. All Rights Reserved.
+
 #pragma once
 
 #define SC_MAJOR_VERSION 1
 #define SC_MINOR_VERSION 0
-#define SC_PATCH_VERSION 1
-#define SC_VERSION_STRING "1.0.1"
+#define SC_PATCH_VERSION 2
+#define SC_VERSION_STRING "1.0.2"
diff --git a/Engine/Plugins/Animation/RigLogic/Source/RigLogicLib/Public/terse/version/Version.h b/Engine/Plugins/Animation/RigLogic/Source/RigLogicLib/Public/terse/version/Version.h
--- a/Engine/Plugins/Animation/RigLogic/Source/RigLogicLib/Public/terse/version/Version.h
+++ b/Engine/Plugins/Animation/RigLogic/Source/RigLogicLib/Public/terse/version/Version.h
@@ -1,6 +1,8 @@
+// Copyright Epic Games, Inc. All Rights Reserved.
+
 #pragma once
 
 #define TERSE_MAJOR_VERSION 1
 #define TERSE_MINOR_VERSION 1
-#define TERSE_PATCH_VERSION 1
-#define TERSE_VERSION_STRING "1.1.1"
+#define TERSE_PATCH_VERSION 2
+#define TERSE_VERSION_STRING "1.1.2"
diff --git a/Engine/Plugins/Animation/RigLogic/Source/RigLogicLib/Public/trimd/version/Version.h b/Engine/Plugins/Animation/RigLogic/Source/RigLogicLib/Public/trimd/version/Version.h
--- a/Engine/Plugins/Animation/RigLogic/Source/RigLogicLib/Public/trimd/version/Version.h
+++ b/Engine/Plugins/Animation/RigLogic/Source/RigLogicLib/Public/trimd/version/Version.h
@@ -1,6 +1,8 @@
+// Copyright Epic Games, Inc. All Rights Reserved.
+
 #pragma once
 
 #define TRIMD_MAJOR_VERSION 1
 #define TRIMD_MINOR_VERSION 0
-#define TRIMD_PATCH_VERSION 0
-#define TRIMD_VERSION_STRING "1.0.0"
+#define TRIMD_PATCH_VERSION 1
+#define TRIMD_VERSION_STRING "1.0.1"
diff --git a/Engine/Plugins/Animation/RigLogic/Source/RigLogicLib/Public/trio/version/Version.h b/Engine/Plugins/Animation/RigLogic/Source/RigLogicLib/Public/trio/version/Version.h
--- a/Engine/Plugins/Animation/RigLogic/Source/RigLogicLib/Public/trio/version/Version.h
+++ b/Engine/Plugins/Animation/RigLogic/Source/RigLogicLib/Public/trio/version/Version.h
@@ -1,6 +1,8 @@
+// Copyright Epic Games, Inc. All Rights Reserved.
+
 #pragma once
 
 #define TRIO_MAJOR_VERSION 1
 #define TRIO_MINOR_VERSION 0
-#define TRIO_PATCH_VERSION 2
-#define TRIO_VERSION_STRING "1.0.2"
+#define TRIO_PATCH_VERSION 3
+#define TRIO_VERSION_STRING "1.0.3"
diff --git a/Engine/Plugins/Animation/RigLogic/Source/RigLogicLibTest/Private/rltests/dna/FakeReader.h b/Engine/Plugins/Animation/RigLogic/Source/RigLogicLibTest/Private/rltests/dna/FakeReader.h
--- a/Engine/Plugins/Animation/RigLogic/Source/RigLogicLibTest/Private/rltests/dna/FakeReader.h
+++ b/Engine/Plugins/Animation/RigLogic/Source/RigLogicLibTest/Private/rltests/dna/FakeReader.h
@@ -1,3 +1,5 @@
+// Copyright Epic Games, Inc. All Rights Reserved.
+
 #pragma once
 
 #include "riglogic/types/Aliases.h"
@@ -6,6 +8,11 @@
 
 #include <cstdint>
 
+#if defined(_MSC_VER) && !defined(__clang__)
+#pragma warning(push)
+#pragma warning(disable:4100) // unreferenced formal parameter
+#endif
+
 namespace dna {
 
 class FakeReader : public Reader {
@@ -437,3 +444,7 @@ class FakeReader : public Reader {
 };
 
 }  // namespace dna
+
+#if defined(_MSC_VER) && !defined(__clang__)
+#pragma warning(pop)
+#endif
diff --git a/Engine/Plugins/Animation/RigLogic/Source/RigLogicModule/Private/DNAAsset.cpp b/Engine/Plugins/Animation/RigLogic/Source/RigLogicModule/Private/DNAAsset.cpp
--- a/Engine/Plugins/Animation/RigLogic/Source/RigLogicModule/Private/DNAAsset.cpp
+++ b/Engine/Plugins/Animation/RigLogic/Source/RigLogicModule/Private/DNAAsset.cpp
@@ -60,7 +60,7 @@ bool UDNAAsset::Init(const FString DNAFilename)
 	}
 	
 	// Load run-time data (behavior) from whole-DNA buffer into BehaviorStreamReader
-	TSharedPtr<IDNAReader> DNAReaderBehavior = ReadDNALayerFromBuffer(&TempFileBuffer, EDNADataLayer::Behavior, 0u); //0u = MaxLOD
+	TSharedPtr<IDNAReader> DNAReaderBehavior = ReadDNAFromBuffer(&TempFileBuffer, EDNADataLayer::Behavior, 0u); //0u = MaxLOD
 
 	if (!DNAReaderBehavior.IsValid())
 	{
@@ -73,7 +73,7 @@ bool UDNAAsset::Init(const FString DNAFilename)
 #if WITH_EDITORONLY_DATA
 	//We use geometry part of the data in MHC only (for updating the SkeletalMesh with
 	//result of GeneSplicer), so we can drop geometry part when cooking for runtime
-	TSharedPtr<IDNAReader> DNAReaderGeometry = ReadDNALayerFromBuffer(&TempFileBuffer, EDNADataLayer::Geometry, 0u); //0u = MaxLOD
+	TSharedPtr<IDNAReader> DNAReaderGeometry = ReadDNAFromBuffer(&TempFileBuffer, EDNADataLayer::Geometry, 0u); //0u = MaxLOD
 	GeometryStreamReader = DNAReaderGeometry;
 	WriteToBuffer(DNAReaderGeometry.Get(), &DesignTimeData, AVG_GEOMETRY_SIZE);
 	//Note: in future, we will want to load geometry data in-game too 
@@ -94,14 +94,14 @@ void UDNAAsset::PostLoad()
 		return;
 	}
 
-	BehaviorStreamReader = ReadDNALayerFromBuffer(&BehaviorData, EDNADataLayer::Behavior, 0u); //0u = max LOD
+	BehaviorStreamReader = ReadDNAFromBuffer(&BehaviorData, EDNADataLayer::Behavior, 0u); //0u = max LOD
 
 #if WITH_EDITORONLY_DATA
 	if (DesignTimeData.Num() == 0)
 	{
 		return;
 	}
-	GeometryStreamReader = ReadDNALayerFromBuffer(&DesignTimeData, EDNADataLayer::Geometry, 0u); //0u = max LOD
+	GeometryStreamReader = ReadDNAFromBuffer(&DesignTimeData, EDNADataLayer::Geometry, 0u); //0u = max LOD
 #endif // #if WITH_EDITORONLY_DATA
 }
 
@@ -120,7 +120,7 @@ TSharedPtr<IDNAReader> UDNAAsset::CopyDNALayer(const IDNAReader* Source, EDNADat
 
 	MemoryStream.seek(0ul);
 
-	return ReadDNALayerFromBuffer(&MemoryBuffer, DNADataLayer);
+	return ReadDNAFromBuffer(&MemoryBuffer, DNADataLayer);
 }
 
 void UDNAAsset::SetBehaviorReader(const TSharedPtr<IDNAReader> SourceDNAReader)
diff --git a/Engine/Plugins/Animation/RigLogic/Source/RigLogicModule/Private/DNAUtils.cpp b/Engine/Plugins/Animation/RigLogic/Source/RigLogicModule/Private/DNAUtils.cpp
--- a/Engine/Plugins/Animation/RigLogic/Source/RigLogicModule/Private/DNAUtils.cpp
+++ b/Engine/Plugins/Animation/RigLogic/Source/RigLogicModule/Private/DNAUtils.cpp
@@ -34,14 +34,14 @@ TSharedPtr<IDNAReader> ReadDNAFromFile(const FString& Path, EDNADataLayer Layer,
 	return ReadDNAStream(MoveTemp(DNAStreamReader));
 }
 
-TSharedPtr<IDNAReader> ReadDNALayerFromBuffer(TArray<uint8>* DNABuffer, EDNADataLayer Layer, uint16_t MaxLOD)
+TSharedPtr<IDNAReader> ReadDNAFromBuffer(TArray<uint8>* DNABuffer, EDNADataLayer Layer, uint16_t MaxLOD)
 {
 	FRigLogicMemoryStream DNAMemoryStream(DNABuffer);
 	auto DNAStreamReader = rl4::makeScoped<dna::StreamReader>(&DNAMemoryStream, static_cast<dna::DataLayer>(Layer), MaxLOD, FMemoryResource::Instance());
 	return ReadDNAStream(MoveTemp(DNAStreamReader));
 }
 
-TSharedPtr<IDNAReader> ReadDNALayerFromBuffer(TArray<uint8>* DNABuffer, EDNADataLayer Layer, TArrayView<uint16_t> LODs)
+TSharedPtr<IDNAReader> ReadDNAFromBuffer(TArray<uint8>* DNABuffer, EDNADataLayer Layer, TArrayView<uint16_t> LODs)
 {
 	FRigLogicMemoryStream DNAMemoryStream(DNABuffer);
 	auto DNAStreamReader = rl4::makeScoped<dna::StreamReader>(&DNAMemoryStream, static_cast<dna::DataLayer>(Layer), LODs.GetData(), static_cast<uint16>(LODs.Num()), FMemoryResource::Instance());
diff --git a/Engine/Plugins/Animation/RigLogic/Source/RigLogicModule/Private/RigLogicMemoryStream.cpp b/Engine/Plugins/Animation/RigLogic/Source/RigLogicModule/Private/RigLogicMemoryStream.cpp
--- a/Engine/Plugins/Animation/RigLogic/Source/RigLogicModule/Private/RigLogicMemoryStream.cpp
+++ b/Engine/Plugins/Animation/RigLogic/Source/RigLogicModule/Private/RigLogicMemoryStream.cpp
@@ -30,7 +30,8 @@ void FRigLogicMemoryStream::write(const char* WriteFromBuffer, size_t Size)
 {
 	FString writeStr;
 	int32 BufferSize = BitStreamBuffer->Num();
-	if (PositionInBuffer + Size > BufferSize)
+	ensure(BufferSize >= 0);
+	if (PositionInBuffer + Size > static_cast<size_t>(BufferSize))
 	{
 		int32 difference = PositionInBuffer + Size - BufferSize;
 		BitStreamBuffer->AddUninitialized(difference);
@@ -43,13 +44,3 @@ size_t FRigLogicMemoryStream::size()
 {
 	return BitStreamBuffer->Num();
 }
-
-bool FRigLogicMemoryStream::hasError()
-{
-	return false;
-}
-
-TArray<uint8>* FRigLogicMemoryStream::GetBitStreamBuffer()
-{
-	return BitStreamBuffer;
-}
\ No newline at end of file
diff --git a/Engine/Plugins/Animation/RigLogic/Source/RigLogicModule/Private/RigLogicMemoryStream.h b/Engine/Plugins/Animation/RigLogic/Source/RigLogicModule/Private/RigLogicMemoryStream.h
--- a/Engine/Plugins/Animation/RigLogic/Source/RigLogicModule/Private/RigLogicMemoryStream.h
+++ b/Engine/Plugins/Animation/RigLogic/Source/RigLogicModule/Private/RigLogicMemoryStream.h
@@ -14,7 +14,6 @@ class FRigLogicMemoryStream: public rl4::BoundedIOStream
 public:
 	/** The buffer is not copied, the pointer to it is stored inside this object**/
 	FRigLogicMemoryStream(TArray<uint8>* Buffer);
-	~FRigLogicMemoryStream() {} //as no memory was allocated, we do not release it
 
 	void seek(size_t Position) override;
 	size_t tell() override;
@@ -23,10 +22,6 @@ public:
 	void read(char* ReadToBuffer, size_t Size) override;
 	void write(const char* WriteFromBuffer, size_t Size) override;
 	size_t size() override;
-	void release() {}
-	static void deallocResource(void* ptr) {}; //no memory was allocated
-	bool hasError();
-	TArray<uint8>* GetBitStreamBuffer();
 
 private:
 	TArray<uint8>* BitStreamBuffer; //doesn't contain the array, only points to the array given to it
diff --git a/Engine/Plugins/Animation/RigLogic/Source/RigLogicModule/Private/RigUnit_RigLogic.cpp b/Engine/Plugins/Animation/RigLogic/Source/RigLogicModule/Private/RigUnit_RigLogic.cpp
--- a/Engine/Plugins/Animation/RigLogic/Source/RigLogicModule/Private/RigUnit_RigLogic.cpp
+++ b/Engine/Plugins/Animation/RigLogic/Source/RigLogicModule/Private/RigUnit_RigLogic.cpp
@@ -7,7 +7,8 @@
 #include "RigInstance.h"
 #include "RigLogic.h"
 #include "SkelMeshDNAUtils.h"
-#include "Runtime/Core/Public/Math/TransformNonVectorized.h"
+#include "Components/SkeletalMeshComponent.h"
+#include "Math/TransformNonVectorized.h"
 #include "Units/RigUnitContext.h"
 
 DEFINE_LOG_CATEGORY(LogRigLogicUnit);
diff --git a/Engine/Plugins/Animation/RigLogic/Source/RigLogicModule/Public/DNAUtils.h b/Engine/Plugins/Animation/RigLogic/Source/RigLogicModule/Public/DNAUtils.h
--- a/Engine/Plugins/Animation/RigLogic/Source/RigLogicModule/Public/DNAUtils.h
+++ b/Engine/Plugins/Animation/RigLogic/Source/RigLogicModule/Public/DNAUtils.h
@@ -10,6 +10,6 @@
 
 RIGLOGICMODULE_API TSharedPtr<IDNAReader> ReadDNAFromFile(const FString& Path, EDNADataLayer Layer = EDNADataLayer::All, uint16_t MaxLOD = 0);
 RIGLOGICMODULE_API TSharedPtr<IDNAReader> ReadDNAFromFile(const FString& Path, EDNADataLayer Layer, TArrayView<uint16_t> LODs);
-RIGLOGICMODULE_API TSharedPtr<IDNAReader> ReadDNALayerFromBuffer(TArray<uint8>* DNABuffer, EDNADataLayer Layer = EDNADataLayer::All, uint16_t MaxLOD = 0);
-RIGLOGICMODULE_API TSharedPtr<IDNAReader> ReadDNALayerFromBuffer(TArray<uint8>* DNABuffer, EDNADataLayer Layer, TArrayView<uint16_t> LODs);
+RIGLOGICMODULE_API TSharedPtr<IDNAReader> ReadDNAFromBuffer(TArray<uint8>* DNABuffer, EDNADataLayer Layer = EDNADataLayer::All, uint16_t MaxLOD = 0);
+RIGLOGICMODULE_API TSharedPtr<IDNAReader> ReadDNAFromBuffer(TArray<uint8>* DNABuffer, EDNADataLayer Layer, TArrayView<uint16_t> LODs);
 RIGLOGICMODULE_API void WriteDNAToFile(const IDNAReader* Reader, EDNADataLayer Layer, const FString& Path);
\ No newline at end of file
