Merging using RigLogic_CharacterTech_to_UE5Main @ CL 13961803

Should fix Copyright notice warnings

#rnx

diff --git a/Engine/Plugins/Animation/RigLogic/Source/RigLogicEditor/Private/DNAAssetImportFactory.cpp b/Engine/Plugins/Animation/RigLogic/Source/RigLogicEditor/Private/DNAAssetImportFactory.cpp
--- a/Engine/Plugins/Animation/RigLogic/Source/RigLogicEditor/Private/DNAAssetImportFactory.cpp
+++ b/Engine/Plugins/Animation/RigLogic/Source/RigLogicEditor/Private/DNAAssetImportFactory.cpp
@@ -1,4 +1,4 @@
-// Copyright 1998-2019 Epic Games, Inc. All Rights Reserved.
+// Copyright Epic Games, Inc. All Rights Reserved.
 
 #include "DNAAssetImportFactory.h"
 #include "DNAAsset.h"
diff --git a/Engine/Plugins/Animation/RigLogic/Source/RigLogicEditor/Private/DNAAssetImportWindow.cpp b/Engine/Plugins/Animation/RigLogic/Source/RigLogicEditor/Private/DNAAssetImportWindow.cpp
--- a/Engine/Plugins/Animation/RigLogic/Source/RigLogicEditor/Private/DNAAssetImportWindow.cpp
+++ b/Engine/Plugins/Animation/RigLogic/Source/RigLogicEditor/Private/DNAAssetImportWindow.cpp
@@ -1,4 +1,5 @@
-// Copyright 1998-2019 Epic Games, Inc. All Rights Reserved.
+// Copyright Epic Games, Inc. All Rights Reserved.
+
 #include "DNAAssetImportWindow.h"
 #include "Modules/ModuleManager.h"
 #include "Widgets/Layout/SBorder.h"
diff --git a/Engine/Plugins/Animation/RigLogic/Source/RigLogicEditor/Private/DNAImporter.cpp b/Engine/Plugins/Animation/RigLogic/Source/RigLogicEditor/Private/DNAImporter.cpp
--- a/Engine/Plugins/Animation/RigLogic/Source/RigLogicEditor/Private/DNAImporter.cpp
+++ b/Engine/Plugins/Animation/RigLogic/Source/RigLogicEditor/Private/DNAImporter.cpp
@@ -1,4 +1,4 @@
-// Copyright 1998-2019 Epic Games, Inc. All Rights Reserved.
+// Copyright Epic Games, Inc. All Rights Reserved.
 
 /*=============================================================================
 	Main implementation of FDNAImporter : import DNA data to Skeletal Mesh
diff --git a/Engine/Plugins/Animation/RigLogic/Source/RigLogicEditor/Private/RigLogicEditor.cpp b/Engine/Plugins/Animation/RigLogic/Source/RigLogicEditor/Private/RigLogicEditor.cpp
--- a/Engine/Plugins/Animation/RigLogic/Source/RigLogicEditor/Private/RigLogicEditor.cpp
+++ b/Engine/Plugins/Animation/RigLogic/Source/RigLogicEditor/Private/RigLogicEditor.cpp
@@ -1,4 +1,4 @@
-// Copyright 1998-2019 Epic Games, Inc. All Rights Reserved.
+// Copyright Epic Games, Inc. All Rights Reserved.
 
 #pragma once
 
@@ -15,10 +15,6 @@ IMPLEMENT_MODULE(FRigLogicEditor, RigLogicEditor)
 
 DEFINE_LOG_CATEGORY_STATIC(LogRigLogicEditor, Log, All);
 
-FRigLogicEditor::FRigLogicEditor()
-{
-}
-
 void FRigLogicEditor::StartupModule()
 {
 }
diff --git a/Engine/Plugins/Animation/RigLogic/Source/RigLogicEditor/Public/DNAAssetImportFactory.h b/Engine/Plugins/Animation/RigLogic/Source/RigLogicEditor/Public/DNAAssetImportFactory.h
--- a/Engine/Plugins/Animation/RigLogic/Source/RigLogicEditor/Public/DNAAssetImportFactory.h
+++ b/Engine/Plugins/Animation/RigLogic/Source/RigLogicEditor/Public/DNAAssetImportFactory.h
@@ -1,4 +1,4 @@
-// Copyright 1998-2019 Epic Games, Inc. All Rights Reserved.
+// Copyright Epic Games, Inc. All Rights Reserved.
 
 #pragma once 
 
@@ -6,7 +6,7 @@
 #include "Factories/ImportSettings.h"
 #include "EditorReimportHandler.h"
 #include "DNAAsset.h"
-#include "Runtime/Engine/Classes/Engine/SkeletalMesh.h"
+#include "Engine/SkeletalMesh.h"
 #include "DNAAssetImportFactory.generated.h"
 
 DECLARE_LOG_CATEGORY_EXTERN(LogDNAImportFactory, Log, All);
diff --git a/Engine/Plugins/Animation/RigLogic/Source/RigLogicEditor/Public/DNAAssetImportUI.h b/Engine/Plugins/Animation/RigLogic/Source/RigLogicEditor/Public/DNAAssetImportUI.h
--- a/Engine/Plugins/Animation/RigLogic/Source/RigLogicEditor/Public/DNAAssetImportUI.h
+++ b/Engine/Plugins/Animation/RigLogic/Source/RigLogicEditor/Public/DNAAssetImportUI.h
@@ -1,4 +1,4 @@
-// Copyright 1998-2019 Epic Games, Inc. All Rights Reserved.
+// Copyright Epic Games, Inc. All Rights Reserved.
 
 /**
  * DNA Asset Importer UI options.
diff --git a/Engine/Plugins/Animation/RigLogic/Source/RigLogicEditor/Public/DNAAssetImportWindow.h b/Engine/Plugins/Animation/RigLogic/Source/RigLogicEditor/Public/DNAAssetImportWindow.h
--- a/Engine/Plugins/Animation/RigLogic/Source/RigLogicEditor/Public/DNAAssetImportWindow.h
+++ b/Engine/Plugins/Animation/RigLogic/Source/RigLogicEditor/Public/DNAAssetImportWindow.h
@@ -1,4 +1,4 @@
-// Copyright 1998-2019 Epic Games, Inc. All Rights Reserved.
+// Copyright Epic Games, Inc. All Rights Reserved.
 
 #pragma once
 
diff --git a/Engine/Plugins/Animation/RigLogic/Source/RigLogicEditor/Public/DNAImporter.h b/Engine/Plugins/Animation/RigLogic/Source/RigLogicEditor/Public/DNAImporter.h
--- a/Engine/Plugins/Animation/RigLogic/Source/RigLogicEditor/Public/DNAImporter.h
+++ b/Engine/Plugins/Animation/RigLogic/Source/RigLogicEditor/Public/DNAImporter.h
@@ -1,4 +1,4 @@
-// Copyright 1998-2019 Epic Games, Inc. All Rights Reserved.
+// Copyright Epic Games, Inc. All Rights Reserved.
 
 #pragma once
 
diff --git a/Engine/Plugins/Animation/RigLogic/Source/RigLogicEditor/Public/RigLogicEditor.h b/Engine/Plugins/Animation/RigLogic/Source/RigLogicEditor/Public/RigLogicEditor.h
--- a/Engine/Plugins/Animation/RigLogic/Source/RigLogicEditor/Public/RigLogicEditor.h
+++ b/Engine/Plugins/Animation/RigLogic/Source/RigLogicEditor/Public/RigLogicEditor.h
@@ -1,17 +1,15 @@
-// Copyright 1998-2019 Epic Games, Inc. All Rights Reserved.
+// Copyright Epic Games, Inc. All Rights Reserved.
 
 #pragma once
 
 #include "CoreMinimal.h"
 #include "Modules/ModuleInterface.h"
 
-struct RIGLOGICEDITOR_API FRigLogicEditor : IModuleInterface 
+class RIGLOGICEDITOR_API FRigLogicEditor : public IModuleInterface 
 {
 public:
-	virtual void StartupModule() override;
-	virtual void ShutdownModule() override;
-
-	FRigLogicEditor();
+	void StartupModule() override;
+	void ShutdownModule() override;
 
 private:
 	TArray<TSharedRef<class IAssetTypeActions>> AssetTypeActions;
diff --git a/Engine/Plugins/Animation/RigLogic/Source/RigLogicEditor/RigLogicEditor.Build.cs b/Engine/Plugins/Animation/RigLogic/Source/RigLogicEditor/RigLogicEditor.Build.cs
--- a/Engine/Plugins/Animation/RigLogic/Source/RigLogicEditor/RigLogicEditor.Build.cs
+++ b/Engine/Plugins/Animation/RigLogic/Source/RigLogicEditor/RigLogicEditor.Build.cs
@@ -1,4 +1,4 @@
-// Copyright 1998-2019 Epic Games, Inc. All Rights Reserved.
+// Copyright Epic Games, Inc. All Rights Reserved.
 
 using System.IO;
 
diff --git a/Engine/Plugins/Animation/RigLogic/Source/RigLogicLib/Private/RigLogicLib.cpp b/Engine/Plugins/Animation/RigLogic/Source/RigLogicLib/Private/RigLogicLib.cpp
--- a/Engine/Plugins/Animation/RigLogic/Source/RigLogicLib/Private/RigLogicLib.cpp
+++ b/Engine/Plugins/Animation/RigLogic/Source/RigLogicLib/Private/RigLogicLib.cpp
@@ -1,11 +1,11 @@
-// Copyright 1998-2019 Epic Games, Inc. All Rights Reserved.
+// Copyright Epic Games, Inc. All Rights Reserved.
 
 #ifndef RIGLOGIC_MODULE_DISCARD
 
 #include "RigLogicLib.h"
 
-#include "Runtime/Core/Public/Modules/ModuleInterface.h"
-#include "Runtime/Core/Public/Modules/ModuleManager.h"
+#include "Modules/ModuleInterface.h"
+#include "Modules/ModuleManager.h"
 
 #define LOCTEXT_NAMESPACE "FRigLogicLib"
 
diff --git a/Engine/Plugins/Animation/RigLogic/Source/RigLogicLib/Public/FMemoryResource.h b/Engine/Plugins/Animation/RigLogic/Source/RigLogicLib/Public/FMemoryResource.h
--- a/Engine/Plugins/Animation/RigLogic/Source/RigLogicLib/Public/FMemoryResource.h
+++ b/Engine/Plugins/Animation/RigLogic/Source/RigLogicLib/Public/FMemoryResource.h
@@ -1,4 +1,5 @@
-// Copyright 1998-2019 Epic Games, Inc. All Rights Reserved.
+// Copyright Epic Games, Inc. All Rights Reserved.
+
 #pragma once
 
 #include "HAL/UnrealMemory.h"
diff --git a/Engine/Plugins/Animation/RigLogic/Source/RigLogicLib/Public/RigLogicLib.h b/Engine/Plugins/Animation/RigLogic/Source/RigLogicLib/Public/RigLogicLib.h
--- a/Engine/Plugins/Animation/RigLogic/Source/RigLogicLib/Public/RigLogicLib.h
+++ b/Engine/Plugins/Animation/RigLogic/Source/RigLogicLib/Public/RigLogicLib.h
@@ -1,4 +1,5 @@
-// Copyright 1998-2019 Epic Games, Inc. All Rights Reserved.
+// Copyright Epic Games, Inc. All Rights Reserved.
+
 #pragma once
 
 #ifndef RIGLOGIC_MODULE_DISCARD
@@ -11,8 +12,8 @@ DECLARE_LOG_CATEGORY_EXTERN(LogRigLogicLib, Log, All);
 class FRigLogicLib : public IModuleInterface
 {
 public:
-	virtual void StartupModule() override;
-	virtual void ShutdownModule() override;
+	void StartupModule() override;
+	void ShutdownModule() override;
 };
 
 #endif  // RIGLOGIC_MODULE_DISCARD
diff --git a/Engine/Plugins/Animation/RigLogic/Source/RigLogicLib/RigLogicLib.Build.cs b/Engine/Plugins/Animation/RigLogic/Source/RigLogicLib/RigLogicLib.Build.cs
--- a/Engine/Plugins/Animation/RigLogic/Source/RigLogicLib/RigLogicLib.Build.cs
+++ b/Engine/Plugins/Animation/RigLogic/Source/RigLogicLib/RigLogicLib.Build.cs
@@ -1,4 +1,4 @@
-// Copyright 1998-2019 Epic Games, Inc. All Rights Reserved.
+// Copyright Epic Games, Inc. All Rights Reserved.
 
 using System;
 using System.IO;
diff --git a/Engine/Plugins/Animation/RigLogic/Source/RigLogicLibTest/Private/RigLogicLibTest.cpp b/Engine/Plugins/Animation/RigLogic/Source/RigLogicLibTest/Private/RigLogicLibTest.cpp
--- a/Engine/Plugins/Animation/RigLogic/Source/RigLogicLibTest/Private/RigLogicLibTest.cpp
+++ b/Engine/Plugins/Animation/RigLogic/Source/RigLogicLibTest/Private/RigLogicLibTest.cpp
@@ -1,9 +1,9 @@
-// Copyright 1998-2019 Epic Games, Inc. All Rights Reserved.
+// Copyright Epic Games, Inc. All Rights Reserved.
 
 #include "RigLogicLibTest.h"
 
-#include "Runtime/Core/Public/Modules/ModuleInterface.h"
-#include "Runtime/Core/Public/Modules/ModuleManager.h"
+#include "Modules/ModuleInterface.h"
+#include "Modules/ModuleManager.h"
 
 #define LOCTEXT_NAMESPACE "FRigLogicLibTest"
 
diff --git a/Engine/Plugins/Animation/RigLogic/Source/RigLogicLibTest/Private/RigLogicLibTestSuite.cpp b/Engine/Plugins/Animation/RigLogic/Source/RigLogicLibTest/Private/RigLogicLibTestSuite.cpp
--- a/Engine/Plugins/Animation/RigLogic/Source/RigLogicLibTest/Private/RigLogicLibTestSuite.cpp
+++ b/Engine/Plugins/Animation/RigLogic/Source/RigLogicLibTest/Private/RigLogicLibTestSuite.cpp
@@ -1,4 +1,4 @@
-// Copyright 1998-2019 Epic Games, Inc. All Rights Reserved.
+// Copyright Epic Games, Inc. All Rights Reserved.
 
 #if WITH_DEV_AUTOMATION_TESTS
 #pragma warning(disable:4668)
diff --git a/Engine/Plugins/Animation/RigLogic/Source/RigLogicLibTest/Private/dnatests/Defs.h b/Engine/Plugins/Animation/RigLogic/Source/RigLogicLibTest/Private/dnatests/Defs.h
--- a/Engine/Plugins/Animation/RigLogic/Source/RigLogicLibTest/Private/dnatests/Defs.h
+++ b/Engine/Plugins/Animation/RigLogic/Source/RigLogicLibTest/Private/dnatests/Defs.h
@@ -1,3 +1,5 @@
+// Copyright Epic Games, Inc. All Rights Reserved.
+
 #pragma once
 
 #ifdef _MSC_VER
diff --git a/Engine/Plugins/Animation/RigLogic/Source/RigLogicLibTest/Private/dnatests/FakeStream.h b/Engine/Plugins/Animation/RigLogic/Source/RigLogicLibTest/Private/dnatests/FakeStream.h
--- a/Engine/Plugins/Animation/RigLogic/Source/RigLogicLibTest/Private/dnatests/FakeStream.h
+++ b/Engine/Plugins/Animation/RigLogic/Source/RigLogicLibTest/Private/dnatests/FakeStream.h
@@ -1,3 +1,5 @@
+// Copyright Epic Games, Inc. All Rights Reserved.
+
 #pragma once
 
 #include <trio/Stream.h>
diff --git a/Engine/Plugins/Animation/RigLogic/Source/RigLogicLibTest/Private/dnatests/Fixtures.cpp b/Engine/Plugins/Animation/RigLogic/Source/RigLogicLibTest/Private/dnatests/Fixtures.cpp
--- a/Engine/Plugins/Animation/RigLogic/Source/RigLogicLibTest/Private/dnatests/Fixtures.cpp
+++ b/Engine/Plugins/Animation/RigLogic/Source/RigLogicLibTest/Private/dnatests/Fixtures.cpp
@@ -1,3 +1,5 @@
+// Copyright Epic Games, Inc. All Rights Reserved.
+
 #include "dnatests/Fixtures.h"
 
 #include "dna/Reader.h"
diff --git a/Engine/Plugins/Animation/RigLogic/Source/RigLogicLibTest/Private/dnatests/Fixtures.h b/Engine/Plugins/Animation/RigLogic/Source/RigLogicLibTest/Private/dnatests/Fixtures.h
--- a/Engine/Plugins/Animation/RigLogic/Source/RigLogicLibTest/Private/dnatests/Fixtures.h
+++ b/Engine/Plugins/Animation/RigLogic/Source/RigLogicLibTest/Private/dnatests/Fixtures.h
@@ -1,3 +1,5 @@
+// Copyright Epic Games, Inc. All Rights Reserved.
+
 #pragma once
 
 #include "dna/DNA.h"
diff --git a/Engine/Plugins/Animation/RigLogic/Source/RigLogicLibTest/Private/dnatests/TestFilteredInputArchive.cpp b/Engine/Plugins/Animation/RigLogic/Source/RigLogicLibTest/Private/dnatests/TestFilteredInputArchive.cpp
--- a/Engine/Plugins/Animation/RigLogic/Source/RigLogicLibTest/Private/dnatests/TestFilteredInputArchive.cpp
+++ b/Engine/Plugins/Animation/RigLogic/Source/RigLogicLibTest/Private/dnatests/TestFilteredInputArchive.cpp
@@ -1,3 +1,5 @@
+// Copyright Epic Games, Inc. All Rights Reserved.
+
 #include "dnatests/Defs.h"
 #include "dnatests/FakeStream.h"
 #include "dnatests/Fixtures.h"
diff --git a/Engine/Plugins/Animation/RigLogic/Source/RigLogicLibTest/Private/dnatests/TestLODMapping.cpp b/Engine/Plugins/Animation/RigLogic/Source/RigLogicLibTest/Private/dnatests/TestLODMapping.cpp
--- a/Engine/Plugins/Animation/RigLogic/Source/RigLogicLibTest/Private/dnatests/TestLODMapping.cpp
+++ b/Engine/Plugins/Animation/RigLogic/Source/RigLogicLibTest/Private/dnatests/TestLODMapping.cpp
@@ -1,3 +1,5 @@
+// Copyright Epic Games, Inc. All Rights Reserved.
+
 #include "dnatests/Defs.h"
 #include "dnatests/Fixtures.h"
 
diff --git a/Engine/Plugins/Animation/RigLogic/Source/RigLogicLibTest/Private/dnatests/TestStreamReadWriteIntegration.cpp b/Engine/Plugins/Animation/RigLogic/Source/RigLogicLibTest/Private/dnatests/TestStreamReadWriteIntegration.cpp
--- a/Engine/Plugins/Animation/RigLogic/Source/RigLogicLibTest/Private/dnatests/TestStreamReadWriteIntegration.cpp
+++ b/Engine/Plugins/Animation/RigLogic/Source/RigLogicLibTest/Private/dnatests/TestStreamReadWriteIntegration.cpp
@@ -1,3 +1,5 @@
+// Copyright Epic Games, Inc. All Rights Reserved.
+
 #include "dnatests/TestStreamReadWriteIntegration.h"
 
 #include "dnatests/Defs.h"
diff --git a/Engine/Plugins/Animation/RigLogic/Source/RigLogicLibTest/Private/dnatests/TestStreamReadWriteIntegration.h b/Engine/Plugins/Animation/RigLogic/Source/RigLogicLibTest/Private/dnatests/TestStreamReadWriteIntegration.h
--- a/Engine/Plugins/Animation/RigLogic/Source/RigLogicLibTest/Private/dnatests/TestStreamReadWriteIntegration.h
+++ b/Engine/Plugins/Animation/RigLogic/Source/RigLogicLibTest/Private/dnatests/TestStreamReadWriteIntegration.h
@@ -1,3 +1,5 @@
+// Copyright Epic Games, Inc. All Rights Reserved.
+
 #pragma once
 
 #include "dnatests/Defs.h"
diff --git a/Engine/Plugins/Animation/RigLogic/Source/RigLogicLibTest/Private/dnatests/TestStreamWriter.cpp b/Engine/Plugins/Animation/RigLogic/Source/RigLogicLibTest/Private/dnatests/TestStreamWriter.cpp
--- a/Engine/Plugins/Animation/RigLogic/Source/RigLogicLibTest/Private/dnatests/TestStreamWriter.cpp
+++ b/Engine/Plugins/Animation/RigLogic/Source/RigLogicLibTest/Private/dnatests/TestStreamWriter.cpp
@@ -1,3 +1,5 @@
+// Copyright Epic Games, Inc. All Rights Reserved.
+
 #include "dnatests/TestStreamWriter.h"
 
 TEST_F(StreamWriterTest, SetName) {
diff --git a/Engine/Plugins/Animation/RigLogic/Source/RigLogicLibTest/Private/dnatests/TestStreamWriter.h b/Engine/Plugins/Animation/RigLogic/Source/RigLogicLibTest/Private/dnatests/TestStreamWriter.h
--- a/Engine/Plugins/Animation/RigLogic/Source/RigLogicLibTest/Private/dnatests/TestStreamWriter.h
+++ b/Engine/Plugins/Animation/RigLogic/Source/RigLogicLibTest/Private/dnatests/TestStreamWriter.h
@@ -1,3 +1,5 @@
+// Copyright Epic Games, Inc. All Rights Reserved.
+
 #pragma once
 
 #include "dnatests/Defs.h"
diff --git a/Engine/Plugins/Animation/RigLogic/Source/RigLogicLibTest/Private/rltests/dna/FakeReader.h b/Engine/Plugins/Animation/RigLogic/Source/RigLogicLibTest/Private/rltests/dna/FakeReader.h
--- a/Engine/Plugins/Animation/RigLogic/Source/RigLogicLibTest/Private/rltests/dna/FakeReader.h
+++ b/Engine/Plugins/Animation/RigLogic/Source/RigLogicLibTest/Private/rltests/dna/FakeReader.h
@@ -9,8 +9,8 @@
 #include <cstdint>
 
 #if defined(_MSC_VER) && !defined(__clang__)
-#pragma warning(push)
-#pragma warning(disable:4100) // unreferenced formal parameter
+    #pragma warning(push)
+    #pragma warning(disable : 4100)
 #endif
 
 namespace dna {
@@ -446,5 +446,5 @@ class FakeReader : public Reader {
 }  // namespace dna
 
 #if defined(_MSC_VER) && !defined(__clang__)
-#pragma warning(pop)
+    #pragma warning(pop)
 #endif
diff --git a/Engine/Plugins/Animation/RigLogic/Source/RigLogicLibTest/Public/RigLogicLibTest.h b/Engine/Plugins/Animation/RigLogic/Source/RigLogicLibTest/Public/RigLogicLibTest.h
--- a/Engine/Plugins/Animation/RigLogic/Source/RigLogicLibTest/Public/RigLogicLibTest.h
+++ b/Engine/Plugins/Animation/RigLogic/Source/RigLogicLibTest/Public/RigLogicLibTest.h
@@ -1,4 +1,5 @@
-// Copyright 1998-2019 Epic Games, Inc. All Rights Reserved.
+// Copyright Epic Games, Inc. All Rights Reserved.
+
 #pragma once
 
 #include "CoreMinimal.h"
diff --git a/Engine/Plugins/Animation/RigLogic/Source/RigLogicLibTest/RigLogicLibTest.Build.cs b/Engine/Plugins/Animation/RigLogic/Source/RigLogicLibTest/RigLogicLibTest.Build.cs
--- a/Engine/Plugins/Animation/RigLogic/Source/RigLogicLibTest/RigLogicLibTest.Build.cs
+++ b/Engine/Plugins/Animation/RigLogic/Source/RigLogicLibTest/RigLogicLibTest.Build.cs
@@ -1,4 +1,4 @@
-// Copyright 1998-2019 Epic Games, Inc. All Rights Reserved.
+// Copyright Epic Games, Inc. All Rights Reserved.
 
 using UnrealBuildTool;
 using System;
diff --git a/Engine/Plugins/Animation/RigLogic/Source/RigLogicModule/Private/DNAAsset.cpp b/Engine/Plugins/Animation/RigLogic/Source/RigLogicModule/Private/DNAAsset.cpp
--- a/Engine/Plugins/Animation/RigLogic/Source/RigLogicModule/Private/DNAAsset.cpp
+++ b/Engine/Plugins/Animation/RigLogic/Source/RigLogicModule/Private/DNAAsset.cpp
@@ -1,4 +1,4 @@
-// Copyright 1998-2019 Epic Games, Inc. All Rights Reserved.
+// Copyright Epic Games, Inc. All Rights Reserved.
 
 #include "DNAAsset.h"
 #include "DNAUtils.h"
@@ -12,11 +12,11 @@
     #include "EditorFramework/AssetImportData.h"
 #endif
 #include "Engine/AssetUserData.h"
-#include "Runtime/Core/Public/Serialization/BufferArchive.h"
-#include "Runtime/Core/Public/Serialization/MemoryReader.h"
-#include "Runtime/Core/Public/Misc/Paths.h"
-#include "Runtime/Core/Public/Misc/FileHelper.h"
-#include "Runtime/CoreUObject/Public/UObject/UObjectGlobals.h"
+#include "Serialization/BufferArchive.h"
+#include "Serialization/MemoryReader.h"
+#include "Misc/Paths.h"
+#include "Misc/FileHelper.h"
+#include "UObject/UObjectGlobals.h"
 
 #include "riglogic/RigLogic.h"
 
diff --git a/Engine/Plugins/Animation/RigLogic/Source/RigLogicModule/Private/DNAReader.cpp b/Engine/Plugins/Animation/RigLogic/Source/RigLogicModule/Private/DNAReader.cpp
--- a/Engine/Plugins/Animation/RigLogic/Source/RigLogicModule/Private/DNAReader.cpp
+++ b/Engine/Plugins/Animation/RigLogic/Source/RigLogicModule/Private/DNAReader.cpp
@@ -1,4 +1,4 @@
-// Copyright 1998-2019 Epic Games, Inc. All Rights Reserved.
+// Copyright Epic Games, Inc. All Rights Reserved.
 
 #include "DNAReader.h"
 
diff --git a/Engine/Plugins/Animation/RigLogic/Source/RigLogicModule/Private/DNAToSkelMeshMap.cpp b/Engine/Plugins/Animation/RigLogic/Source/RigLogicModule/Private/DNAToSkelMeshMap.cpp
--- a/Engine/Plugins/Animation/RigLogic/Source/RigLogicModule/Private/DNAToSkelMeshMap.cpp
+++ b/Engine/Plugins/Animation/RigLogic/Source/RigLogicModule/Private/DNAToSkelMeshMap.cpp
@@ -1,4 +1,4 @@
-// Copyright 1998-2019 Epic Games, Inc. All Rights Reserved.
+// Copyright Epic Games, Inc. All Rights Reserved.
 
 #include "DNAToSkelMeshMap.h"
 
diff --git a/Engine/Plugins/Animation/RigLogic/Source/RigLogicModule/Private/DNAUtils.cpp b/Engine/Plugins/Animation/RigLogic/Source/RigLogicModule/Private/DNAUtils.cpp
--- a/Engine/Plugins/Animation/RigLogic/Source/RigLogicModule/Private/DNAUtils.cpp
+++ b/Engine/Plugins/Animation/RigLogic/Source/RigLogicModule/Private/DNAUtils.cpp
@@ -1,4 +1,4 @@
-// Copyright 1998-2019 Epic Games, Inc. All Rights Reserved.
+// Copyright Epic Games, Inc. All Rights Reserved.
 
 #include "DNAUtils.h"
 
diff --git a/Engine/Plugins/Animation/RigLogic/Source/RigLogicModule/Private/RigInstance.cpp b/Engine/Plugins/Animation/RigLogic/Source/RigLogicModule/Private/RigInstance.cpp
--- a/Engine/Plugins/Animation/RigLogic/Source/RigLogicModule/Private/RigInstance.cpp
+++ b/Engine/Plugins/Animation/RigLogic/Source/RigLogicModule/Private/RigInstance.cpp
@@ -1,4 +1,4 @@
-// Copyright 1998-2019 Epic Games, Inc. All Rights Reserved.
+// Copyright Epic Games, Inc. All Rights Reserved.
 
 #include "RigInstance.h"
 
diff --git a/Engine/Plugins/Animation/RigLogic/Source/RigLogicModule/Private/RigLogic.cpp b/Engine/Plugins/Animation/RigLogic/Source/RigLogicModule/Private/RigLogic.cpp
--- a/Engine/Plugins/Animation/RigLogic/Source/RigLogicModule/Private/RigLogic.cpp
+++ b/Engine/Plugins/Animation/RigLogic/Source/RigLogicModule/Private/RigLogic.cpp
@@ -1,4 +1,4 @@
-// Copyright 1998-2019 Epic Games, Inc. All Rights Reserved.
+// Copyright Epic Games, Inc. All Rights Reserved.
 
 #include "RigLogic.h"
 
diff --git a/Engine/Plugins/Animation/RigLogic/Source/RigLogicModule/Private/RigLogicMemoryStream.cpp b/Engine/Plugins/Animation/RigLogic/Source/RigLogicModule/Private/RigLogicMemoryStream.cpp
--- a/Engine/Plugins/Animation/RigLogic/Source/RigLogicModule/Private/RigLogicMemoryStream.cpp
+++ b/Engine/Plugins/Animation/RigLogic/Source/RigLogicModule/Private/RigLogicMemoryStream.cpp
@@ -1,3 +1,5 @@
+// Copyright Epic Games, Inc. All Rights Reserved.
+
 #include "RigLogicMemoryStream.h"
 
 FRigLogicMemoryStream::FRigLogicMemoryStream(TArray<uint8>* Buffer)
diff --git a/Engine/Plugins/Animation/RigLogic/Source/RigLogicModule/Private/RigLogicMemoryStream.h b/Engine/Plugins/Animation/RigLogic/Source/RigLogicModule/Private/RigLogicMemoryStream.h
--- a/Engine/Plugins/Animation/RigLogic/Source/RigLogicModule/Private/RigLogicMemoryStream.h
+++ b/Engine/Plugins/Animation/RigLogic/Source/RigLogicModule/Private/RigLogicMemoryStream.h
@@ -1,4 +1,4 @@
-// Copyright 1998-2019 Epic Games, Inc. All Rights Reserved.
+// Copyright Epic Games, Inc. All Rights Reserved.
 
 #pragma once
 
diff --git a/Engine/Plugins/Animation/RigLogic/Source/RigLogicModule/Private/RigLogicModule.cpp b/Engine/Plugins/Animation/RigLogic/Source/RigLogicModule/Private/RigLogicModule.cpp
--- a/Engine/Plugins/Animation/RigLogic/Source/RigLogicModule/Private/RigLogicModule.cpp
+++ b/Engine/Plugins/Animation/RigLogic/Source/RigLogicModule/Private/RigLogicModule.cpp
@@ -1,16 +1,16 @@
-// Copyright 1998-2019 Epic Games, Inc. All Rights Reserved.
+// Copyright Epic Games, Inc. All Rights Reserved.
 
 #include "RigLogicModule.h"
 
 #include "RigUnit_RigLogic.h"
 
 #include "Interfaces/IPluginManager.h"
-#include "Runtime/Core/Public/Misc/Paths.h"
+#include "Misc/Paths.h"
 
-#include "Runtime/Core/Public/Modules/ModuleInterface.h"
-#include "Runtime/Core/Public/Modules/ModuleManager.h"
+#include "Modules/ModuleInterface.h"
+#include "Modules/ModuleManager.h"
 
-#include "Runtime/Core/Public/HAL/PlatformProcess.h"
+#include "HAL/PlatformProcess.h"
 
 #define LOCTEXT_NAMESPACE "FRigLogicModule"
 
diff --git a/Engine/Plugins/Animation/RigLogic/Source/RigLogicModule/Private/RigUnit_RigLogic.cpp b/Engine/Plugins/Animation/RigLogic/Source/RigLogicModule/Private/RigUnit_RigLogic.cpp
--- a/Engine/Plugins/Animation/RigLogic/Source/RigLogicModule/Private/RigUnit_RigLogic.cpp
+++ b/Engine/Plugins/Animation/RigLogic/Source/RigLogicModule/Private/RigUnit_RigLogic.cpp
@@ -1,4 +1,4 @@
-// Copyright 1998-2019 Epic Games, Inc. All Rights Reserved.
+// Copyright Epic Games, Inc. All Rights Reserved.
 
 #include "RigUnit_RigLogic.h"
 
diff --git a/Engine/Plugins/Animation/RigLogic/Source/RigLogicModule/Private/RigUnit_RigLogic_Test.cpp b/Engine/Plugins/Animation/RigLogic/Source/RigLogicModule/Private/RigUnit_RigLogic_Test.cpp
--- a/Engine/Plugins/Animation/RigLogic/Source/RigLogicModule/Private/RigUnit_RigLogic_Test.cpp
+++ b/Engine/Plugins/Animation/RigLogic/Source/RigLogicModule/Private/RigUnit_RigLogic_Test.cpp
@@ -1,3 +1,5 @@
+// Copyright Epic Games, Inc. All Rights Reserved.
+
 #if WITH_DEV_AUTOMATION_TESTS
 	#include "RigUnit_RigLogic_Test.h"
 #endif
@@ -5,12 +7,11 @@
 #include "RigUnit_RigLogic.h"
 #include "Units/RigUnitContext.h"
 #include "ControlRig.h"
-#include "riglogic/RigLogic.h"
-#include "Runtime/Core/Public/Math/TransformNonVectorized.h"
+#include "Math/TransformNonVectorized.h"
 #include "SkelMeshDNAUtils.h"
+#include "DNAUtils.h"
 
 #include "riglogic/RigLogic.h"
-#include "DNAUtils.h"
 
 #if WITH_DEV_AUTOMATION_TESTS
 #include "Units/RigUnitTest.h"
diff --git a/Engine/Plugins/Animation/RigLogic/Source/RigLogicModule/Private/SkelMeshDNAReader.cpp b/Engine/Plugins/Animation/RigLogic/Source/RigLogicModule/Private/SkelMeshDNAReader.cpp
--- a/Engine/Plugins/Animation/RigLogic/Source/RigLogicModule/Private/SkelMeshDNAReader.cpp
+++ b/Engine/Plugins/Animation/RigLogic/Source/RigLogicModule/Private/SkelMeshDNAReader.cpp
@@ -1,4 +1,4 @@
-// Copyright 1998-2019 Epic Games, Inc. All Rights Reserved.
+// Copyright Epic Games, Inc. All Rights Reserved.
 
 #include "SkelMeshDNAReader.h"
 
@@ -9,11 +9,11 @@
 #include "EditorFramework/AssetImportData.h"
 #endif
 #include "Engine/AssetUserData.h"
-#include "Runtime/Core/Public/Serialization/BufferArchive.h"
-#include "Runtime/Core/Public/Serialization/MemoryReader.h"
-#include "Runtime/Core/Public/Misc/Paths.h"
-#include "Runtime/Core/Public/Misc/FileHelper.h"
-#include "Runtime/CoreUObject/Public/UObject/UObjectGlobals.h"
+#include "Serialization/BufferArchive.h"
+#include "Serialization/MemoryReader.h"
+#include "Misc/Paths.h"
+#include "Misc/FileHelper.h"
+#include "UObject/UObjectGlobals.h"
 
 DEFINE_LOG_CATEGORY(LogSkelMeshDNAReader);
 
diff --git a/Engine/Plugins/Animation/RigLogic/Source/RigLogicModule/Private/SkelMeshDNAReader.h b/Engine/Plugins/Animation/RigLogic/Source/RigLogicModule/Private/SkelMeshDNAReader.h
--- a/Engine/Plugins/Animation/RigLogic/Source/RigLogicModule/Private/SkelMeshDNAReader.h
+++ b/Engine/Plugins/Animation/RigLogic/Source/RigLogicModule/Private/SkelMeshDNAReader.h
@@ -1,4 +1,4 @@
-// Copyright 1998-2019 Epic Games, Inc. All Rights Reserved.
+// Copyright Epic Games, Inc. All Rights Reserved.
 
 #pragma once
 
@@ -7,8 +7,8 @@
 #include "CoreMinimal.h"
 #include "Engine/AssetUserData.h"
 #include "Interfaces/Interface_AssetUserData.h"
-#include "Runtime/Core/Public/HAL/FileManager.h"
-#include "Runtime/Core/Public/Serialization/Archive.h"
+#include "HAL/FileManager.h"
+#include "Serialization/Archive.h"
 #include "UObject/NoExportTypes.h"
 
 class UDNAAsset;
diff --git a/Engine/Plugins/Animation/RigLogic/Source/RigLogicModule/Private/SkelMeshDNAUtils.cpp b/Engine/Plugins/Animation/RigLogic/Source/RigLogicModule/Private/SkelMeshDNAUtils.cpp
--- a/Engine/Plugins/Animation/RigLogic/Source/RigLogicModule/Private/SkelMeshDNAUtils.cpp
+++ b/Engine/Plugins/Animation/RigLogic/Source/RigLogicModule/Private/SkelMeshDNAUtils.cpp
@@ -1,4 +1,4 @@
-// Copyright 1998-2019 Epic Games, Inc. All Rights Reserved.
+// Copyright Epic Games, Inc. All Rights Reserved.
 
 #include "SkelMeshDNAUtils.h"
 #include "DNAToSkelMeshMap.h"
diff --git a/Engine/Plugins/Animation/RigLogic/Source/RigLogicModule/Public/DNAAsset.h b/Engine/Plugins/Animation/RigLogic/Source/RigLogicModule/Public/DNAAsset.h
--- a/Engine/Plugins/Animation/RigLogic/Source/RigLogicModule/Public/DNAAsset.h
+++ b/Engine/Plugins/Animation/RigLogic/Source/RigLogicModule/Public/DNAAsset.h
@@ -1,4 +1,4 @@
-// Copyright 1998-2019 Epic Games, Inc. All Rights Reserved.
+// Copyright Epic Games, Inc. All Rights Reserved.
 
 #pragma once
 
diff --git a/Engine/Plugins/Animation/RigLogic/Source/RigLogicModule/Public/DNACommon.h b/Engine/Plugins/Animation/RigLogic/Source/RigLogicModule/Public/DNACommon.h
--- a/Engine/Plugins/Animation/RigLogic/Source/RigLogicModule/Public/DNACommon.h
+++ b/Engine/Plugins/Animation/RigLogic/Source/RigLogicModule/Public/DNACommon.h
@@ -1,4 +1,4 @@
-// Copyright 1998-2019 Epic Games, Inc. All Rights Reserved.
+// Copyright Epic Games, Inc. All Rights Reserved.
 
 #pragma once
 
diff --git a/Engine/Plugins/Animation/RigLogic/Source/RigLogicModule/Public/DNAReader.h b/Engine/Plugins/Animation/RigLogic/Source/RigLogicModule/Public/DNAReader.h
--- a/Engine/Plugins/Animation/RigLogic/Source/RigLogicModule/Public/DNAReader.h
+++ b/Engine/Plugins/Animation/RigLogic/Source/RigLogicModule/Public/DNAReader.h
@@ -1,4 +1,4 @@
-// Copyright 1998-2019 Epic Games, Inc. All Rights Reserved.
+// Copyright Epic Games, Inc. All Rights Reserved.
 
 #pragma once
 
diff --git a/Engine/Plugins/Animation/RigLogic/Source/RigLogicModule/Public/DNAReaderAdapter.h b/Engine/Plugins/Animation/RigLogic/Source/RigLogicModule/Public/DNAReaderAdapter.h
--- a/Engine/Plugins/Animation/RigLogic/Source/RigLogicModule/Public/DNAReaderAdapter.h
+++ b/Engine/Plugins/Animation/RigLogic/Source/RigLogicModule/Public/DNAReaderAdapter.h
@@ -1,4 +1,5 @@
-// Copyright 1998-2019 Epic Games, Inc. All Rights Reserved.
+// Copyright Epic Games, Inc. All Rights Reserved.
+
 #pragma once
 
 #include "DNAReader.h"
diff --git a/Engine/Plugins/Animation/RigLogic/Source/RigLogicModule/Public/DNAReaderAdapterImpl.h b/Engine/Plugins/Animation/RigLogic/Source/RigLogicModule/Public/DNAReaderAdapterImpl.h
--- a/Engine/Plugins/Animation/RigLogic/Source/RigLogicModule/Public/DNAReaderAdapterImpl.h
+++ b/Engine/Plugins/Animation/RigLogic/Source/RigLogicModule/Public/DNAReaderAdapterImpl.h
@@ -1,4 +1,5 @@
-// Copyright 1998-2019 Epic Games, Inc. All Rights Reserved.
+// Copyright Epic Games, Inc. All Rights Reserved.
+
 #pragma once
 
 template <class TWrappedReader>
diff --git a/Engine/Plugins/Animation/RigLogic/Source/RigLogicModule/Public/DNAToSkelMeshMap.h b/Engine/Plugins/Animation/RigLogic/Source/RigLogicModule/Public/DNAToSkelMeshMap.h
--- a/Engine/Plugins/Animation/RigLogic/Source/RigLogicModule/Public/DNAToSkelMeshMap.h
+++ b/Engine/Plugins/Animation/RigLogic/Source/RigLogicModule/Public/DNAToSkelMeshMap.h
@@ -1,10 +1,10 @@
-// Copyright 1998-2019 Epic Games, Inc. All Rights Reserved.
+// Copyright Epic Games, Inc. All Rights Reserved.
 
 #pragma once
 
 #include "CoreMinimal.h"
-#include "Runtime/Core/Public/HAL/FileManager.h"
-#include "Runtime/Core/Public/Serialization/Archive.h"
+#include "HAL/FileManager.h"
+#include "Serialization/Archive.h"
 #include "UObject/NoExportTypes.h"
 
 #include "DNAReader.h"
diff --git a/Engine/Plugins/Animation/RigLogic/Source/RigLogicModule/Public/DNAUtils.h b/Engine/Plugins/Animation/RigLogic/Source/RigLogicModule/Public/DNAUtils.h
--- a/Engine/Plugins/Animation/RigLogic/Source/RigLogicModule/Public/DNAUtils.h
+++ b/Engine/Plugins/Animation/RigLogic/Source/RigLogicModule/Public/DNAUtils.h
@@ -1,4 +1,4 @@
-// Copyright 1998-2019 Epic Games, Inc. All Rights Reserved.
+// Copyright Epic Games, Inc. All Rights Reserved.
 
 #pragma once
 
diff --git a/Engine/Plugins/Animation/RigLogic/Source/RigLogicModule/Public/RigInstance.h b/Engine/Plugins/Animation/RigLogic/Source/RigLogicModule/Public/RigInstance.h
--- a/Engine/Plugins/Animation/RigLogic/Source/RigLogicModule/Public/RigInstance.h
+++ b/Engine/Plugins/Animation/RigLogic/Source/RigLogicModule/Public/RigInstance.h
@@ -1,4 +1,4 @@
-// Copyright 1998-2019 Epic Games, Inc. All Rights Reserved.
+// Copyright Epic Games, Inc. All Rights Reserved.
 
 #pragma once
 
diff --git a/Engine/Plugins/Animation/RigLogic/Source/RigLogicModule/Public/RigLogic.h b/Engine/Plugins/Animation/RigLogic/Source/RigLogicModule/Public/RigLogic.h
--- a/Engine/Plugins/Animation/RigLogic/Source/RigLogicModule/Public/RigLogic.h
+++ b/Engine/Plugins/Animation/RigLogic/Source/RigLogicModule/Public/RigLogic.h
@@ -1,4 +1,4 @@
-// Copyright 1998-2019 Epic Games, Inc. All Rights Reserved.
+// Copyright Epic Games, Inc. All Rights Reserved.
 
 #pragma once
 
diff --git a/Engine/Plugins/Animation/RigLogic/Source/RigLogicModule/Public/RigLogicModule.h b/Engine/Plugins/Animation/RigLogic/Source/RigLogicModule/Public/RigLogicModule.h
--- a/Engine/Plugins/Animation/RigLogic/Source/RigLogicModule/Public/RigLogicModule.h
+++ b/Engine/Plugins/Animation/RigLogic/Source/RigLogicModule/Public/RigLogicModule.h
@@ -1,4 +1,4 @@
-// Copyright 1998-2019 Epic Games, Inc. All Rights Reserved.
+// Copyright Epic Games, Inc. All Rights Reserved.
 
 #pragma once
 
@@ -7,9 +7,9 @@
 
 DECLARE_LOG_CATEGORY_EXTERN(LogRigLogic, Log, All);
 
-struct RIGLOGICMODULE_API FRigLogicModule: IModuleInterface
+class RIGLOGICMODULE_API FRigLogicModule: public IModuleInterface
 {
 public:
-	virtual void StartupModule() override;
-	virtual void ShutdownModule() override;
+	void StartupModule() override;
+	void ShutdownModule() override;
 };
diff --git a/Engine/Plugins/Animation/RigLogic/Source/RigLogicModule/Public/RigUnit_RigLogic.h b/Engine/Plugins/Animation/RigLogic/Source/RigLogicModule/Public/RigUnit_RigLogic.h
--- a/Engine/Plugins/Animation/RigLogic/Source/RigLogicModule/Public/RigUnit_RigLogic.h
+++ b/Engine/Plugins/Animation/RigLogic/Source/RigLogicModule/Public/RigUnit_RigLogic.h
@@ -1,4 +1,4 @@
-// Copyright 1998-2019 Epic Games, Inc. All Rights Reserved.
+// Copyright Epic Games, Inc. All Rights Reserved.
 
 #pragma once
 
@@ -6,8 +6,8 @@
 #include "RigLogic.h"
 
 #include "ControlRig.h"
-#include "../Public/DNAAsset.h"
-#include "Runtime/Engine/Classes/Materials/Material.h"
+#include "DNAAsset.h"
+#include "Materials/Material.h"
 #include "Units/RigUnit.h"
 
 #include "RigUnit_RigLogic.generated.h"
diff --git a/Engine/Plugins/Animation/RigLogic/Source/RigLogicModule/Public/RigUnit_RigLogic_Test.h b/Engine/Plugins/Animation/RigLogic/Source/RigLogicModule/Public/RigUnit_RigLogic_Test.h
--- a/Engine/Plugins/Animation/RigLogic/Source/RigLogicModule/Public/RigUnit_RigLogic_Test.h
+++ b/Engine/Plugins/Animation/RigLogic/Source/RigLogicModule/Public/RigUnit_RigLogic_Test.h
@@ -1,3 +1,5 @@
+// Copyright Epic Games, Inc. All Rights Reserved.
+
 #pragma once
 
 #include "CoreTypes.h"
diff --git a/Engine/Plugins/Animation/RigLogic/Source/RigLogicModule/Public/SkelMeshDNAUtils.h b/Engine/Plugins/Animation/RigLogic/Source/RigLogicModule/Public/SkelMeshDNAUtils.h
--- a/Engine/Plugins/Animation/RigLogic/Source/RigLogicModule/Public/SkelMeshDNAUtils.h
+++ b/Engine/Plugins/Animation/RigLogic/Source/RigLogicModule/Public/SkelMeshDNAUtils.h
@@ -1,9 +1,9 @@
-// Copyright 1998-2019 Epic Games, Inc. All Rights Reserved.
+// Copyright Epic Games, Inc. All Rights Reserved.
 
 #pragma once
 
 #include "DNAToSkelMeshMap.h"
-#include "Runtime/Engine/Classes/Engine/SkeletalMesh.h"
+#include "Engine/SkeletalMesh.h"
 #include "SkelMeshDNAUtils.generated.h"
 
 DECLARE_LOG_CATEGORY_EXTERN(LogDNAUtils, Log, All);
diff --git a/Engine/Plugins/Animation/RigLogic/Source/RigLogicModule/Public/TransformArrayView.h b/Engine/Plugins/Animation/RigLogic/Source/RigLogicModule/Public/TransformArrayView.h
--- a/Engine/Plugins/Animation/RigLogic/Source/RigLogicModule/Public/TransformArrayView.h
+++ b/Engine/Plugins/Animation/RigLogic/Source/RigLogicModule/Public/TransformArrayView.h
@@ -1,4 +1,4 @@
-// Copyright 1998-2019 Epic Games, Inc. All Rights Reserved.
+// Copyright Epic Games, Inc. All Rights Reserved.
 
 #pragma once
 
diff --git a/Engine/Plugins/Animation/RigLogic/Source/RigLogicModule/RigLogicModule.Build.cs b/Engine/Plugins/Animation/RigLogic/Source/RigLogicModule/RigLogicModule.Build.cs
--- a/Engine/Plugins/Animation/RigLogic/Source/RigLogicModule/RigLogicModule.Build.cs
+++ b/Engine/Plugins/Animation/RigLogic/Source/RigLogicModule/RigLogicModule.Build.cs
@@ -1,4 +1,4 @@
-// Copyright 1998-2019 Epic Games, Inc. All Rights Reserved.
+// Copyright Epic Games, Inc. All Rights Reserved.
 
 using System;
 using System.IO;
