Merging using RigLogic_CharacterTech_to_UE5Main at CL 13990315

#rnx

diff --git a/Engine/Plugins/Animation/RigLogic/Source/RigLogicLib/Private/dna/DNA.h b/Engine/Plugins/Animation/RigLogic/Source/RigLogicLib/Private/dna/DNA.h
--- a/Engine/Plugins/Animation/RigLogic/Source/RigLogicLib/Private/dna/DNA.h
+++ b/Engine/Plugins/Animation/RigLogic/Source/RigLogicLib/Private/dna/DNA.h
@@ -711,6 +711,7 @@ struct DNA {
         archive(signature, version);
         if (signature.matches() && version.matches()) {
             archive(sections, descriptor, definition, behavior, geometry, eof);
+            assert(eof.matches());
         }
     }
 
diff --git a/Engine/Plugins/Animation/RigLogic/Source/RigLogicLib/Private/dna/stream/FilteredInputArchive.cpp b/Engine/Plugins/Animation/RigLogic/Source/RigLogicLib/Private/dna/stream/FilteredInputArchive.cpp
--- a/Engine/Plugins/Animation/RigLogic/Source/RigLogicLib/Private/dna/stream/FilteredInputArchive.cpp
+++ b/Engine/Plugins/Animation/RigLogic/Source/RigLogicLib/Private/dna/stream/FilteredInputArchive.cpp
@@ -207,26 +207,40 @@ void FilteredInputArchive::process(RawAnimatedMaps& dest) {
 }
 
 void FilteredInputArchive::process(RawGeometry& dest) {
-    if (contains(layerBitmask, DataLayerBitmask::GeometryRest)) {
-        process(dest.marker);
-        if (!lodConstraint.hasImpactOn(unconstrainedLODCount)) {
-            process(dest.meshes);
-            return;
-        }
-        // Perform filtered load only if a different maxLOD is set
+    process(dest.marker);
+
+    if (!contains(layerBitmask, DataLayerBitmask::GeometryRest)) {
+        // As mesh sizes are variable, iterate over each of them, reading only the mesh
+        // offsets and jumping over the actual data of the meshes.
+        // This will correctly position the underlying stream (end of geometry layer),
+        // while still not reading the data.
         const auto meshCount = processSize();
-        dest.meshes.reserve(meshCount);
+        decltype(RawMesh::offset) meshOffset{};
+        decltype(RawMesh::marker) meshMarker{meshOffset};
         for (std::uint16_t i = {}; i < meshCount; ++i) {
-            RawMesh mesh{memRes};
-            // Check if the mesh indices filtered for the current maxLOD permit loading this mesh
-            if (MeshFilter::passes(i)) {
-                process(mesh);
-                dest.meshes.push_back(std::move(mesh));
-            } else {
-                // Jump over the whole section of data related to this mesh
-                process(mesh.offset);
-                process(mesh.marker);
-            }
+            process(meshOffset);
+            process(meshMarker);
+        }
+        return;
+    }
+
+    if (!lodConstraint.hasImpactOn(unconstrainedLODCount)) {
+        process(dest.meshes);
+        return;
+    }
+    // Perform filtered load only if a different maxLOD is set
+    const auto meshCount = processSize();
+    dest.meshes.reserve(meshCount);
+    for (std::uint16_t i = {}; i < meshCount; ++i) {
+        RawMesh mesh{memRes};
+        // Check if the mesh indices filtered for the current maxLOD permit loading this mesh
+        if (MeshFilter::passes(i)) {
+            process(mesh);
+            dest.meshes.push_back(std::move(mesh));
+        } else {
+            // Jump over the whole section of data related to this mesh
+            process(mesh.offset);
+            process(mesh.marker);
         }
     }
 }
diff --git a/Engine/Plugins/Animation/RigLogic/Source/RigLogicLibTest/Private/dnatests/Fixtures.cpp b/Engine/Plugins/Animation/RigLogic/Source/RigLogicLibTest/Private/dnatests/Fixtures.cpp
--- a/Engine/Plugins/Animation/RigLogic/Source/RigLogicLibTest/Private/dnatests/Fixtures.cpp
+++ b/Engine/Plugins/Animation/RigLogic/Source/RigLogicLibTest/Private/dnatests/Fixtures.cpp
@@ -883,7 +883,7 @@ const unsigned char geometry[] = {
     0x00, 0x00, 0x00, 0x02,  // Vertex position: 2
     0x00, 0x06,  // Blend shape index in Definition
     // Mesh-2
-    0x00, 0x00, 0x0d, 0x17,  // Mesh-2 offset
+    0x00, 0x00, 0x0d, 0x18,  // Mesh-2 offset
     0x00, 0x00, 0x00, 0x03,  // Vertex positions X values length
     0x3f, 0x80, 0x00, 0x00,  // 1.0f
     0x40, 0x00, 0x00, 0x00,  // 2.0f
diff --git a/Engine/Plugins/Animation/RigLogic/Source/RigLogicLibTest/Private/rltests/dna/DNAFixtures.cpp b/Engine/Plugins/Animation/RigLogic/Source/RigLogicLibTest/Private/rltests/dna/DNAFixtures.cpp
--- a/Engine/Plugins/Animation/RigLogic/Source/RigLogicLibTest/Private/rltests/dna/DNAFixtures.cpp
+++ b/Engine/Plugins/Animation/RigLogic/Source/RigLogicLibTest/Private/rltests/dna/DNAFixtures.cpp
@@ -884,7 +884,7 @@ const unsigned char geometry[] = {
     0x00, 0x00, 0x00, 0x02,  // Vertex position: 2
     0x00, 0x06,  // Blend shape index in Definition
     // Mesh-2
-    0x00, 0x00, 0x0d, 0x17,  // Mesh-2 offset
+    0x00, 0x00, 0x0d, 0x18,  // Mesh-2 offset
     0x00, 0x00, 0x00, 0x03,  // Vertex positions X values length
     0x3f, 0x80, 0x00, 0x00,  // 1.0f
     0x40, 0x00, 0x00, 0x00,  // 2.0f
diff --git a/Engine/Plugins/Animation/RigLogic/Source/RigLogicModule/Private/RigUnit_RigLogic_Test.cpp b/Engine/Plugins/Animation/RigLogic/Source/RigLogicModule/Private/RigUnit_RigLogic_Test.cpp
--- a/Engine/Plugins/Animation/RigLogic/Source/RigLogicModule/Private/RigUnit_RigLogic_Test.cpp
+++ b/Engine/Plugins/Animation/RigLogic/Source/RigLogicModule/Private/RigUnit_RigLogic_Test.cpp
@@ -5,6 +5,7 @@
 #endif
 
 #include "RigUnit_RigLogic.h"
+#include "Components/SkeletalMeshComponent.h"
 #include "Units/RigUnitContext.h"
 #include "ControlRig.h"
 #include "Math/TransformNonVectorized.h"
