Merging using RigLogic_CharacterTech_to_UE5Main up to CL 14012845.

Fixes static analysis warnings.

diff --git a/Engine/Plugins/Animation/RigLogic/Source/RigLogicLib/Private/riglogic/riglogic/RigLogicImpl.cpp b/Engine/Plugins/Animation/RigLogic/Source/RigLogicLib/Private/riglogic/riglogic/RigLogicImpl.cpp
--- a/Engine/Plugins/Animation/RigLogic/Source/RigLogicLib/Private/riglogic/riglogic/RigLogicImpl.cpp
+++ b/Engine/Plugins/Animation/RigLogic/Source/RigLogicLib/Private/riglogic/riglogic/RigLogicImpl.cpp
@@ -16,8 +16,15 @@
 #include <terse/archives/binary/OutputArchive.h>
 #include <trio/Stream.h>
 
+#ifdef _MSC_VER
+    #pragma warning(push)
+    #pragma warning(disable : 4365 4987)
+#endif
 #include <cstddef>
 #include <utility>
+#ifdef _MSC_VER
+    #pragma warning(pop)
+#endif
 
 namespace rl4 {
 
@@ -95,12 +102,12 @@ RigLogicImpl::RigLogicImpl(Configuration config_,
                            AnimatedMapsPtr animatedMaps_,
                            MemoryResource* memRes_) :
     memRes{memRes_},
-    config{config_},
-    metrics{metrics_},
     controls{std::move(controls_)},
     joints{std::move(joints_)},
     blendShapes{std::move(blendShapes_)},
-    animatedMaps{std::move(animatedMaps_)} {
+    animatedMaps{std::move(animatedMaps_)},
+    config{config_},
+    metrics{metrics_} {
 }
 
 void RigLogicImpl::dump(BoundedIOStream* destination) const {
diff --git a/Engine/Plugins/Animation/RigLogic/Source/RigLogicLib/Public/dna/version/Version.h b/Engine/Plugins/Animation/RigLogic/Source/RigLogicLib/Public/dna/version/Version.h
--- a/Engine/Plugins/Animation/RigLogic/Source/RigLogicLib/Public/dna/version/Version.h
+++ b/Engine/Plugins/Animation/RigLogic/Source/RigLogicLib/Public/dna/version/Version.h
@@ -4,5 +4,5 @@
 
 #define DNA_MAJOR_VERSION 1
 #define DNA_MINOR_VERSION 1
-#define DNA_PATCH_VERSION 2
-#define DNA_VERSION_STRING "1.1.2"
+#define DNA_PATCH_VERSION 3
+#define DNA_VERSION_STRING "1.1.3"
diff --git a/Engine/Plugins/Animation/RigLogic/Source/RigLogicLib/Public/pma/version/Version.h b/Engine/Plugins/Animation/RigLogic/Source/RigLogicLib/Public/pma/version/Version.h
--- a/Engine/Plugins/Animation/RigLogic/Source/RigLogicLib/Public/pma/version/Version.h
+++ b/Engine/Plugins/Animation/RigLogic/Source/RigLogicLib/Public/pma/version/Version.h
@@ -4,5 +4,5 @@
 
 #define PMA_MAJOR_VERSION 1
 #define PMA_MINOR_VERSION 0
-#define PMA_PATCH_VERSION 3
-#define PMA_VERSION_STRING "1.0.3"
+#define PMA_PATCH_VERSION 4
+#define PMA_VERSION_STRING "1.0.4"
diff --git a/Engine/Plugins/Animation/RigLogic/Source/RigLogicLib/Public/status/version/Version.h b/Engine/Plugins/Animation/RigLogic/Source/RigLogicLib/Public/status/version/Version.h
--- a/Engine/Plugins/Animation/RigLogic/Source/RigLogicLib/Public/status/version/Version.h
+++ b/Engine/Plugins/Animation/RigLogic/Source/RigLogicLib/Public/status/version/Version.h
@@ -4,5 +4,5 @@
 
 #define SC_MAJOR_VERSION 1
 #define SC_MINOR_VERSION 0
-#define SC_PATCH_VERSION 2
-#define SC_VERSION_STRING "1.0.2"
+#define SC_PATCH_VERSION 3
+#define SC_VERSION_STRING "1.0.3"
diff --git a/Engine/Plugins/Animation/RigLogic/Source/RigLogicLib/Public/terse/version/Version.h b/Engine/Plugins/Animation/RigLogic/Source/RigLogicLib/Public/terse/version/Version.h
--- a/Engine/Plugins/Animation/RigLogic/Source/RigLogicLib/Public/terse/version/Version.h
+++ b/Engine/Plugins/Animation/RigLogic/Source/RigLogicLib/Public/terse/version/Version.h
@@ -4,5 +4,5 @@
 
 #define TERSE_MAJOR_VERSION 1
 #define TERSE_MINOR_VERSION 1
-#define TERSE_PATCH_VERSION 2
-#define TERSE_VERSION_STRING "1.1.2"
+#define TERSE_PATCH_VERSION 3
+#define TERSE_VERSION_STRING "1.1.3"
diff --git a/Engine/Plugins/Animation/RigLogic/Source/RigLogicLib/Public/trimd/version/Version.h b/Engine/Plugins/Animation/RigLogic/Source/RigLogicLib/Public/trimd/version/Version.h
--- a/Engine/Plugins/Animation/RigLogic/Source/RigLogicLib/Public/trimd/version/Version.h
+++ b/Engine/Plugins/Animation/RigLogic/Source/RigLogicLib/Public/trimd/version/Version.h
@@ -4,5 +4,5 @@
 
 #define TRIMD_MAJOR_VERSION 1
 #define TRIMD_MINOR_VERSION 0
-#define TRIMD_PATCH_VERSION 1
-#define TRIMD_VERSION_STRING "1.0.1"
+#define TRIMD_PATCH_VERSION 2
+#define TRIMD_VERSION_STRING "1.0.2"
diff --git a/Engine/Plugins/Animation/RigLogic/Source/RigLogicLib/Public/trio/version/Version.h b/Engine/Plugins/Animation/RigLogic/Source/RigLogicLib/Public/trio/version/Version.h
--- a/Engine/Plugins/Animation/RigLogic/Source/RigLogicLib/Public/trio/version/Version.h
+++ b/Engine/Plugins/Animation/RigLogic/Source/RigLogicLib/Public/trio/version/Version.h
@@ -4,5 +4,5 @@
 
 #define TRIO_MAJOR_VERSION 1
 #define TRIO_MINOR_VERSION 0
-#define TRIO_PATCH_VERSION 3
-#define TRIO_VERSION_STRING "1.0.3"
+#define TRIO_PATCH_VERSION 4
+#define TRIO_VERSION_STRING "1.0.4"
diff --git a/Engine/Plugins/Animation/RigLogic/Source/RigLogicLibTest/Private/RigLogicLibTestSuite.cpp b/Engine/Plugins/Animation/RigLogic/Source/RigLogicLibTest/Private/RigLogicLibTestSuite.cpp
--- a/Engine/Plugins/Animation/RigLogic/Source/RigLogicLibTest/Private/RigLogicLibTestSuite.cpp
+++ b/Engine/Plugins/Animation/RigLogic/Source/RigLogicLibTest/Private/RigLogicLibTestSuite.cpp
@@ -1,11 +1,15 @@
 // Copyright Epic Games, Inc. All Rights Reserved.
 
 #if WITH_DEV_AUTOMATION_TESTS
+#pragma warning(disable:4121)
 #pragma warning(disable:4668)
 #pragma warning(disable:4996)
+#pragma warning(disable:6330)
 #include "gtest/gtest.h"
+#pragma warning(default:4121)
 #pragma warning(default:4668)
 #pragma warning(default:4996)
+#pragma warning(default:6330)
 
 #include "CoreMinimal.h"
 #include "Misc/AutomationTest.h"
diff --git a/Engine/Plugins/Animation/RigLogic/Source/RigLogicLibTest/Private/dnatests/Defs.h b/Engine/Plugins/Animation/RigLogic/Source/RigLogicLibTest/Private/dnatests/Defs.h
--- a/Engine/Plugins/Animation/RigLogic/Source/RigLogicLibTest/Private/dnatests/Defs.h
+++ b/Engine/Plugins/Animation/RigLogic/Source/RigLogicLibTest/Private/dnatests/Defs.h
@@ -4,7 +4,7 @@
 
 #ifdef _MSC_VER
     #pragma warning(push)
-    #pragma warning(disable : 4365 4987 4668 4996)
+    #pragma warning(disable : 4121 4365 4987 4668 4996 6330)
 #endif
 #ifdef __clang__
     #pragma clang diagnostic push
diff --git a/Engine/Plugins/Animation/RigLogic/Source/RigLogicLibTest/Private/dnatests/FakeStream.cpp b/Engine/Plugins/Animation/RigLogic/Source/RigLogicLibTest/Private/dnatests/FakeStream.cpp
new file mode 100644
--- /dev/null
+++ b/Engine/Plugins/Animation/RigLogic/Source/RigLogicLibTest/Private/dnatests/FakeStream.cpp
@@ -0,0 +1,9 @@
+// Copyright Epic Games, Inc. All Rights Reserved.
+
+#include "dnatests/FakeStream.h"
+
+namespace dnatests {
+
+FakeStream::~FakeStream() = default;
+
+}  // namespace dnatests
diff --git a/Engine/Plugins/Animation/RigLogic/Source/RigLogicLibTest/Private/dnatests/FakeStream.h b/Engine/Plugins/Animation/RigLogic/Source/RigLogicLibTest/Private/dnatests/FakeStream.h
--- a/Engine/Plugins/Animation/RigLogic/Source/RigLogicLibTest/Private/dnatests/FakeStream.h
+++ b/Engine/Plugins/Animation/RigLogic/Source/RigLogicLibTest/Private/dnatests/FakeStream.h
@@ -23,6 +23,8 @@ struct FakeStream : public trio::BoundedIOStream {
     FakeStream() : position{} {
     }
 
+    ~FakeStream();
+
     void open() override {
         position = 0ul;
     }
diff --git a/Engine/Plugins/Animation/RigLogic/Source/RigLogicLibTest/Private/dnatests/Fixtures.cpp b/Engine/Plugins/Animation/RigLogic/Source/RigLogicLibTest/Private/dnatests/Fixtures.cpp
--- a/Engine/Plugins/Animation/RigLogic/Source/RigLogicLibTest/Private/dnatests/Fixtures.cpp
+++ b/Engine/Plugins/Animation/RigLogic/Source/RigLogicLibTest/Private/dnatests/Fixtures.cpp
@@ -1023,6 +1023,10 @@ std::vector<unsigned char> getBytes() {
 
 namespace decoded {
 
+#ifdef __clang__
+    #pragma clang diagnostic push
+    #pragma clang diagnostic ignored "-Wexit-time-destructors"
+#endif
 // Descriptor
 const pma::String<char> name = "test";
 const Archetype archetype = Archetype::other;
@@ -2031,6 +2035,9 @@ const pma::Matrix<pma::Matrix<std::uint32_t> > correctiveBlendShapeVertexIndices
         }
     }
 };
+#ifdef __clang__
+    #pragma clang diagnostic pop
+#endif
 
 std::size_t Fixtures::lodConstraintToIndex(std::uint16_t maxLOD, std::uint16_t minLOD) {
     // Relies on having only TWO available LODs (0, 1)
diff --git a/Engine/Plugins/Animation/RigLogic/Source/RigLogicLibTest/Private/dnatests/Fixtures.h b/Engine/Plugins/Animation/RigLogic/Source/RigLogicLibTest/Private/dnatests/Fixtures.h
--- a/Engine/Plugins/Animation/RigLogic/Source/RigLogicLibTest/Private/dnatests/Fixtures.h
+++ b/Engine/Plugins/Animation/RigLogic/Source/RigLogicLibTest/Private/dnatests/Fixtures.h
@@ -2,6 +2,8 @@
 
 #pragma once
 
+#include "dnatests/Defs.h"
+
 #include "dna/DNA.h"
 #include "dna/Reader.h"
 
diff --git a/Engine/Plugins/Animation/RigLogic/Source/RigLogicLibTest/Private/dnatests/TestFilteredInputArchive.cpp b/Engine/Plugins/Animation/RigLogic/Source/RigLogicLibTest/Private/dnatests/TestFilteredInputArchive.cpp
--- a/Engine/Plugins/Animation/RigLogic/Source/RigLogicLibTest/Private/dnatests/TestFilteredInputArchive.cpp
+++ b/Engine/Plugins/Animation/RigLogic/Source/RigLogicLibTest/Private/dnatests/TestFilteredInputArchive.cpp
@@ -12,6 +12,8 @@
 
 #include <memory>
 
+namespace {
+
 struct LODConstraint {
     std::uint16_t maxLOD;
     std::uint16_t minLOD;
@@ -38,6 +40,8 @@ class FilteredDNAInputArchiveTest : public ::testing::TestWithParam<LODConstrain
         LODConstraint lodConstraint;
 };
 
+}  // namespace
+
 TEST_P(FilteredDNAInputArchiveTest, FilterJoints) {
     const auto& result = dnaInstance->behavior.joints;
     const auto expected = dna::decoded::Fixtures::getJoints(lodConstraint.maxLOD, lodConstraint.minLOD, &amr);
@@ -120,6 +124,8 @@ INSTANTIATE_TEST_SUITE_P(FilteredDNAInputArchiveTest, FilteredDNAInputArchiveTes
                              LODConstraint{0u, 0u}
                              ));
 
+namespace {
+
 class GeometryFilteringTest : public ::testing::Test {
     protected:
         void SetUp() override {
@@ -136,6 +142,8 @@ class GeometryFilteringTest : public ::testing::Test {
         std::unique_ptr<dna::DNA> dnaInstance;
 };
 
+}  // namespace
+
 TEST_F(GeometryFilteringTest, IncludeBlendShapeTargets) {
     dna::FilteredInputArchive archive{&stream, dna::DataLayer::Geometry, 0u, std::numeric_limits<std::uint16_t>::max(), &amr};
     archive >> *dnaInstance;
diff --git a/Engine/Plugins/Animation/RigLogic/Source/RigLogicLibTest/Private/dnatests/TestLODMapping.cpp b/Engine/Plugins/Animation/RigLogic/Source/RigLogicLibTest/Private/dnatests/TestLODMapping.cpp
--- a/Engine/Plugins/Animation/RigLogic/Source/RigLogicLibTest/Private/dnatests/TestLODMapping.cpp
+++ b/Engine/Plugins/Animation/RigLogic/Source/RigLogicLibTest/Private/dnatests/TestLODMapping.cpp
@@ -8,6 +8,8 @@
 
 #include <pma/resources/DefaultMemoryResource.h>
 
+namespace {
+
 class LODMappingTest : public ::testing::Test {
     protected:
         void SetUp() override {
@@ -32,6 +34,8 @@ class LODMappingTest : public ::testing::Test {
         dna::LODMapping mapping{&memRes};
 };
 
+}  // namespace
+
 #define ASSERT_LODMAPPING_EQ(result, expected)   \
     for (std::uint16_t lod = 0u; lod < result.getLODCount(); ++lod) {               \
         dna::ConstArrayView<std::uint16_t> resultIndices = result.getIndices(lod);  \
diff --git a/Engine/Plugins/Animation/RigLogic/Source/RigLogicLibTest/Private/dnatests/TestStreamReadWriteIntegration.cpp b/Engine/Plugins/Animation/RigLogic/Source/RigLogicLibTest/Private/dnatests/TestStreamReadWriteIntegration.cpp
--- a/Engine/Plugins/Animation/RigLogic/Source/RigLogicLibTest/Private/dnatests/TestStreamReadWriteIntegration.cpp
+++ b/Engine/Plugins/Animation/RigLogic/Source/RigLogicLibTest/Private/dnatests/TestStreamReadWriteIntegration.cpp
@@ -13,6 +13,8 @@
 
 #include <pma/resources/AlignedMemoryResource.h>
 
+StreamReadWriteIntegrationTest::~StreamReadWriteIntegrationTest() = default;
+
 namespace dna {
 
 static void verifyDescriptor(DescriptorReader* reader, const LODParameters& params) {
diff --git a/Engine/Plugins/Animation/RigLogic/Source/RigLogicLibTest/Private/dnatests/TestStreamReadWriteIntegration.h b/Engine/Plugins/Animation/RigLogic/Source/RigLogicLibTest/Private/dnatests/TestStreamReadWriteIntegration.h
--- a/Engine/Plugins/Animation/RigLogic/Source/RigLogicLibTest/Private/dnatests/TestStreamReadWriteIntegration.h
+++ b/Engine/Plugins/Animation/RigLogic/Source/RigLogicLibTest/Private/dnatests/TestStreamReadWriteIntegration.h
@@ -13,6 +13,9 @@ struct LODParameters {
 };
 
 class StreamReadWriteIntegrationTest : public ::testing::TestWithParam<LODParameters> {
+    public:
+        ~StreamReadWriteIntegrationTest();
+
     protected:
         void SetUp() override {
             params = GetParam();
diff --git a/Engine/Plugins/Animation/RigLogic/Source/RigLogicLibTest/Private/dnatests/TestStreamWriter.cpp b/Engine/Plugins/Animation/RigLogic/Source/RigLogicLibTest/Private/dnatests/TestStreamWriter.cpp
--- a/Engine/Plugins/Animation/RigLogic/Source/RigLogicLibTest/Private/dnatests/TestStreamWriter.cpp
+++ b/Engine/Plugins/Animation/RigLogic/Source/RigLogicLibTest/Private/dnatests/TestStreamWriter.cpp
@@ -2,6 +2,8 @@
 
 #include "dnatests/TestStreamWriter.h"
 
+StreamWriterTest::~StreamWriterTest() = default;
+
 TEST_F(StreamWriterTest, SetName) {
     writer->setLODCount(1);
     writer->setName("test");
diff --git a/Engine/Plugins/Animation/RigLogic/Source/RigLogicLibTest/Private/dnatests/TestStreamWriter.h b/Engine/Plugins/Animation/RigLogic/Source/RigLogicLibTest/Private/dnatests/TestStreamWriter.h
--- a/Engine/Plugins/Animation/RigLogic/Source/RigLogicLibTest/Private/dnatests/TestStreamWriter.h
+++ b/Engine/Plugins/Animation/RigLogic/Source/RigLogicLibTest/Private/dnatests/TestStreamWriter.h
@@ -11,6 +11,9 @@
 #include <pma/resources/DefaultMemoryResource.h>
 
 class StreamWriterTest : public ::testing::Test {
+    public:
+        ~StreamWriterTest();
+
     protected:
         void SetUp() override {
             writer = dna::StreamWriter::create(&stream, &memRes);
diff --git a/Engine/Plugins/Animation/RigLogic/Source/RigLogicLibTest/Private/pmatests/Defs.h b/Engine/Plugins/Animation/RigLogic/Source/RigLogicLibTest/Private/pmatests/Defs.h
--- a/Engine/Plugins/Animation/RigLogic/Source/RigLogicLibTest/Private/pmatests/Defs.h
+++ b/Engine/Plugins/Animation/RigLogic/Source/RigLogicLibTest/Private/pmatests/Defs.h
@@ -4,7 +4,7 @@
 
 #ifdef _MSC_VER
     #pragma warning(push)
-    #pragma warning(disable : 4365 4987 4668 4996)
+    #pragma warning(disable : 4121 4365 4987 4668 4996 6330)
 #endif
 #ifdef __clang__
     #pragma clang diagnostic push
diff --git a/Engine/Plugins/Animation/RigLogic/Source/RigLogicLibTest/Private/pmatests/TestScopedPtr.cpp b/Engine/Plugins/Animation/RigLogic/Source/RigLogicLibTest/Private/pmatests/TestScopedPtr.cpp
--- a/Engine/Plugins/Animation/RigLogic/Source/RigLogicLibTest/Private/pmatests/TestScopedPtr.cpp
+++ b/Engine/Plugins/Animation/RigLogic/Source/RigLogicLibTest/Private/pmatests/TestScopedPtr.cpp
@@ -6,6 +6,8 @@
 
 namespace pmatests {
 
+namespace {
+
 struct Counters {
     int constructed;
     int destructed;
@@ -45,6 +47,8 @@ struct Base {
 
 struct Derived : public Base {};
 
+}  // namespace
+
 }  // namespace pmatests
 
 TEST(ScopedPtrTest, EmptyConstruction) {
@@ -133,7 +137,7 @@ TEST(ScopedPtrTest, UseNewDelete) {
 }
 
 TEST(ScopedPtrTest, UseNewDeleteForArrays) {
-    auto sp = pma::makeScoped<int[], pma::New, pma::Delete>(100);
+    auto sp = pma::makeScoped<int[], pma::New, pma::Delete>(100ul);
     ASSERT_EQ(sp[0], 0);
 }
 
@@ -152,11 +156,13 @@ TEST(ScopedPtrTest, CustomDestroyer) {
 TEST(ScopedPtrTest, MoveAssignWithCustomDestroyer) {
     std::size_t timesCalled = 0ul;
     {
-        pma::ScopedPtr<pmatests::Base, std::function<void(pmatests::Base*)> > sp;
-        sp = pma::ScopedPtr<pmatests::Derived, std::function<void(pmatests::Base*)> >{new pmatests::Derived{}, [&timesCalled](pmatests::Base* ptr) {
-                                                                      delete ptr;
-                                                                      ++timesCalled;
-                                                                  }};
+        using PtrType = pma::ScopedPtr<pmatests::Base, std::function<void (pmatests::Base*)> >;
+
+        PtrType sp;
+        sp = PtrType{new pmatests::Derived{}, [&timesCalled](pmatests::Base* ptr) {
+                         delete ptr;
+                         ++timesCalled;
+                     }};
         ASSERT_EQ(timesCalled, 0ul);
     }
     ASSERT_EQ(timesCalled, 1ul);
@@ -165,10 +171,11 @@ TEST(ScopedPtrTest, MoveAssignWithCustomDestroyer) {
 TEST(ScopedPtrTest, MoveConstructWithCustomDestroyer) {
     std::size_t timesCalled = 0ul;
     {
-        pma::ScopedPtr<pmatests::Base, std::function<void(pmatests::Base*)> > sp{new pmatests::Derived{}, [&timesCalled](pmatests::Base* ptr) {
-                                                                 delete ptr;
-                                                                 ++timesCalled;
-                                                             }};
+        pma::ScopedPtr<pmatests::Base, std::function<void(pmatests::Base*)> > sp{new pmatests::Derived{},
+                                                                                 [&timesCalled](pmatests::Base* ptr) {
+                                                                                     delete ptr;
+                                                                                     ++timesCalled;
+                                                                                 }};
         ASSERT_EQ(timesCalled, 0ul);
     }
     ASSERT_EQ(timesCalled, 1ul);
@@ -178,6 +185,6 @@ TEST(ScopedPtrTest, UseDefaultCreateDestroy) {
     auto spPrimitive = pma::makeScoped<int>(42);
     ASSERT_EQ(*spPrimitive, 42);
 
-    auto spArray = pma::makeScoped<int[]>(10);
+    auto spArray = pma::makeScoped<int[]>(10ul);
     ASSERT_EQ(spArray[0], 0);
 }
diff --git a/Engine/Plugins/Animation/RigLogic/Source/RigLogicLibTest/Private/pmatests/resources/TestAlignedMemoryResource.cpp b/Engine/Plugins/Animation/RigLogic/Source/RigLogicLibTest/Private/pmatests/resources/TestAlignedMemoryResource.cpp
--- a/Engine/Plugins/Animation/RigLogic/Source/RigLogicLibTest/Private/pmatests/resources/TestAlignedMemoryResource.cpp
+++ b/Engine/Plugins/Animation/RigLogic/Source/RigLogicLibTest/Private/pmatests/resources/TestAlignedMemoryResource.cpp
@@ -6,15 +6,19 @@
 
 #include <cstdint>
 
+namespace {
+
 class AlignedMemoryResourceTest : public ::testing::TestWithParam<std::size_t> {
 };
 
+}  // namespace
+
 TEST_P(AlignedMemoryResourceTest, AlignmentRequirementHonored) {
     std::size_t allocSize = 1024u;
     std::size_t alignment = GetParam();
     pma::AlignedMemoryResource amr;
     void* ptr = amr.allocate(allocSize, alignment);
-    ASSERT_EQ(reinterpret_cast<std::uintptr_t>(ptr) & (alignment - 1), 0);
+    ASSERT_EQ(reinterpret_cast<std::uintptr_t>(ptr) & (alignment - 1), static_cast<std::uintptr_t>(0));
     amr.deallocate(ptr, allocSize, alignment);
 }
 
diff --git a/Engine/Plugins/Animation/RigLogic/Source/RigLogicLibTest/Private/pmatests/utils/TestManagedInstance.cpp b/Engine/Plugins/Animation/RigLogic/Source/RigLogicLibTest/Private/pmatests/utils/TestManagedInstance.cpp
--- a/Engine/Plugins/Animation/RigLogic/Source/RigLogicLibTest/Private/pmatests/utils/TestManagedInstance.cpp
+++ b/Engine/Plugins/Animation/RigLogic/Source/RigLogicLibTest/Private/pmatests/utils/TestManagedInstance.cpp
@@ -6,6 +6,8 @@
 
 namespace pmatests {
 
+namespace {
+
 struct Stats {
     int constructed;
     int destructed;
@@ -36,6 +38,8 @@ class ManagableType : public BaseType {
         Stats* stats;
 };
 
+}  // namespace
+
 }  // namespace pmatests
 
 TEST(ManagedInstanceTest, UniqueInstance) {
diff --git a/Engine/Plugins/Animation/RigLogic/Source/RigLogicLibTest/Private/rltests/Defs.h b/Engine/Plugins/Animation/RigLogic/Source/RigLogicLibTest/Private/rltests/Defs.h
--- a/Engine/Plugins/Animation/RigLogic/Source/RigLogicLibTest/Private/rltests/Defs.h
+++ b/Engine/Plugins/Animation/RigLogic/Source/RigLogicLibTest/Private/rltests/Defs.h
@@ -4,7 +4,7 @@
 
 #ifdef _MSC_VER
     #pragma warning(push)
-    #pragma warning(disable : 4365 4987 4668 4996)
+    #pragma warning(disable : 4121 4365 4987 4668 4996 6330)
 #endif
 #ifdef __clang__
     #pragma clang diagnostic push
diff --git a/Engine/Plugins/Animation/RigLogic/Source/RigLogicLibTest/Private/rltests/animatedmaps/TestAnimatedMaps.cpp b/Engine/Plugins/Animation/RigLogic/Source/RigLogicLibTest/Private/rltests/animatedmaps/TestAnimatedMaps.cpp
--- a/Engine/Plugins/Animation/RigLogic/Source/RigLogicLibTest/Private/rltests/animatedmaps/TestAnimatedMaps.cpp
+++ b/Engine/Plugins/Animation/RigLogic/Source/RigLogicLibTest/Private/rltests/animatedmaps/TestAnimatedMaps.cpp
@@ -9,11 +9,22 @@
 
 #include <pma/resources/AlignedMemoryResource.h>
 
+#ifdef _MSC_VER
+    #pragma warning(push)
+    #pragma warning(disable : 4365 4987)
+#endif
 #include <array>
+#ifdef _MSC_VER
+    #pragma warning(pop)
+#endif
+
+namespace {
 
 class AnimatedMapsTest : public ::testing::TestWithParam<std::uint16_t> {
 };
 
+}  // namespace
+
 TEST_P(AnimatedMapsTest, LODLimitsCondTableSize) {
     pma::AlignedMemoryResource amr;
     rl4::Vector<std::uint16_t> lods{2, 1};
diff --git a/Engine/Plugins/Animation/RigLogic/Source/RigLogicLibTest/Private/rltests/blendshapes/BlendShapeFixtures.h b/Engine/Plugins/Animation/RigLogic/Source/RigLogicLibTest/Private/rltests/blendshapes/BlendShapeFixtures.h
--- a/Engine/Plugins/Animation/RigLogic/Source/RigLogicLibTest/Private/rltests/blendshapes/BlendShapeFixtures.h
+++ b/Engine/Plugins/Animation/RigLogic/Source/RigLogicLibTest/Private/rltests/blendshapes/BlendShapeFixtures.h
@@ -7,8 +7,15 @@
 
 #include <pma/MemoryResource.h>
 
+#ifdef _MSC_VER
+    #pragma warning(push)
+    #pragma warning(disable : 4365 4987)
+#endif
 #include <array>
 #include <cstddef>
+#ifdef _MSC_VER
+    #pragma warning(pop)
+#endif
 
 static const std::size_t LODCount = 4ul;
 static const std::uint16_t LOD0 = 10u;
diff --git a/Engine/Plugins/Animation/RigLogic/Source/RigLogicLibTest/Private/rltests/blendshapes/TestBlendShapes.cpp b/Engine/Plugins/Animation/RigLogic/Source/RigLogicLibTest/Private/rltests/blendshapes/TestBlendShapes.cpp
--- a/Engine/Plugins/Animation/RigLogic/Source/RigLogicLibTest/Private/rltests/blendshapes/TestBlendShapes.cpp
+++ b/Engine/Plugins/Animation/RigLogic/Source/RigLogicLibTest/Private/rltests/blendshapes/TestBlendShapes.cpp
@@ -9,11 +9,22 @@
 
 #include <pma/resources/AlignedMemoryResource.h>
 
+#ifdef _MSC_VER
+    #pragma warning(push)
+    #pragma warning(disable : 4365 4987)
+#endif
 #include <array>
+#ifdef _MSC_VER
+    #pragma warning(pop)
+#endif
+
+namespace {
 
 class BlendShapesTest : public ::testing::TestWithParam<std::uint16_t> {
 };
 
+}  // namespace
+
 TEST_P(BlendShapesTest, Calculate) {
     // Test input parameter - Which LOD level to test
     const auto lod = GetParam();
diff --git a/Engine/Plugins/Animation/RigLogic/Source/RigLogicLibTest/Private/rltests/conditionaltable/ConditionalTableFixtures.h b/Engine/Plugins/Animation/RigLogic/Source/RigLogicLibTest/Private/rltests/conditionaltable/ConditionalTableFixtures.h
--- a/Engine/Plugins/Animation/RigLogic/Source/RigLogicLibTest/Private/rltests/conditionaltable/ConditionalTableFixtures.h
+++ b/Engine/Plugins/Animation/RigLogic/Source/RigLogicLibTest/Private/rltests/conditionaltable/ConditionalTableFixtures.h
@@ -7,8 +7,15 @@
 
 #include <pma/MemoryResource.h>
 
+#ifdef _MSC_VER
+    #pragma warning(push)
+    #pragma warning(disable : 4365 4987)
+#endif
 #include <array>
 #include <cstddef>
+#ifdef _MSC_VER
+    #pragma warning(pop)
+#endif
 
 static const std::array<float, 9ul> conditionalTableInputs = {
     0.1f, 0.2f, 0.3f, 0.4f, 0.5f, 0.6f, 0.7f, 0.8f, 0.9f
diff --git a/Engine/Plugins/Animation/RigLogic/Source/RigLogicLibTest/Private/rltests/conditionaltable/TestConditionalTable.cpp b/Engine/Plugins/Animation/RigLogic/Source/RigLogicLibTest/Private/rltests/conditionaltable/TestConditionalTable.cpp
--- a/Engine/Plugins/Animation/RigLogic/Source/RigLogicLibTest/Private/rltests/conditionaltable/TestConditionalTable.cpp
+++ b/Engine/Plugins/Animation/RigLogic/Source/RigLogicLibTest/Private/rltests/conditionaltable/TestConditionalTable.cpp
@@ -7,6 +7,8 @@
 
 #include <pma/resources/AlignedMemoryResource.h>
 
+namespace {
+
 struct CalcTestData {
     float fromValues[2];
     float toValues[2];
@@ -16,6 +18,8 @@ struct CalcTestData {
 class ConditionalTableTest : public ::testing::TestWithParam<CalcTestData> {
 };
 
+}  // namespace
+
 TEST_P(ConditionalTableTest, CheckCalculationBorderCases) {
     pma::AlignedMemoryResource amr;
     const std::size_t conditionalsSize = 2ul;
diff --git a/Engine/Plugins/Animation/RigLogic/Source/RigLogicLibTest/Private/rltests/controls/TestControls.cpp b/Engine/Plugins/Animation/RigLogic/Source/RigLogicLibTest/Private/rltests/controls/TestControls.cpp
--- a/Engine/Plugins/Animation/RigLogic/Source/RigLogicLibTest/Private/rltests/controls/TestControls.cpp
+++ b/Engine/Plugins/Animation/RigLogic/Source/RigLogicLibTest/Private/rltests/controls/TestControls.cpp
@@ -8,7 +8,14 @@
 
 #include <pma/resources/AlignedMemoryResource.h>
 
+#ifdef _MSC_VER
+    #pragma warning(push)
+    #pragma warning(disable : 4365 4987)
+#endif
 #include <array>
+#ifdef _MSC_VER
+    #pragma warning(pop)
+#endif
 
 TEST(ControlsTest, GUIToRawMapping) {
     pma::AlignedMemoryResource amr;
diff --git a/Engine/Plugins/Animation/RigLogic/Source/RigLogicLibTest/Private/rltests/dna/DNAFixtures.cpp b/Engine/Plugins/Animation/RigLogic/Source/RigLogicLibTest/Private/rltests/dna/DNAFixtures.cpp
--- a/Engine/Plugins/Animation/RigLogic/Source/RigLogicLibTest/Private/rltests/dna/DNAFixtures.cpp
+++ b/Engine/Plugins/Animation/RigLogic/Source/RigLogicLibTest/Private/rltests/dna/DNAFixtures.cpp
@@ -7,6 +7,12 @@
 
 namespace rltests {
 
+#ifdef __clang__
+    #pragma clang diagnostic push
+    #pragma clang diagnostic ignored "-Wglobal-constructors"
+    #pragma clang diagnostic ignored "-Wexit-time-destructors"
+#endif
+
 namespace raw {
 
 const unsigned char header[] = {
@@ -1775,4 +1781,8 @@ const Matrix<Matrix<std::uint32_t> > correctiveBlendShapeVertexIndices = {
 
 }  // namespace decoded
 
+#ifdef __clang__
+    #pragma clang diagnostic pop
+#endif
+
 }  // namespace rltests
diff --git a/Engine/Plugins/Animation/RigLogic/Source/RigLogicLibTest/Private/rltests/dna/FakeIOStream.cpp b/Engine/Plugins/Animation/RigLogic/Source/RigLogicLibTest/Private/rltests/dna/FakeIOStream.cpp
new file mode 100644
--- /dev/null
+++ b/Engine/Plugins/Animation/RigLogic/Source/RigLogicLibTest/Private/rltests/dna/FakeIOStream.cpp
@@ -0,0 +1,9 @@
+// Copyright Epic Games, Inc. All Rights Reserved.
+
+#include "rltests/dna/FakeIOStream.h"
+
+namespace rltests {
+
+FakeIOStream::~FakeIOStream() = default;
+
+}  // namespace rltests
diff --git a/Engine/Plugins/Animation/RigLogic/Source/RigLogicLibTest/Private/rltests/dna/FakeIOStream.h b/Engine/Plugins/Animation/RigLogic/Source/RigLogicLibTest/Private/rltests/dna/FakeIOStream.h
new file mode 100644
--- /dev/null
+++ b/Engine/Plugins/Animation/RigLogic/Source/RigLogicLibTest/Private/rltests/dna/FakeIOStream.h
@@ -0,0 +1,67 @@
+// Copyright Epic Games, Inc. All Rights Reserved.
+
+#pragma once
+
+#include <trio/Stream.h>
+
+#ifdef _MSC_VER
+    #pragma warning(push)
+    #pragma warning(disable : 4365 4987)
+#endif
+#include <cassert>
+#include <cstring>
+#include <cstddef>
+#include <vector>
+#ifdef _MSC_VER
+    #pragma warning(pop)
+#endif
+
+namespace rltests {
+
+struct FakeIOStream : public trio::BoundedIOStream {
+
+    FakeIOStream() : position{} {
+    }
+
+    ~FakeIOStream();
+
+    void open() override {
+        position = 0ul;
+    }
+
+    void close() override {
+        position = 0ul;
+    }
+
+    std::size_t tell() override {
+        return position;
+    }
+
+    void seek(std::size_t position_) override {
+        position = position_;
+    }
+
+    void read(char* buffer, std::size_t size) override {
+        assert(position + size <= data.size());
+        std::memcpy(buffer, &data[position], size);
+        position += size;
+    }
+
+    void write(const char* buffer, std::size_t size) override {
+        std::size_t available = data.size() - position;
+        if (available < size) {
+            data.resize(data.size() + (size - available));
+        }
+        std::memcpy(&data[position], buffer, size);
+        position += size;
+    }
+
+    std::size_t size() override {
+        return data.size();
+    }
+
+    std::vector<char> data;
+    std::size_t position;
+};
+
+}  // namespace rltests
diff --git a/Engine/Plugins/Animation/RigLogic/Source/RigLogicLibTest/Private/rltests/dna/FakeReader.cpp b/Engine/Plugins/Animation/RigLogic/Source/RigLogicLibTest/Private/rltests/dna/FakeReader.cpp
new file mode 100644
--- /dev/null
+++ b/Engine/Plugins/Animation/RigLogic/Source/RigLogicLibTest/Private/rltests/dna/FakeReader.cpp
@@ -0,0 +1,9 @@
+// Copyright Epic Games, Inc. All Rights Reserved.
+
+#include "rltests/dna/FakeReader.h"
+
+namespace dna {
+
+FakeReader::~FakeReader() = default;
+
+}  // namespace dna
diff --git a/Engine/Plugins/Animation/RigLogic/Source/RigLogicLibTest/Private/rltests/dna/FakeReader.h b/Engine/Plugins/Animation/RigLogic/Source/RigLogicLibTest/Private/rltests/dna/FakeReader.h
--- a/Engine/Plugins/Animation/RigLogic/Source/RigLogicLibTest/Private/rltests/dna/FakeReader.h
+++ b/Engine/Plugins/Animation/RigLogic/Source/RigLogicLibTest/Private/rltests/dna/FakeReader.h
@@ -8,6 +8,10 @@
 
 #include <cstdint>
 
+#if defined(__clang__) || defined(__GNUC__)
+    #pragma GCC diagnostic push
+    #pragma GCC diagnostic ignored "-Wunused-parameter"
+#endif
 #if defined(_MSC_VER) && !defined(__clang__)
     #pragma warning(push)
     #pragma warning(disable : 4100)
@@ -17,6 +21,8 @@ namespace dna {
 
 class FakeReader : public Reader {
     public:
+        ~FakeReader();
+
         // DescriptorReader methods start
         StringView getName() const override {
             return {};
@@ -445,6 +451,9 @@ class FakeReader : public Reader {
 
 }  // namespace dna
 
+#if defined(__clang__) || defined(__GNUC__)
+    #pragma GCC diagnostic pop
+#endif
 #if defined(_MSC_VER) && !defined(__clang__)
     #pragma warning(pop)
 #endif
diff --git a/Engine/Plugins/Animation/RigLogic/Source/RigLogicLibTest/Private/rltests/dna/FakeStream.h b/Engine/Plugins/Animation/RigLogic/Source/RigLogicLibTest/Private/rltests/dna/FakeStream.h
deleted file mode 100644
diff --git a/Engine/Plugins/Animation/RigLogic/Source/RigLogicLibTest/Private/rltests/joints/bpcm/JointFixturesBlock4.cpp b/Engine/Plugins/Animation/RigLogic/Source/RigLogicLibTest/Private/rltests/joints/bpcm/JointFixturesBlock4.cpp
--- a/Engine/Plugins/Animation/RigLogic/Source/RigLogicLibTest/Private/rltests/joints/bpcm/JointFixturesBlock4.cpp
+++ b/Engine/Plugins/Animation/RigLogic/Source/RigLogicLibTest/Private/rltests/joints/bpcm/JointFixturesBlock4.cpp
@@ -10,6 +10,12 @@
 
 namespace block4 {
 
+#ifdef __clang__
+    #pragma clang diagnostic push
+    #pragma clang diagnostic ignored "-Wglobal-constructors"
+    #pragma clang diagnostic ignored "-Wexit-time-destructors"
+#endif
+
 namespace unoptimized {
 
 const std::uint16_t lodCount = 4u;
@@ -858,6 +864,12 @@ const rl4::Matrix<float> valuesPerLOD = {
 
 }  // namespace output
 
+#ifdef __clang__
+    #pragma clang diagnostic pop
+#endif
+
+CanonicalReader::~CanonicalReader() = default;
+
 template<typename TValue>
 struct OptimizedValues;
 
diff --git a/Engine/Plugins/Animation/RigLogic/Source/RigLogicLibTest/Private/rltests/joints/bpcm/JointFixturesBlock4.h b/Engine/Plugins/Animation/RigLogic/Source/RigLogicLibTest/Private/rltests/joints/bpcm/JointFixturesBlock4.h
--- a/Engine/Plugins/Animation/RigLogic/Source/RigLogicLibTest/Private/rltests/joints/bpcm/JointFixturesBlock4.h
+++ b/Engine/Plugins/Animation/RigLogic/Source/RigLogicLibTest/Private/rltests/joints/bpcm/JointFixturesBlock4.h
@@ -12,9 +12,16 @@
 #include "riglogic/joints/bpcm/LODRegion.h"
 #include "riglogic/riglogic/RigLogic.h"
 
+#ifdef _MSC_VER
+    #pragma warning(push)
+    #pragma warning(disable : 4365 4987)
+#endif
 #include <cstdint>
 #include <functional>
 #include <memory>
+#ifdef _MSC_VER
+    #pragma warning(pop)
+#endif
 
 namespace block4 {
 
@@ -66,6 +73,8 @@ extern const Matrix<float> valuesPerLOD;
 
 class CanonicalReader : public dna::FakeReader {
     public:
+        ~CanonicalReader();
+
         std::uint16_t getLODCount() const override {
             return unoptimized::lodCount;
         }
diff --git a/Engine/Plugins/Animation/RigLogic/Source/RigLogicLibTest/Private/rltests/joints/bpcm/JointFixturesBlock8.cpp b/Engine/Plugins/Animation/RigLogic/Source/RigLogicLibTest/Private/rltests/joints/bpcm/JointFixturesBlock8.cpp
--- a/Engine/Plugins/Animation/RigLogic/Source/RigLogicLibTest/Private/rltests/joints/bpcm/JointFixturesBlock8.cpp
+++ b/Engine/Plugins/Animation/RigLogic/Source/RigLogicLibTest/Private/rltests/joints/bpcm/JointFixturesBlock8.cpp
@@ -9,6 +9,12 @@
 
 namespace block8 {
 
+#ifdef __clang__
+    #pragma clang diagnostic push
+    #pragma clang diagnostic ignored "-Wglobal-constructors"
+    #pragma clang diagnostic ignored "-Wexit-time-destructors"
+#endif
+
 namespace unoptimized {
 
 const std::uint16_t lodCount = 4u;
@@ -863,6 +869,12 @@ const rl4::Matrix<float> valuesPerLOD = {
 
 }  // namespace output
 
+#ifdef __clang__
+    #pragma clang diagnostic pop
+#endif
+
+CanonicalReader::~CanonicalReader() = default;
+
 template<typename TValue>
 struct OptimizedValues;
 
diff --git a/Engine/Plugins/Animation/RigLogic/Source/RigLogicLibTest/Private/rltests/joints/bpcm/JointFixturesBlock8.h b/Engine/Plugins/Animation/RigLogic/Source/RigLogicLibTest/Private/rltests/joints/bpcm/JointFixturesBlock8.h
--- a/Engine/Plugins/Animation/RigLogic/Source/RigLogicLibTest/Private/rltests/joints/bpcm/JointFixturesBlock8.h
+++ b/Engine/Plugins/Animation/RigLogic/Source/RigLogicLibTest/Private/rltests/joints/bpcm/JointFixturesBlock8.h
@@ -12,6 +12,17 @@
 #include "riglogic/joints/bpcm/JointsEvaluator.h"
 #include "riglogic/joints/bpcm/LODRegion.h"
 
+#ifdef _MSC_VER
+    #pragma warning(push)
+    #pragma warning(disable : 4365 4987)
+#endif
+#include <cstdint>
+#include <functional>
+#include <memory>
+#ifdef _MSC_VER
+    #pragma warning(pop)
+#endif
+
 namespace block8 {
 
 using namespace rl4;
@@ -55,6 +66,8 @@ extern const Matrix<float> valuesPerLOD;
 
 class CanonicalReader : public dna::FakeReader {
     public:
+        ~CanonicalReader();
+
         std::uint16_t getLODCount() const override {
             return unoptimized::lodCount;
         }
diff --git a/Engine/Plugins/Animation/RigLogic/Source/RigLogicLibTest/Private/rltests/joints/bpcm/StorageValueType.h b/Engine/Plugins/Animation/RigLogic/Source/RigLogicLibTest/Private/rltests/joints/bpcm/StorageValueType.h
--- a/Engine/Plugins/Animation/RigLogic/Source/RigLogicLibTest/Private/rltests/joints/bpcm/StorageValueType.h
+++ b/Engine/Plugins/Animation/RigLogic/Source/RigLogicLibTest/Private/rltests/joints/bpcm/StorageValueType.h
@@ -2,6 +2,8 @@
 
 #pragma once
 
+#include <cstdint>
+
 #ifdef RL_USE_HALF_FLOATS
     using StorageValueType = std::uint16_t;
 #else
diff --git a/Engine/Plugins/Animation/RigLogic/Source/RigLogicLibTest/Private/rltests/joints/bpcm/TestJointCalculationStrategyAVX.cpp b/Engine/Plugins/Animation/RigLogic/Source/RigLogicLibTest/Private/rltests/joints/bpcm/TestJointCalculationStrategyAVX.cpp
--- a/Engine/Plugins/Animation/RigLogic/Source/RigLogicLibTest/Private/rltests/joints/bpcm/TestJointCalculationStrategyAVX.cpp
+++ b/Engine/Plugins/Animation/RigLogic/Source/RigLogicLibTest/Private/rltests/joints/bpcm/TestJointCalculationStrategyAVX.cpp
@@ -13,6 +13,8 @@
 #include <pma/resources/AlignedMemoryResource.h>
 #include <pma/utils/ManagedInstance.h>
 
+namespace {
+
 class AVXJointCalculationStrategyTest : public ::testing::TestWithParam<StrategyTestParams> {
     protected:
         void SetUp() {
@@ -36,6 +38,8 @@ class AVXJointCalculationStrategyTest : public ::testing::TestWithParam<Strategy
 
 };
 
+}  // namespace
+
 TEST_P(AVXJointCalculationStrategyTest, Block8Padded) {
     const auto params = GetParam();
     const OutputScope scope{params.lod, 0ul, 5ul};
diff --git a/Engine/Plugins/Animation/RigLogic/Source/RigLogicLibTest/Private/rltests/joints/bpcm/TestJointCalculationStrategySSE.cpp b/Engine/Plugins/Animation/RigLogic/Source/RigLogicLibTest/Private/rltests/joints/bpcm/TestJointCalculationStrategySSE.cpp
--- a/Engine/Plugins/Animation/RigLogic/Source/RigLogicLibTest/Private/rltests/joints/bpcm/TestJointCalculationStrategySSE.cpp
+++ b/Engine/Plugins/Animation/RigLogic/Source/RigLogicLibTest/Private/rltests/joints/bpcm/TestJointCalculationStrategySSE.cpp
@@ -13,6 +13,8 @@
 #include <pma/resources/AlignedMemoryResource.h>
 #include <pma/utils/ManagedInstance.h>
 
+namespace {
+
 class SSEJointCalculationStrategyTest : public ::testing::TestWithParam<StrategyTestParams> {
     protected:
         void SetUp() {
@@ -35,6 +37,8 @@ class SSEJointCalculationStrategyTest : public ::testing::TestWithParam<Strategy
         block4::OptimizedStorage<StorageValueType>::StrategyPtr strategy;
 };
 
+}  // namespace
+
 TEST_P(SSEJointCalculationStrategyTest, Block4Padded) {
     const auto params = GetParam();
     const OutputScope scope{params.lod, 0ul, 1ul};
diff --git a/Engine/Plugins/Animation/RigLogic/Source/RigLogicLibTest/Private/rltests/joints/bpcm/TestJointCalculationStrategyScalar.cpp b/Engine/Plugins/Animation/RigLogic/Source/RigLogicLibTest/Private/rltests/joints/bpcm/TestJointCalculationStrategyScalar.cpp
--- a/Engine/Plugins/Animation/RigLogic/Source/RigLogicLibTest/Private/rltests/joints/bpcm/TestJointCalculationStrategyScalar.cpp
+++ b/Engine/Plugins/Animation/RigLogic/Source/RigLogicLibTest/Private/rltests/joints/bpcm/TestJointCalculationStrategyScalar.cpp
@@ -11,6 +11,8 @@
 #include <pma/resources/AlignedMemoryResource.h>
 #include <pma/utils/ManagedInstance.h>
 
+namespace {
+
 class ScalarJointCalculationStrategyTest : public ::testing::TestWithParam<StrategyTestParams> {
     protected:
         void SetUp() {
@@ -33,6 +35,8 @@ class ScalarJointCalculationStrategyTest : public ::testing::TestWithParam<Strat
         block4::OptimizedStorage<float>::StrategyPtr strategy;
 };
 
+}  // namespace
+
 TEST_P(ScalarJointCalculationStrategyTest, Block4Padded) {
     const auto params = GetParam();
     const OutputScope scope{params.lod, 0ul, 1ul};
diff --git a/Engine/Plugins/Animation/RigLogic/Source/RigLogicLibTest/Private/rltests/joints/bpcm/TestJointStorageBuilderAVX.cpp b/Engine/Plugins/Animation/RigLogic/Source/RigLogicLibTest/Private/rltests/joints/bpcm/TestJointStorageBuilderAVX.cpp
--- a/Engine/Plugins/Animation/RigLogic/Source/RigLogicLibTest/Private/rltests/joints/bpcm/TestJointStorageBuilderAVX.cpp
+++ b/Engine/Plugins/Animation/RigLogic/Source/RigLogicLibTest/Private/rltests/joints/bpcm/TestJointStorageBuilderAVX.cpp
@@ -14,6 +14,8 @@
 #include <pma/resources/AlignedMemoryResource.h>
 #include <pma/utils/ManagedInstance.h>
 
+namespace {
+
 template<typename TValue>
 class AVXJointStorageBuilderTest : public ::testing::Test {
     protected:
@@ -40,8 +42,10 @@ class AVXJointStorageBuilderTest : public ::testing::Test {
 
 };
 
+}  // namespace
+
 using StorageValueTypeList = ::testing::Types<StorageValueType>;
-TYPED_TEST_SUITE(AVXJointStorageBuilderTest, StorageValueTypeList);
+TYPED_TEST_SUITE(AVXJointStorageBuilderTest, StorageValueTypeList, );
 
 TYPED_TEST(AVXJointStorageBuilderTest, LayoutOptimization) {
     this->buildStorage();
diff --git a/Engine/Plugins/Animation/RigLogic/Source/RigLogicLibTest/Private/rltests/joints/bpcm/TestJointStorageBuilderSSE.cpp b/Engine/Plugins/Animation/RigLogic/Source/RigLogicLibTest/Private/rltests/joints/bpcm/TestJointStorageBuilderSSE.cpp
--- a/Engine/Plugins/Animation/RigLogic/Source/RigLogicLibTest/Private/rltests/joints/bpcm/TestJointStorageBuilderSSE.cpp
+++ b/Engine/Plugins/Animation/RigLogic/Source/RigLogicLibTest/Private/rltests/joints/bpcm/TestJointStorageBuilderSSE.cpp
@@ -14,6 +14,8 @@
 #include <pma/resources/AlignedMemoryResource.h>
 #include <pma/utils/ManagedInstance.h>
 
+namespace {
+
 template<typename TValue>
 class SSEJointStorageBuilderTest : public ::testing::Test {
     protected:
@@ -40,8 +42,10 @@ class SSEJointStorageBuilderTest : public ::testing::Test {
 
 };
 
+}  // namespace
+
 using StorageValueTypeList = ::testing::Types<StorageValueType>;
-TYPED_TEST_SUITE(SSEJointStorageBuilderTest, StorageValueTypeList);
+TYPED_TEST_SUITE(SSEJointStorageBuilderTest, StorageValueTypeList, );
 
 TYPED_TEST(SSEJointStorageBuilderTest, LayoutOptimization) {
     this->buildStorage();
diff --git a/Engine/Plugins/Animation/RigLogic/Source/RigLogicLibTest/Private/rltests/joints/bpcm/TestJointStorageBuilderScalar.cpp b/Engine/Plugins/Animation/RigLogic/Source/RigLogicLibTest/Private/rltests/joints/bpcm/TestJointStorageBuilderScalar.cpp
--- a/Engine/Plugins/Animation/RigLogic/Source/RigLogicLibTest/Private/rltests/joints/bpcm/TestJointStorageBuilderScalar.cpp
+++ b/Engine/Plugins/Animation/RigLogic/Source/RigLogicLibTest/Private/rltests/joints/bpcm/TestJointStorageBuilderScalar.cpp
@@ -13,6 +13,8 @@
 #include <pma/resources/AlignedMemoryResource.h>
 #include <pma/utils/ManagedInstance.h>
 
+namespace {
+
 class ScalarJointStorageBuilderTest : public ::testing::Test {
     protected:
         void buildStorage() {
@@ -38,6 +40,8 @@ class ScalarJointStorageBuilderTest : public ::testing::Test {
 
 };
 
+}  // namespace
+
 TEST_F(ScalarJointStorageBuilderTest, LayoutOptimization) {
     this->buildStorage();
 }
diff --git a/Engine/Plugins/Animation/RigLogic/Source/RigLogicLibTest/Private/rltests/joints/bpcm/TestLODRegion.cpp b/Engine/Plugins/Animation/RigLogic/Source/RigLogicLibTest/Private/rltests/joints/bpcm/TestLODRegion.cpp
--- a/Engine/Plugins/Animation/RigLogic/Source/RigLogicLibTest/Private/rltests/joints/bpcm/TestLODRegion.cpp
+++ b/Engine/Plugins/Animation/RigLogic/Source/RigLogicLibTest/Private/rltests/joints/bpcm/TestLODRegion.cpp
@@ -6,13 +6,7 @@
 
 #include <cstdint>
 
-TEST(LODRegionTest, ConstructFromFixtures) {
-    rl4::bpcm::LODRegion lod{1, 2, 3};
-
-    ASSERT_EQ(lod.size, 1);
-    ASSERT_EQ(lod.sizeAlignedToLastFullBlock, 2);
-    ASSERT_EQ(lod.sizeAlignedToSecondLastFullBlock, 3);
-}
+namespace {
 
 struct LODTestSetup {
     std::uint32_t rowCount;
@@ -26,6 +20,16 @@ struct LODTestSetup {
 class LODRegionTest : public ::testing::TestWithParam<LODTestSetup> {
 };
 
+}  // namespace
+
+TEST(LODRegionTest, ConstructFromFixtures) {
+    rl4::bpcm::LODRegion lod{1, 2, 3};
+
+    ASSERT_EQ(lod.size, 1u);
+    ASSERT_EQ(lod.sizeAlignedToLastFullBlock, 2u);
+    ASSERT_EQ(lod.sizeAlignedToSecondLastFullBlock, 3u);
+}
+
 TEST_P(LODRegionTest, ConstructLODRegionFromVaryingParameters) {
     auto params = GetParam();
     rl4::bpcm::LODRegion lod{params.lodEndRow, params.rowCount, params.blockHeight, params.padTo};
diff --git a/Engine/Plugins/Animation/RigLogic/Source/RigLogicLibTest/Private/rltests/riglogic/TestRigLogic.cpp b/Engine/Plugins/Animation/RigLogic/Source/RigLogicLibTest/Private/rltests/riglogic/TestRigLogic.cpp
--- a/Engine/Plugins/Animation/RigLogic/Source/RigLogicLibTest/Private/rltests/riglogic/TestRigLogic.cpp
+++ b/Engine/Plugins/Animation/RigLogic/Source/RigLogicLibTest/Private/rltests/riglogic/TestRigLogic.cpp
@@ -2,12 +2,14 @@
 
 #include "rltests/Defs.h"
 #include "rltests/dna/DNAFixtures.h"
-#include "rltests/dna/FakeStream.h"
+#include "rltests/dna/FakeIOStream.h"
 
 #include "riglogic/RigLogic.h"
 
 #include <pma/resources/AlignedMemoryResource.h>
 
+namespace {
+
 class RigLogicTest : public ::testing::Test {
     protected:
         void SetUp() override {
@@ -30,12 +32,14 @@ class RigLogicTest : public ::testing::Test {
 
     protected:
         pma::AlignedMemoryResource memRes;
-        rltests::FakeStream stream;
+        rltests::FakeIOStream stream;
         dna::StreamReader* reader;
         rl4::RigLogic* rigLogic;
         rl4::RigInstance* rigInstance;
 };
 
+}  // namespace
+
 TEST_F(RigLogicTest, EvaluateRigInstance) {
     // Try both approaches
     float guiControls[] = {1.0f, 2.0f, 3.0f, 4.0f, 5.0f, 6.0f, 7.0f, 8.0f, 9.0f};
@@ -62,7 +66,7 @@ TEST_F(RigLogicTest, AccessJointVariableAttributeIndices) {
 TEST_F(RigLogicTest, DumpStateThenRestore) {
     float guiControls[] = {1.0f, 2.0f, 3.0f, 4.0f, 5.0f, 6.0f, 7.0f, 8.0f, 9.0f};
 
-    rltests::FakeStream dumpedState;
+    rltests::FakeIOStream dumpedState;
     rigLogic->dump(&dumpedState);
     dumpedState.seek(0);
     auto cloneRigLogic = rl4::RigLogic::restore(&dumpedState);
diff --git a/Engine/Plugins/Animation/RigLogic/Source/RigLogicLibTest/Private/tersetests/Defs.h b/Engine/Plugins/Animation/RigLogic/Source/RigLogicLibTest/Private/tersetests/Defs.h
--- a/Engine/Plugins/Animation/RigLogic/Source/RigLogicLibTest/Private/tersetests/Defs.h
+++ b/Engine/Plugins/Animation/RigLogic/Source/RigLogicLibTest/Private/tersetests/Defs.h
@@ -4,7 +4,7 @@
 
 #ifdef _MSC_VER
     #pragma warning(push)
-    #pragma warning(disable : 4365 4987 4668 4996)
+    #pragma warning(disable : 4121 4365 4987 4668 4996 6330)
 #endif
 #ifdef __clang__
     #pragma clang diagnostic push
diff --git a/Engine/Plugins/Animation/RigLogic/Source/RigLogicLibTest/Private/tersetests/archives/binary/TestInputArchive.cpp b/Engine/Plugins/Animation/RigLogic/Source/RigLogicLibTest/Private/tersetests/archives/binary/TestInputArchive.cpp
--- a/Engine/Plugins/Animation/RigLogic/Source/RigLogicLibTest/Private/tersetests/archives/binary/TestInputArchive.cpp
+++ b/Engine/Plugins/Animation/RigLogic/Source/RigLogicLibTest/Private/tersetests/archives/binary/TestInputArchive.cpp
@@ -93,7 +93,11 @@ struct SectionSkipper {
     terse::ArchiveOffset<std::uint32_t>::Proxy afterValueMarker;
     std::uint32_t afterValue;
 
-    SectionSkipper() : afterValueMarker{afterValueOffset} {
+    SectionSkipper() :
+        afterValueOffset{},
+        ignored{},
+        afterValueMarker{afterValueOffset},
+        afterValue{} {
     }
 
     template<class Archive>
diff --git a/Engine/Plugins/Animation/RigLogic/Source/RigLogicModule/Private/RigUnit_RigLogic_Test.cpp b/Engine/Plugins/Animation/RigLogic/Source/RigLogicModule/Private/RigUnit_RigLogic_Test.cpp
--- a/Engine/Plugins/Animation/RigLogic/Source/RigLogicModule/Private/RigUnit_RigLogic_Test.cpp
+++ b/Engine/Plugins/Animation/RigLogic/Source/RigLogicModule/Private/RigUnit_RigLogic_Test.cpp
@@ -27,28 +27,28 @@ FRigUnit_RigLogic::TestAccessor::TestAccessor(FRigUnit_RigLogic* Unit)
 
 /** ====== Map Input Curves ===== **/
 
-TestBehaviorReader* FRigUnit_RigLogic::TestAccessor::CreateBehaviorReaderEmpty()
+TUniquePtr<TestBehaviorReader> FRigUnit_RigLogic::TestAccessor::CreateBehaviorReaderEmpty()
 {
-	return new TestBehaviorReader();
+	return MakeUnique<TestBehaviorReader>();
 }
 
-FRigCurveContainer* FRigUnit_RigLogic::TestAccessor::CreateCurveContainerEmpty()
+TUniquePtr<FRigCurveContainer> FRigUnit_RigLogic::TestAccessor::CreateCurveContainerEmpty()
 {
-	return new FRigCurveContainer();
+	return MakeUnique<FRigCurveContainer>();
 }
 
-TestBehaviorReader* FRigUnit_RigLogic::TestAccessor::CreateBehaviorReaderOneCurve(FString ControlNameStr)
+TUniquePtr<TestBehaviorReader> FRigUnit_RigLogic::TestAccessor::CreateBehaviorReaderOneCurve(FString ControlNameStr)
 {
-	TestBehaviorReader* BehaviorReader = new TestBehaviorReader;
+	TUniquePtr<TestBehaviorReader> BehaviorReader = MakeUnique<TestBehaviorReader>();
 	//StringView ControlName(TCHAR_TO_ANSI(*ControlNameStr), ControlNameStr.Len());
 	BehaviorReader->rawControls.Add(ControlNameStr);
 	BehaviorReader->LODCount = 1;
 	return BehaviorReader;
 }
 
-FRigCurveContainer* FRigUnit_RigLogic::TestAccessor::CreateCurveContainerOneCurve(FString CurveNameStr)
+TUniquePtr<FRigCurveContainer> FRigUnit_RigLogic::TestAccessor::CreateCurveContainerOneCurve(FString CurveNameStr)
 {
-	FRigCurveContainer* ValidCurveContainer = new  FRigCurveContainer();
+	TUniquePtr<FRigCurveContainer> ValidCurveContainer = MakeUnique<FRigCurveContainer>();
 	ValidCurveContainer->Add(FName(*CurveNameStr));
 	ValidCurveContainer->Initialize();
 	return ValidCurveContainer;
@@ -62,14 +62,14 @@ void FRigUnit_RigLogic::TestAccessor::Exec_MapInputCurve(TestBehaviorReader* Tes
 /** ====== Map Joints ===== **/
 
 
-FRigBoneHierarchy* FRigUnit_RigLogic::TestAccessor::CreateBoneHierarchyEmpty()
+TUniquePtr<FRigBoneHierarchy> FRigUnit_RigLogic::TestAccessor::CreateBoneHierarchyEmpty()
 {
-	return new FRigBoneHierarchy();
+	return MakeUnique<FRigBoneHierarchy>();
 }
 
-FRigBoneHierarchy* FRigUnit_RigLogic::TestAccessor::CreateBoneHierarchyTwoBones( FString Bone1NameStr, FString Bone2NameStr )
+TUniquePtr<FRigBoneHierarchy> FRigUnit_RigLogic::TestAccessor::CreateBoneHierarchyTwoBones(FString Bone1NameStr, FString Bone2NameStr)
 {
-	FRigBoneHierarchy* TestHierarchy = new FRigBoneHierarchy();
+	TUniquePtr<FRigBoneHierarchy> TestHierarchy = MakeUnique<FRigBoneHierarchy>();
 	TestHierarchy->Reset();
 	TestHierarchy->Add(*Bone1NameStr, NAME_None, ERigBoneType::User, FTransform(FVector(1.f, 0.f, 0.f)));
 	TestHierarchy->Add(*Bone2NameStr, *Bone1NameStr, ERigBoneType::User, FTransform(FVector(1.f, 2.f, 0.f)));
@@ -77,9 +77,9 @@ FRigBoneHierarchy* FRigUnit_RigLogic::TestAccessor::CreateBoneHierarchyTwoBones(
 	return TestHierarchy;
 }
 
-TestBehaviorReader* FRigUnit_RigLogic::TestAccessor::CreateBehaviorReaderTwoJoints( FString Joint1NameStr, FString Joint2NameStr)
+TUniquePtr<TestBehaviorReader> FRigUnit_RigLogic::TestAccessor::CreateBehaviorReaderTwoJoints(FString Joint1NameStr, FString Joint2NameStr)
 {
-	TestBehaviorReader* TestReader = new TestBehaviorReader();
+	TUniquePtr<TestBehaviorReader> TestReader = MakeUnique<TestBehaviorReader>();
 	TestReader->addJoint(Joint1NameStr);
 	TestReader->addJoint(Joint2NameStr);
 	TestReader->LODCount = 1;
@@ -93,17 +93,17 @@ void FRigUnit_RigLogic::TestAccessor::Exec_MapJoints(TestBehaviorReader* TestRea
 
 /** ====== Map Morph Targets ===== **/
 
-TestBehaviorReader* FRigUnit_RigLogic::TestAccessor::CreateBehaviorReaderNoBlendshapes(FString MeshNameStr)
+TUniquePtr<TestBehaviorReader> FRigUnit_RigLogic::TestAccessor::CreateBehaviorReaderNoBlendshapes(FString MeshNameStr)
 {
-	TestBehaviorReader* BehaviorReader = new TestBehaviorReader;
+	TUniquePtr<TestBehaviorReader> BehaviorReader = MakeUnique<TestBehaviorReader>();
 	BehaviorReader->addMeshName(*MeshNameStr);
 	BehaviorReader->LODCount = 1; //there is one mesh, so LODs exist
 	return BehaviorReader;
 }
 
-TestBehaviorReader* FRigUnit_RigLogic::TestAccessor::CreateBehaviorReaderOneBlendShape(FString MeshNameStr, FString BlendShapeNameStr)
+TUniquePtr<TestBehaviorReader> FRigUnit_RigLogic::TestAccessor::CreateBehaviorReaderOneBlendShape(FString MeshNameStr, FString BlendShapeNameStr)
 {
-	TestBehaviorReader* BehaviorReader = new TestBehaviorReader;
+	TUniquePtr<TestBehaviorReader> BehaviorReader = MakeUnique<TestBehaviorReader>();
 	BehaviorReader->addBlendShapeChannelName(*BlendShapeNameStr);
 	BehaviorReader->addBlendShapeChannelName(*BlendShapeNameStr);
 	BehaviorReader->addMeshName(*MeshNameStr);
@@ -114,9 +114,9 @@ TestBehaviorReader* FRigUnit_RigLogic::TestAccessor::CreateBehaviorReaderOneBlen
 	return BehaviorReader;
 }
 
-TestBehaviorReader* FRigUnit_RigLogic::TestAccessor::CreateBehaviorReaderTwoBlendShapes(FString MeshNameStr, FString BlendShape1Str, FString BlendShape2Str)
+TUniquePtr<TestBehaviorReader> FRigUnit_RigLogic::TestAccessor::CreateBehaviorReaderTwoBlendShapes(FString MeshNameStr, FString BlendShape1Str, FString BlendShape2Str)
 {
-	TestBehaviorReader* BehaviorReader = new TestBehaviorReader;
+	TUniquePtr<TestBehaviorReader> BehaviorReader = MakeUnique<TestBehaviorReader>();
 	BehaviorReader->addBlendShapeChannelName(*BlendShape1Str);
 	BehaviorReader->addBlendShapeChannelName(*BlendShape2Str);
 	BehaviorReader->addMeshName(*MeshNameStr);
@@ -127,17 +127,17 @@ TestBehaviorReader* FRigUnit_RigLogic::TestAccessor::CreateBehaviorReaderTwoBlen
 	return BehaviorReader;
 }
 
-FRigCurveContainer* FRigUnit_RigLogic::TestAccessor::CreateCurveContainerOneMorphTarget(FString MorphTargetStr)
+TUniquePtr<FRigCurveContainer> FRigUnit_RigLogic::TestAccessor::CreateCurveContainerOneMorphTarget(FString MorphTargetStr)
 {
-	FRigCurveContainer* ValidCurveContainer = new FRigCurveContainer();
+	TUniquePtr<FRigCurveContainer> ValidCurveContainer = MakeUnique<FRigCurveContainer>();
 	ValidCurveContainer->Add(FName(*MorphTargetStr));
 	ValidCurveContainer->Initialize();
 	return ValidCurveContainer;
 }
 
-FRigCurveContainer* FRigUnit_RigLogic::TestAccessor::CreateCurveContainerTwoMorphTargets(FString MorphTarget1Str, FString MorphTarget2Str)
+TUniquePtr<FRigCurveContainer> FRigUnit_RigLogic::TestAccessor::CreateCurveContainerTwoMorphTargets(FString MorphTarget1Str, FString MorphTarget2Str)
 {
-	FRigCurveContainer* ValidCurveContainer = new FRigCurveContainer();
+	TUniquePtr<FRigCurveContainer> ValidCurveContainer = MakeUnique<FRigCurveContainer>();
 	ValidCurveContainer->Add(FName(*MorphTarget1Str));
 	ValidCurveContainer->Add(FName(*MorphTarget2Str));
 	ValidCurveContainer->Initialize();
@@ -152,9 +152,9 @@ void FRigUnit_RigLogic::TestAccessor::Exec_MapMorphTargets(TestBehaviorReader* T
 
 /** ====== Map Mask Multipliers ===== **/
 
-TestBehaviorReader* FRigUnit_RigLogic::TestAccessor::CreateBehaviorReaderOneAnimatedMap(FString AnimatedMapNameStr)
+TUniquePtr<TestBehaviorReader> FRigUnit_RigLogic::TestAccessor::CreateBehaviorReaderOneAnimatedMap(FString AnimatedMapNameStr)
 {
-	TestBehaviorReader* BehaviorReader = new TestBehaviorReader;
+	TUniquePtr<TestBehaviorReader> BehaviorReader = MakeUnique<TestBehaviorReader>();
 	BehaviorReader->animatedMaps.Add(AnimatedMapNameStr);
 	BehaviorReader->addAnimatedMapIndicesToLOD(0, 0);
 	BehaviorReader->LODCount = 1;
@@ -254,25 +254,24 @@ IMPLEMENT_RIGUNIT_AUTOMATION_TEST(FRigUnit_RigLogic)
 	//=== MapInputCurve ValidReader ValidCurvesNameMismatch ===
 
 	//Prepare
-	TestBehaviorReader* TestReaderValid = Test.CreateBehaviorReaderOneCurve("CTRL_Expressions.Some_Control");
-	FRigCurveContainer* TestCurveContainerNameMismatch = Test.CreateCurveContainerOneCurve("CTRL_Expressions_NOT_ThatControl");
+	TUniquePtr<TestBehaviorReader> TestReaderValid = Test.CreateBehaviorReaderOneCurve("CTRL_Expressions.Some_Control");
+	TUniquePtr<FRigCurveContainer> TestCurveContainerNameMismatch = Test.CreateCurveContainerOneCurve("CTRL_Expressions_NOT_ThatControl");
 	//Test
-	Test.Exec_MapInputCurve(TestReaderValid, TestCurveContainerNameMismatch);
+	Test.Exec_MapInputCurve(TestReaderValid.Get(), TestCurveContainerNameMismatch.Get());
 	//Assert
 	AddErrorIfFalse(
 		Test.GetData()->InputCurveIndices.Num() == 1 &&
 		Test.GetData()->InputCurveIndices[0] == INDEX_NONE,
 		TEXT("MapInputCurve - ValidReader CurveContainerWithNameMismatch")
 	);
-	delete TestCurveContainerNameMismatch;
 
 	//=== MapInputCurve EmptyReader ValidCurve ===
 
 	//Prepare
-	TestBehaviorReader* TestReaderEmpty = Test.CreateBehaviorReaderEmpty();
-	FRigCurveContainer* TestCurveContainerValid = Test.CreateCurveContainerOneCurve("CTRL_Expressions_Some_Control");
+	TUniquePtr<TestBehaviorReader> TestReaderEmpty = Test.CreateBehaviorReaderEmpty();
+	TUniquePtr<FRigCurveContainer> TestCurveContainerValid = Test.CreateCurveContainerOneCurve("CTRL_Expressions_Some_Control");
 	//Test
-	Test.Exec_MapInputCurve(TestReaderEmpty, TestCurveContainerValid);
+	Test.Exec_MapInputCurve(TestReaderEmpty.Get(), TestCurveContainerValid.Get());
 	//Assert
 	AddErrorIfFalse(
 		Test.GetData()->InputCurveIndices.Num() == 0,
@@ -282,10 +281,10 @@ IMPLEMENT_RIGUNIT_AUTOMATION_TEST(FRigUnit_RigLogic)
 	//=== MapInputCurve ValidReader EmptyCurveContainer ===
 
 	//Prepare
-	FRigCurveContainer* TestCurveContainerEmpty = Test.CreateCurveContainerEmpty();
+	TUniquePtr<FRigCurveContainer> TestCurveContainerEmpty = Test.CreateCurveContainerEmpty();
 	TestCurveContainerEmpty->Initialize();
 	//Test
-	Test.Exec_MapInputCurve(TestReaderValid, TestCurveContainerEmpty);
+	Test.Exec_MapInputCurve(TestReaderValid.Get(), TestCurveContainerEmpty.Get());
 	//Assert
 	AddErrorIfFalse(
 		Test.GetData()->InputCurveIndices.Num() == 1 &&
@@ -295,25 +294,23 @@ IMPLEMENT_RIGUNIT_AUTOMATION_TEST(FRigUnit_RigLogic)
 
 	//=== MapInputCurve InvalidReader ValidCurveContainer ===
 	//Prepare
-	TestBehaviorReader* TestInvalidReader = Test.CreateBehaviorReaderOneCurve("InvalidControlNameNoDot");
+	TUniquePtr<TestBehaviorReader> TestInvalidReader = Test.CreateBehaviorReaderOneCurve("InvalidControlNameNoDot");
 	//Expected error
 	AddExpectedError("RigUnit_R: Missing '.' in ");
 	//Test
-	Test.Exec_MapInputCurve(TestInvalidReader, TestCurveContainerValid);
+	Test.Exec_MapInputCurve(TestInvalidReader.Get(), TestCurveContainerValid.Get());
 	//Assert
 	AddErrorIfFalse(
 		Test.GetData()->InputCurveIndices.Num() == 0,
 		TEXT("MapInputCurve - InvalidReader ValidCurveContainer")
 	);
 
-	delete TestInvalidReader;
-
 	//=== MapInputCurve Valid Inputs ===
 
 	//Prepare
 	//  ---
 	//Test
-	Test.Exec_MapInputCurve(TestReaderValid, TestCurveContainerValid);
+	Test.Exec_MapInputCurve(TestReaderValid.Get(), TestCurveContainerValid.Get());
 	//Assert
 	AddErrorIfFalse(
 		Test.GetData()->InputCurveIndices.Num() == 1 &&
@@ -321,17 +318,14 @@ IMPLEMENT_RIGUNIT_AUTOMATION_TEST(FRigUnit_RigLogic)
 		TEXT("MapInputCurve - Valid Inputs")
 	);
 
-	delete TestReaderValid;
-	delete TestCurveContainerValid;
-
 	//===================== JOINTS MAPPING =====================
 
 
 	//=== MapJoints EmptyInputs ===
 	//Prepare
-	FRigBoneHierarchy* TestHierarchyEmpty = Test.CreateBoneHierarchyEmpty();
+	TUniquePtr<FRigBoneHierarchy> TestHierarchyEmpty = Test.CreateBoneHierarchyEmpty();
 	//Test
-	Test.Exec_MapJoints(TestReaderEmpty, TestHierarchyEmpty);
+	Test.Exec_MapJoints(TestReaderEmpty.Get(), TestHierarchyEmpty.Get());
 	//Assert
 	AddErrorIfFalse(
 		Test.GetData()->HierarchyBoneIndices.Num() == 0,
@@ -340,9 +334,9 @@ IMPLEMENT_RIGUNIT_AUTOMATION_TEST(FRigUnit_RigLogic)
 
 	//=== MapJoints EmptyReader TwoBones ===
 	//Prepare
-	FRigBoneHierarchy* TestHierarchyTwoBones = Test.CreateBoneHierarchyTwoBones("BoneA", "BoneB");
+	TUniquePtr<FRigBoneHierarchy> TestHierarchyTwoBones = Test.CreateBoneHierarchyTwoBones("BoneA", "BoneB");
 	//Test
-	Test.Exec_MapJoints(TestReaderEmpty, TestHierarchyTwoBones);
+	Test.Exec_MapJoints(TestReaderEmpty.Get(), TestHierarchyTwoBones.Get());
 	AddErrorIfFalse(
 		Test.GetData()->HierarchyBoneIndices.Num() == 0,
 		TEXT("MapJoints - EmptyReader TwoBones")
@@ -350,9 +344,9 @@ IMPLEMENT_RIGUNIT_AUTOMATION_TEST(FRigUnit_RigLogic)
 
 	//=== MapJoints TwoJoints NoBones ===
 	//Prepare
-	TestBehaviorReader* TestReaderTwoJoints = Test.CreateBehaviorReaderTwoJoints("BoneA", "BoneB");
+	TUniquePtr<TestBehaviorReader> TestReaderTwoJoints = Test.CreateBehaviorReaderTwoJoints("BoneA", "BoneB");
 	//Test
-	Test.Exec_MapJoints(TestReaderTwoJoints, TestHierarchyEmpty);
+	Test.Exec_MapJoints(TestReaderTwoJoints.Get(), TestHierarchyEmpty.Get());
 	AddErrorIfFalse(
 		Test.GetData()->HierarchyBoneIndices.Num() == 2,
 		TEXT("MapJoints - TwoJoints NoBones - expected 2 bone indices")
@@ -368,13 +362,11 @@ IMPLEMENT_RIGUNIT_AUTOMATION_TEST(FRigUnit_RigLogic)
 		TEXT("MapJoints - TwoJoints NoBones - Expected joint 1 index to be NONE")
 	);
 
-	delete TestHierarchyEmpty;
-
 	//=== MapJoints TwoJoints TwoBones ===
 	//Prepare
 	//  already done
 	//Test
-	Test.Exec_MapJoints(TestReaderTwoJoints, TestHierarchyTwoBones);
+	Test.Exec_MapJoints(TestReaderTwoJoints.Get(), TestHierarchyTwoBones.Get());
 	//Assert
 	AddErrorIfFalse(
 		Test.GetData()->HierarchyBoneIndices.Num() == 2,
@@ -391,19 +383,16 @@ IMPLEMENT_RIGUNIT_AUTOMATION_TEST(FRigUnit_RigLogic)
 		TEXT("MapJoints - TwoJoints TwoBones - Expected bone index 1 index to be 1")
 	);
 
-	delete TestHierarchyTwoBones;
-	delete TestReaderTwoJoints;
-
 	//===================== BLENDSHAPES MAPPING =====================
 
 
 	//=== MapMorphTargets ValidReader MorphTargetWithNameMismatch ===
 
 	//Prepare
-	TestBehaviorReader* TestReaderBlendshapeValid = Test.CreateBehaviorReaderOneBlendShape("head", "blendshape");
-	FRigCurveContainer* TestMorphTargetNameMismatch = Test.CreateCurveContainerOneMorphTarget("head_NOT_that_blendshape");
+	TUniquePtr<TestBehaviorReader> TestReaderBlendshapeValid = Test.CreateBehaviorReaderOneBlendShape("head", "blendshape");
+	TUniquePtr<FRigCurveContainer> TestMorphTargetNameMismatch = Test.CreateCurveContainerOneMorphTarget("head_NOT_that_blendshape");
 	//Test
-	Test.Exec_MapMorphTargets(TestReaderBlendshapeValid, TestMorphTargetNameMismatch);
+	Test.Exec_MapMorphTargets(TestReaderBlendshapeValid.Get(), TestMorphTargetNameMismatch.Get());
 	//Assert
 	AddErrorIfFalse(
 		Test.GetData()->BlendShapeIndices.Num() == 1 && //LOD 0
@@ -415,15 +404,13 @@ IMPLEMENT_RIGUNIT_AUTOMATION_TEST(FRigUnit_RigLogic)
 		TEXT("MapMorphTargets - ValidReader MorphTargetWithNameMismatch")
 	);
 
-	delete TestMorphTargetNameMismatch;
-
 	//=== MapMorphTargets EmptyReader ValidMorphTargetCurve ===
 
 	//Prepare
 	//Empty reader (no meshes, no blendshapes)
-	FRigCurveContainer* TestMorphTargetCurveValid = Test.CreateCurveContainerOneMorphTarget("head__blendshape");
+	TUniquePtr<FRigCurveContainer> TestMorphTargetCurveValid = Test.CreateCurveContainerOneMorphTarget("head__blendshape");
 	//Test
-	Test.Exec_MapMorphTargets(TestReaderEmpty, TestMorphTargetCurveValid);
+	Test.Exec_MapMorphTargets(TestReaderEmpty.Get(), TestMorphTargetCurveValid.Get());
 	//Assert
 	AddErrorIfFalse(
 		Test.GetData()->MorphTargetCurveIndices.Num() == 0 &&
@@ -435,9 +422,9 @@ IMPLEMENT_RIGUNIT_AUTOMATION_TEST(FRigUnit_RigLogic)
 	//=== MapMorphTargets NoBlendShapes ValidMorphTargetCurve ===
 
 	//Prepare
-	TestBehaviorReader* TestReaderNoBlendshapes = Test.CreateBehaviorReaderNoBlendshapes("head"); //has a mesh, but no blend shapes
+	TUniquePtr<TestBehaviorReader> TestReaderNoBlendshapes = Test.CreateBehaviorReaderNoBlendshapes("head"); //has a mesh, but no blend shapes
 	//Test
-	Test.Exec_MapMorphTargets(TestReaderNoBlendshapes, TestMorphTargetCurveValid);
+	Test.Exec_MapMorphTargets(TestReaderNoBlendshapes.Get(), TestMorphTargetCurveValid.Get());
 	//Assert
 	AddErrorIfFalse(
 		Test.GetData()->BlendShapeIndices.Num() == 1 && //LOD 0 exists
@@ -446,14 +433,12 @@ IMPLEMENT_RIGUNIT_AUTOMATION_TEST(FRigUnit_RigLogic)
 		TEXT("MapMorphTargets - NoBlendShapes ValidMorphTargetCurve")
 	);
 
-	delete TestReaderNoBlendshapes;
-
 	//=== MapMorphTargets ValidReader EmptyCurveContainer ===
 
 	//Prepare
 	//  ---
 	//Test
-	Test.Exec_MapMorphTargets(TestReaderBlendshapeValid, TestCurveContainerEmpty);
+	Test.Exec_MapMorphTargets(TestReaderBlendshapeValid.Get(), TestCurveContainerEmpty.Get());
 	//Assert
 	AddErrorIfFalse(
 		Test.GetData()->BlendShapeIndices.Num() == 1 && //LOD 0
@@ -467,9 +452,9 @@ IMPLEMENT_RIGUNIT_AUTOMATION_TEST(FRigUnit_RigLogic)
 	//=== MapMorphTargets InvalidReader ValidMorphTargetCurve ===
 
 	//Prepare
-	TestBehaviorReader* TestReaderBlendshapesInvalid = Test.CreateBehaviorReaderOneBlendShape("head", "");
+	TUniquePtr<TestBehaviorReader> TestReaderBlendshapesInvalid = Test.CreateBehaviorReaderOneBlendShape("head", "");
 	//Test
-	Test.Exec_MapMorphTargets(TestReaderBlendshapesInvalid, TestMorphTargetCurveValid);
+	Test.Exec_MapMorphTargets(TestReaderBlendshapesInvalid.Get(), TestMorphTargetCurveValid.Get());
 	//Assert
 	AddErrorIfFalse(
 		Test.GetData()->BlendShapeIndices.Num() == 1 && //LOD 0
@@ -480,14 +465,12 @@ IMPLEMENT_RIGUNIT_AUTOMATION_TEST(FRigUnit_RigLogic)
 		TEXT("MapMorphTargets - InvalidReader ValidMorphTargetCurve")
 	);
 
-	delete TestReaderBlendshapesInvalid;
-
 	//=== MapMorphTargets ValidReader InvalidMorphTargetCurve ===
 
 	//Prepare
-	FRigCurveContainer* TestMorphTargetCurvesInvalid = Test.CreateCurveContainerOneMorphTarget("");
+	TUniquePtr<FRigCurveContainer> TestMorphTargetCurvesInvalid = Test.CreateCurveContainerOneMorphTarget("");
 	//Test
-	Test.Exec_MapMorphTargets(TestReaderBlendshapeValid, TestMorphTargetCurvesInvalid);
+	Test.Exec_MapMorphTargets(TestReaderBlendshapeValid.Get(), TestMorphTargetCurvesInvalid.Get());
 	//Assert
 	AddErrorIfFalse(
 		Test.GetData()->BlendShapeIndices.Num() == 1 && //LOD 0
@@ -498,14 +481,12 @@ IMPLEMENT_RIGUNIT_AUTOMATION_TEST(FRigUnit_RigLogic)
 		TEXT("MapMorphTargets - ValidReader InvalidMorphTargetCurve")
 	);
 
-	delete TestMorphTargetCurvesInvalid;
-
 	//=== MapMorphTargets Valid Inputs ===
 
 	//Prepare
 	//  ---
 	//Test
-	Test.Exec_MapMorphTargets(TestReaderBlendshapeValid, TestMorphTargetCurveValid);
+	Test.Exec_MapMorphTargets(TestReaderBlendshapeValid.Get(), TestMorphTargetCurveValid.Get());
 	//Assert
 	AddErrorIfFalse(
 		Test.GetData()->BlendShapeIndices.Num() == 1 && //LOD 0
@@ -517,14 +498,11 @@ IMPLEMENT_RIGUNIT_AUTOMATION_TEST(FRigUnit_RigLogic)
 		TEXT("MapMorphTargets - ValidReader ValidTestMorphTarget")
 	);
 
-	delete TestReaderBlendshapeValid;
-	delete TestMorphTargetCurveValid;
-
 	//=== MapMorphTargets LOD0(AB) LOD1(A) ===
 
 	//Prepare
-	TestBehaviorReader* TestReaderBlendshapes_LOD0AB_LOD1A = Test.CreateBehaviorReaderTwoBlendShapes("head", "blendshapeA", "blendshapeB");
-	FRigCurveContainer* TestMorphTargetTwoCurves = Test.CreateCurveContainerTwoMorphTargets("head__blendshapeA", "head__blendshapeB");
+	TUniquePtr<TestBehaviorReader> TestReaderBlendshapes_LOD0AB_LOD1A = Test.CreateBehaviorReaderTwoBlendShapes("head", "blendshapeA", "blendshapeB");
+	TUniquePtr<FRigCurveContainer> TestMorphTargetTwoCurves = Test.CreateCurveContainerTwoMorphTargets("head__blendshapeA", "head__blendshapeB");
 	//NOTE: indices in the first param here are not blendshape indices, but rather mappings from blendshapes to meshes
 	//in this test, they will correspond to blendshape indices
 
@@ -535,7 +513,7 @@ IMPLEMENT_RIGUNIT_AUTOMATION_TEST(FRigUnit_RigLogic)
 
 	TestReaderBlendshapes_LOD0AB_LOD1A->LODCount = 2; //needs to be set explicitly if not default (=1)
 	//Test
-	Test.Exec_MapMorphTargets(TestReaderBlendshapes_LOD0AB_LOD1A, TestMorphTargetTwoCurves);
+	Test.Exec_MapMorphTargets(TestReaderBlendshapes_LOD0AB_LOD1A.Get(), TestMorphTargetTwoCurves.Get());
 	//Assert
 	AddErrorIfFalse(
 		Test.GetData()->BlendShapeIndices.Num() == 2 &&  //2 LODs
@@ -563,18 +541,16 @@ IMPLEMENT_RIGUNIT_AUTOMATION_TEST(FRigUnit_RigLogic)
 		TEXT("MapMorphTargets LOD0(AB) LOD1(A) - resulting morph target indices not correct")
 	);
 
-	delete TestReaderBlendshapes_LOD0AB_LOD1A;
-
 	//=== MapMorphTargets LOD0(AB) LOD1(-) ===
 
 	//Prepare
-	TestBehaviorReader* TestReaderBlendshapes_LOD0AB_LOD1N = Test.CreateBehaviorReaderTwoBlendShapes("head", "blendshapeA", "blendshapeB");
+	TUniquePtr<TestBehaviorReader> TestReaderBlendshapes_LOD0AB_LOD1N = Test.CreateBehaviorReaderTwoBlendShapes("head", "blendshapeA", "blendshapeB");
 	TestReaderBlendshapes_LOD0AB_LOD1N->addBlendShapeMappingIndicesToLOD(0, 0); //LOD 0
 	TestReaderBlendshapes_LOD0AB_LOD1N->addBlendShapeMappingIndicesToLOD(1, 0); //LOD 0
 	//LODCount = 1 by default
 
 	//Test
-	Test.Exec_MapMorphTargets(TestReaderBlendshapes_LOD0AB_LOD1N, TestMorphTargetTwoCurves);
+	Test.Exec_MapMorphTargets(TestReaderBlendshapes_LOD0AB_LOD1N.Get(), TestMorphTargetTwoCurves.Get());
 	//Assert
 	AddErrorIfFalse(
 		Test.GetData()->BlendShapeIndices.Num() == 1 &&  //1 LOD
@@ -599,17 +575,15 @@ IMPLEMENT_RIGUNIT_AUTOMATION_TEST(FRigUnit_RigLogic)
 		TEXT("MapMorphTargets LOD0(AB) LOD1(-) - Resulting morph targets not correct")
 	);
 
-	delete TestReaderBlendshapes_LOD0AB_LOD1N;
-
 	//=== MapMorphTargets LOD0(A) LOD1(B) ===
 	//Prepare
-	TestBehaviorReader* TestReaderBlendshapes_LOD0A_LOD1B = Test.CreateBehaviorReaderTwoBlendShapes("head", "blendshapeA", "blendshapeB");
+	TUniquePtr<TestBehaviorReader> TestReaderBlendshapes_LOD0A_LOD1B = Test.CreateBehaviorReaderTwoBlendShapes("head", "blendshapeA", "blendshapeB");
 	TestReaderBlendshapes_LOD0A_LOD1B->addBlendShapeMappingIndicesToLOD(0, 0);
 	TestReaderBlendshapes_LOD0A_LOD1B->addBlendShapeMappingIndicesToLOD(1, 1);
 	TestReaderBlendshapes_LOD0A_LOD1B->LODCount = 2;
 
 	//Test
-	Test.Exec_MapMorphTargets(TestReaderBlendshapes_LOD0A_LOD1B, TestMorphTargetTwoCurves);
+	Test.Exec_MapMorphTargets(TestReaderBlendshapes_LOD0A_LOD1B.Get(), TestMorphTargetTwoCurves.Get());
 	//Assert
 	AddErrorIfFalse(
 		Test.GetData()->BlendShapeIndices.Num() == 2 &&  //2 LODs
@@ -635,19 +609,16 @@ IMPLEMENT_RIGUNIT_AUTOMATION_TEST(FRigUnit_RigLogic)
 		TEXT("MapMorphTargets LOD0(A) LOD1(B) - Resulting morph target indices not correct")
 	);
 
-	delete TestMorphTargetTwoCurves;
-	delete TestReaderBlendshapes_LOD0A_LOD1B;
-
 	//=============== MASK MULTIPLIERS MAPPING ====================
 
 
 	//=== MapMaskMultipliers ValidReader ValidAnimatedMapNameMismatch ===
 
 	//Prepare
-	TestBehaviorReader* TestReaderAnimMapsValid = Test.CreateBehaviorReaderOneAnimatedMap("CTRL_AnimMap.Some_Multiplier");
-	FRigCurveContainer* TestCurveContainerForAnimMapsNameMismatch = Test.CreateCurveContainerOneCurve("CTRL_AnimMap_NOT_ThatMultiploer");
+	TUniquePtr<TestBehaviorReader> TestReaderAnimMapsValid = Test.CreateBehaviorReaderOneAnimatedMap("CTRL_AnimMap.Some_Multiplier");
+	TUniquePtr<FRigCurveContainer> TestCurveContainerForAnimMapsNameMismatch = Test.CreateCurveContainerOneCurve("CTRL_AnimMap_NOT_ThatMultiploer");
 	//Test
-	Test.Exec_MapMaskMultipliers(TestReaderAnimMapsValid, TestCurveContainerForAnimMapsNameMismatch);
+	Test.Exec_MapMaskMultipliers(TestReaderAnimMapsValid.Get(), TestCurveContainerForAnimMapsNameMismatch.Get());
 	//Assert
 	AddErrorIfFalse(
 		Test.GetData()->CurveContainerIndicesForAnimMaps.Num() == 1 &&
@@ -659,14 +630,12 @@ IMPLEMENT_RIGUNIT_AUTOMATION_TEST(FRigUnit_RigLogic)
 		TEXT("MapMaskMultipliers - ValidReader ValidAnimatedMapNameMismatch")
 	);
 
-	delete TestCurveContainerForAnimMapsNameMismatch;
-
 	//=== MapMaskMultipliers EmptyReader ValidAnimatedMap ===
 
 	//Prepare
-	FRigCurveContainer* TestCurveContainerForAnimMapsValid = Test.CreateCurveContainerOneCurve("CTRL_AnimMap_Some_Multiplier");
+	TUniquePtr<FRigCurveContainer> TestCurveContainerForAnimMapsValid = Test.CreateCurveContainerOneCurve("CTRL_AnimMap_Some_Multiplier");
 	//Test
-	Test.Exec_MapMaskMultipliers(TestReaderEmpty, TestCurveContainerForAnimMapsValid);
+	Test.Exec_MapMaskMultipliers(TestReaderEmpty.Get(), TestCurveContainerForAnimMapsValid.Get());
 	//Assert
 	AddErrorIfFalse(
 		Test.GetData()->CurveContainerIndicesForAnimMaps.Num() == 0 &&
@@ -674,12 +643,10 @@ IMPLEMENT_RIGUNIT_AUTOMATION_TEST(FRigUnit_RigLogic)
 		TEXT("MapMaskMultipliers - EmptyReader ValidAnimatedMap")
 	);
 
-	delete TestReaderEmpty;
-
 	//=== MapMaskMultipliers ValidReader EmptyCurveContainer ===
 
 	//Test
-	Test.Exec_MapMaskMultipliers(TestReaderAnimMapsValid, TestCurveContainerEmpty);
+	Test.Exec_MapMaskMultipliers(TestReaderAnimMapsValid.Get(), TestCurveContainerEmpty.Get());
 	//Assert
 	AddErrorIfFalse(
 		Test.GetData()->CurveContainerIndicesForAnimMaps.Num() == 1 &&
@@ -691,14 +658,12 @@ IMPLEMENT_RIGUNIT_AUTOMATION_TEST(FRigUnit_RigLogic)
 		TEXT("MapMaskMultipliers - ValidReader EmptyCurveContainer")
 	);
 
-	delete TestCurveContainerEmpty;
-
 	//=== MapMaskMultipliers Valid Inputs ===
 
 	//Prepare
 	//  ---
 	//Test
-	Test.Exec_MapMaskMultipliers(TestReaderAnimMapsValid, TestCurveContainerForAnimMapsValid);
+	Test.Exec_MapMaskMultipliers(TestReaderAnimMapsValid.Get(), TestCurveContainerForAnimMapsValid.Get());
 	//Assert
 	AddErrorIfFalse(
 		Test.GetData()->CurveContainerIndicesForAnimMaps.Num() == 1 &&
@@ -710,9 +675,6 @@ IMPLEMENT_RIGUNIT_AUTOMATION_TEST(FRigUnit_RigLogic)
 		TEXT("MapMaskMultipliers - Valid Inputs")
 	);
 
-	delete TestReaderAnimMapsValid;
-	delete TestCurveContainerForAnimMapsValid;
-
 	Unit.ExecuteContext.Hierarchy = &HierarchyContainer;
 	//BoneHierarchy belongs to HierarchyContainer
 	BoneHierarchy = *TestHierarchyTwoBones;
@@ -722,7 +684,6 @@ IMPLEMENT_RIGUNIT_AUTOMATION_TEST(FRigUnit_RigLogic)
 	
 	FRigHierarchyContainer* HierarchyContainerPtr = &HierarchyContainer;
 
-
 	//Prepare
 	//-----
 	//create neutral transforms for two bones
diff --git a/Engine/Plugins/Animation/RigLogic/Source/RigLogicModule/Private/SkelMeshDNAUtils.cpp b/Engine/Plugins/Animation/RigLogic/Source/RigLogicModule/Private/SkelMeshDNAUtils.cpp
--- a/Engine/Plugins/Animation/RigLogic/Source/RigLogicModule/Private/SkelMeshDNAUtils.cpp
+++ b/Engine/Plugins/Animation/RigLogic/Source/RigLogicModule/Private/SkelMeshDNAUtils.cpp
@@ -208,7 +208,7 @@ void USkelMeshDNAUtils::UpdateMorphTargets(USkeletalMesh* SkelMesh, IDNAReader*
 				for (int32 DeltaIndex = 0; DeltaIndex < NumDeltas; DeltaIndex++)
 				{
 					uint32 DeltaSourceIndex = Deltas[DeltaIndex].SourceIdx;
-					uint32 DNAVertexIndex = DNAToSkelMeshMap->ImportVtxToDNAVtxIndex[LODIndex][DeltaSourceIndex];
+					int32 DNAVertexIndex = DNAToSkelMeshMap->ImportVtxToDNAVtxIndex[LODIndex][DeltaSourceIndex];
 					if (DNAVertexIndex >= 0)
 					{
 						// Retrieve delta index from the mapping.
diff --git a/Engine/Plugins/Animation/RigLogic/Source/RigLogicModule/Public/RigUnit_RigLogic_Test.h b/Engine/Plugins/Animation/RigLogic/Source/RigLogicModule/Public/RigUnit_RigLogic_Test.h
--- a/Engine/Plugins/Animation/RigLogic/Source/RigLogicModule/Public/RigUnit_RigLogic_Test.h
+++ b/Engine/Plugins/Animation/RigLogic/Source/RigLogicModule/Public/RigUnit_RigLogic_Test.h
@@ -468,38 +468,37 @@ struct FRigUnit_RigLogic::TestAccessor
 
 	 /** MapInputCurve Tests **/
 
-	TestBehaviorReader* CreateBehaviorReaderEmpty();
-	TestBehaviorReader* CreateBehaviorReaderOneCurve(FString ControlNameStr);
+	TUniquePtr<TestBehaviorReader> CreateBehaviorReaderEmpty();
+	TUniquePtr<TestBehaviorReader> CreateBehaviorReaderOneCurve(FString ControlNameStr);
 
-	FRigCurveContainer* CreateCurveContainerEmpty();
-	FRigCurveContainer* CreateCurveContainerOneCurve(FString CurveNameStr);
+	TUniquePtr<FRigCurveContainer> CreateCurveContainerEmpty();
+	TUniquePtr<FRigCurveContainer> CreateCurveContainerOneCurve(FString CurveNameStr);
 
 	void Exec_MapInputCurve(TestBehaviorReader* TestReader, FRigCurveContainer* TestCurveContainer);
 
 	/** MapJoints Tests **/
 
-	TestBehaviorReader* CreateBehaviorReaderTwoJoints(FString Joint1NameStr, FString Joint2NameStr);
-	FRigBoneHierarchy* CreateBoneHierarchyEmpty();
-	FRigBoneHierarchy* CreateBoneHierarchyTwoBones(FString Bone1NameStr, FString Bone2NameStr);
+	TUniquePtr<TestBehaviorReader> CreateBehaviorReaderTwoJoints(FString Joint1NameStr, FString Joint2NameStr);
+	TUniquePtr<FRigBoneHierarchy> CreateBoneHierarchyEmpty();
+	TUniquePtr<FRigBoneHierarchy> CreateBoneHierarchyTwoBones(FString Bone1NameStr, FString Bone2NameStr);
 
 	void Exec_MapJoints(TestBehaviorReader* TestReader, FRigBoneHierarchy* MockHierachy);
 
 	/** MapBlendShapes Tests **/
 
-	TestBehaviorReader* CreateBehaviorReaderNoBlendshapes(FString MeshNameStr);
-	TestBehaviorReader* CreateBehaviorReaderOneBlendShape(FString MeshNameStr, FString BlendShapeNameStr);
-	FRigCurveContainer* CreateCurveContainerOneMorphTarget(FString MorphTargetStr);
-	TestBehaviorReader* CreateBehaviorReaderTwoBlendShapes(FString MeshNameStr, FString BlendShape1Str, FString BlendShape2Str);
-	FRigCurveContainer* CreateCurveContainerTwoMorphTargets(FString MorphTarget1Str, FString MorphTarget2Str);
+	TUniquePtr<TestBehaviorReader> CreateBehaviorReaderNoBlendshapes(FString MeshNameStr);
+	TUniquePtr<TestBehaviorReader> CreateBehaviorReaderOneBlendShape(FString MeshNameStr, FString BlendShapeNameStr);
+	TUniquePtr<FRigCurveContainer> CreateCurveContainerOneMorphTarget(FString MorphTargetStr);
+	TUniquePtr<TestBehaviorReader> CreateBehaviorReaderTwoBlendShapes(FString MeshNameStr, FString BlendShape1Str, FString BlendShape2Str);
+	TUniquePtr<FRigCurveContainer> CreateCurveContainerTwoMorphTargets(FString MorphTarget1Str, FString MorphTarget2Str);
 
 	void Exec_MapMorphTargets(TestBehaviorReader* TestReader, FRigCurveContainer* TestCurveContainer);
 
 	/** MapMaskMultipliers Tests **/
-	TestBehaviorReader* CreateBehaviorReaderOneAnimatedMap(FString AnimatedMapNameStr);
+	TUniquePtr<TestBehaviorReader> CreateBehaviorReaderOneAnimatedMap(FString AnimatedMapNameStr);
 
 	void Exec_MapMaskMultipliers(TestBehaviorReader* TestReader, FRigCurveContainer* TestCurveContainer);
 
-
 	static const uint8 MAX_ATTRS_PER_JOINT;
 
 	/** UpdateJoints Tests **/
