RigLogicLib: Fix wrong include causing circular dependency #rb none

diff --git a/Engine/Plugins/Animation/RigLogic/Source/RigLogicLib/Public/dna/types/StringView.h b/Engine/Plugins/Animation/RigLogic/Source/RigLogicLib/Public/dna/types/StringView.h
--- a/Engine/Plugins/Animation/RigLogic/Source/RigLogicLib/Public/dna/types/StringView.h
+++ b/Engine/Plugins/Animation/RigLogic/Source/RigLogicLib/Public/dna/types/StringView.h
@@ -2,7 +2,7 @@
 
 #pragma once
 
-#include "dna/types/Aliases.h"
+#include "dna/types/ArrayView.h"
 
 namespace dna {
 
diff --git a/Engine/Plugins/Animation/RigLogic/Source/RigLogicLib/Public/dna/version/Version.h b/Engine/Plugins/Animation/RigLogic/Source/RigLogicLib/Public/dna/version/Version.h
--- a/Engine/Plugins/Animation/RigLogic/Source/RigLogicLib/Public/dna/version/Version.h
+++ b/Engine/Plugins/Animation/RigLogic/Source/RigLogicLib/Public/dna/version/Version.h
@@ -4,5 +4,5 @@
 
 #define DNA_MAJOR_VERSION 1
 #define DNA_MINOR_VERSION 1
-#define DNA_PATCH_VERSION 3
-#define DNA_VERSION_STRING "1.1.3"
+#define DNA_PATCH_VERSION 4
+#define DNA_VERSION_STRING "1.1.4"
