[NO NON-PII COMMIT MESSAGE]

diff --git a/Engine/Plugins/Animation/RigLogic/RigLogic.uplugin b/Engine/Plugins/Animation/RigLogic/RigLogic.uplugin
--- a/Engine/Plugins/Animation/RigLogic/RigLogic.uplugin
+++ b/Engine/Plugins/Animation/RigLogic/RigLogic.uplugin
@@ -1,9 +1,9 @@
 {
 	"FileVersion": 3,
 	"Version": 6,
-	"VersionName": "6.0.1",
-	"FriendlyName": "RigLogic Plugin v6.0.1",
-	"Description": "3Lateral RigLogic Plugin for Facial Animation v6.0.1",
+	"VersionName": "6.0.2",
+	"FriendlyName": "RigLogic Plugin v6.0.2",
+	"Description": "3Lateral RigLogic Plugin for Facial Animation v6.0.2",
 	"Category": "Animation",
 	"CreatedBy": "3Lateral",
 	"CreatedByURL": "http://www.3lateral.com/",
diff --git a/Engine/Plugins/Animation/RigLogic/Source/RigLogicLib/Private/trio/streams/MemoryMappedFileStreamFallback.cpp b/Engine/Plugins/Animation/RigLogic/Source/RigLogicLib/Private/trio/streams/MemoryMappedFileStreamFallback.cpp
--- a/Engine/Plugins/Animation/RigLogic/Source/RigLogicLib/Private/trio/streams/MemoryMappedFileStreamFallback.cpp
+++ b/Engine/Plugins/Animation/RigLogic/Source/RigLogicLib/Private/trio/streams/MemoryMappedFileStreamFallback.cpp
@@ -30,9 +30,12 @@ static std::size_t getFileSizeFallback(const char* path) {
     #else
         stream = fopen(path, "rb");
     #endif
-    fseek(stream, 0, SEEK_END);
-    std::size_t fileSize = static_cast<std::size_t>(ftell(stream));
-    fclose(stream);
+    std::size_t fileSize{};
+    if (stream != nullptr) {
+        fseek(stream, 0, SEEK_END);
+        fileSize = static_cast<std::size_t>(ftell(stream));
+        fclose(stream);
+    }
     return fileSize;
 }
 
diff --git a/Engine/Plugins/Animation/RigLogic/Source/RigLogicLib/Public/dna/version/Version.h b/Engine/Plugins/Animation/RigLogic/Source/RigLogicLib/Public/dna/version/Version.h
--- a/Engine/Plugins/Animation/RigLogic/Source/RigLogicLib/Public/dna/version/Version.h
+++ b/Engine/Plugins/Animation/RigLogic/Source/RigLogicLib/Public/dna/version/Version.h
@@ -4,5 +4,5 @@
 
 #define DNA_MAJOR_VERSION 2
 #define DNA_MINOR_VERSION 0
-#define DNA_PATCH_VERSION 1
-#define DNA_VERSION_STRING "2.0.1"
+#define DNA_PATCH_VERSION 2
+#define DNA_VERSION_STRING "2.0.2"
diff --git a/Engine/Plugins/Animation/RigLogic/Source/RigLogicLib/Public/riglogic/version/Version.h b/Engine/Plugins/Animation/RigLogic/Source/RigLogicLib/Public/riglogic/version/Version.h
--- a/Engine/Plugins/Animation/RigLogic/Source/RigLogicLib/Public/riglogic/version/Version.h
+++ b/Engine/Plugins/Animation/RigLogic/Source/RigLogicLib/Public/riglogic/version/Version.h
@@ -4,5 +4,5 @@
 
 #define RL_MAJOR_VERSION 6
 #define RL_MINOR_VERSION 0
-#define RL_PATCH_VERSION 1
-#define RL_VERSION_STRING "6.0.1"
+#define RL_PATCH_VERSION 2
+#define RL_VERSION_STRING "6.0.2"
diff --git a/Engine/Plugins/Animation/RigLogic/Source/RigLogicLib/Public/trio/version/Version.h b/Engine/Plugins/Animation/RigLogic/Source/RigLogicLib/Public/trio/version/Version.h
--- a/Engine/Plugins/Animation/RigLogic/Source/RigLogicLib/Public/trio/version/Version.h
+++ b/Engine/Plugins/Animation/RigLogic/Source/RigLogicLib/Public/trio/version/Version.h
@@ -4,5 +4,5 @@
 
 #define TRIO_MAJOR_VERSION 2
 #define TRIO_MINOR_VERSION 0
-#define TRIO_PATCH_VERSION 1
-#define TRIO_VERSION_STRING "2.0.1"
+#define TRIO_PATCH_VERSION 2
+#define TRIO_VERSION_STRING "2.0.2"
diff --git a/Engine/Plugins/Animation/RigLogic/Source/RigLogicModule/Public/DNAReaderAdapter.h b/Engine/Plugins/Animation/RigLogic/Source/RigLogicModule/Public/DNAReaderAdapter.h
--- a/Engine/Plugins/Animation/RigLogic/Source/RigLogicModule/Public/DNAReaderAdapter.h
+++ b/Engine/Plugins/Animation/RigLogic/Source/RigLogicModule/Public/DNAReaderAdapter.h
@@ -122,4 +122,622 @@ private:
 	TUniquePtr<TWrappedReader, FWrappedReaderDeleter> ReaderPtr;
 };
 
-#include "DNAReaderAdapterImpl.h"
+template <class TWrappedReader>
+FDNAReader<TWrappedReader>::FDNAReader(TWrappedReader* Source) :
+	ReaderPtr{Source}
+{
+}
+
+template <class TWrappedReader>
+dna::Reader* FDNAReader<TWrappedReader>::Unwrap() const
+{
+	return ReaderPtr.Get();
+}
+
+template <class TWrappedReader>
+void FDNAReader<TWrappedReader>::FWrappedReaderDeleter::operator()(TWrappedReader* Pointer)
+{
+	if (Pointer != nullptr)
+	{
+		TWrappedReader::destroy(Pointer);
+	}
+}
+
+template <class TWrappedReader>
+FString FDNAReader<TWrappedReader>::GetName() const
+{
+	return FString(ANSI_TO_TCHAR(ReaderPtr->getName().data()));
+}
+
+template <class TWrappedReader>
+EArchetype FDNAReader<TWrappedReader>::GetArchetype() const
+{
+	return static_cast<EArchetype>(ReaderPtr->getArchetype());
+}
+
+template <class TWrappedReader>
+EGender FDNAReader<TWrappedReader>::GetGender() const
+{
+	return static_cast<EGender>(ReaderPtr->getGender());
+}
+
+template <class TWrappedReader>
+uint16 FDNAReader<TWrappedReader>::GetAge() const
+{
+	return ReaderPtr->getAge();
+}
+
+template <class TWrappedReader>
+uint32 FDNAReader<TWrappedReader>::GetMetaDataCount() const
+{
+	return ReaderPtr->getMetaDataCount();
+}
+
+template <class TWrappedReader>
+FString FDNAReader<TWrappedReader>::GetMetaDataKey(uint32 Index) const
+{
+	return FString(ANSI_TO_TCHAR(ReaderPtr->getMetaDataKey(Index).data()));
+}
+
+template <class TWrappedReader>
+FString FDNAReader<TWrappedReader>::GetMetaDataValue(const FString& Key) const
+{
+	return FString(ANSI_TO_TCHAR(ReaderPtr->getMetaDataValue(TCHAR_TO_ANSI(*Key)).data()));
+}
+
+template <class TWrappedReader>
+ETranslationUnit FDNAReader<TWrappedReader>::GetTranslationUnit() const
+{
+	return static_cast<ETranslationUnit>(ReaderPtr->getTranslationUnit());
+}
+
+template <class TWrappedReader>
+ERotationUnit FDNAReader<TWrappedReader>::GetRotationUnit() const
+{
+	return static_cast<ERotationUnit>(ReaderPtr->getRotationUnit());
+}
+
+template <class TWrappedReader>
+FCoordinateSystem FDNAReader<TWrappedReader>::GetCoordinateSystem() const
+{
+	const auto System = ReaderPtr->getCoordinateSystem();
+	return FCoordinateSystem{
+		static_cast<EDirection>(System.xAxis),
+		static_cast<EDirection>(System.yAxis),
+		static_cast<EDirection>(System.zAxis)
+	};
+}
+
+template <class TWrappedReader>
+uint16 FDNAReader<TWrappedReader>::GetLODCount() const
+{
+	return ReaderPtr->getLODCount();
+}
+
+template <class TWrappedReader>
+uint16 FDNAReader<TWrappedReader>::GetDBMaxLOD() const
+{
+	return ReaderPtr->getDBMaxLOD();
+}
+
+template <class TWrappedReader>
+FString FDNAReader<TWrappedReader>::GetDBComplexity() const
+{
+	return FString(ANSI_TO_TCHAR(ReaderPtr->getDBComplexity().data()));
+}
+
+template <class TWrappedReader>
+FString FDNAReader<TWrappedReader>::GetDBName() const
+{
+	return FString(ANSI_TO_TCHAR(ReaderPtr->getDBName().data()));
+}
+
+
+template <class TWrappedReader>
+uint16 FDNAReader<TWrappedReader>::GetGUIControlCount() const
+{
+	return ReaderPtr->getGUIControlCount();
+}
+
+template <class TWrappedReader>
+FString FDNAReader<TWrappedReader>::GetGUIControlName(uint16 Index) const
+{
+	return FString(ANSI_TO_TCHAR(ReaderPtr->getGUIControlName(Index).data()));
+}
+
+template <class TWrappedReader>
+uint16 FDNAReader<TWrappedReader>::GetRawControlCount() const
+{
+	return ReaderPtr->getRawControlCount();
+}
+
+template <class TWrappedReader>
+FString FDNAReader<TWrappedReader>::GetRawControlName(uint16 Index) const
+{
+	return FString(ANSI_TO_TCHAR(ReaderPtr->getRawControlName(Index).data()));
+}
+
+template <class TWrappedReader>
+uint16 FDNAReader<TWrappedReader>::GetJointCount() const
+{
+	return ReaderPtr->getJointCount();
+}
+
+template <class TWrappedReader>
+FString FDNAReader<TWrappedReader>::GetJointName(uint16 Index) const
+{
+	return FString(ANSI_TO_TCHAR(ReaderPtr->getJointName(Index).data()));
+}
+
+template <class TWrappedReader>
+TArrayView<const uint16> FDNAReader<TWrappedReader>::GetJointIndicesForLOD(uint16 LOD) const
+{
+	const auto Indices = ReaderPtr->getJointIndicesForLOD(LOD);
+	return TArrayView<const uint16>(Indices.data(), static_cast<int32>(Indices.size()));
+}
+
+template <class TWrappedReader>
+uint16 FDNAReader<TWrappedReader>::GetJointParentIndex(uint16 Index) const
+{
+	return ReaderPtr->getJointParentIndex(Index);
+}
+
+template <class TWrappedReader>
+uint16 FDNAReader<TWrappedReader>::GetBlendShapeChannelCount() const
+{
+	return ReaderPtr->getBlendShapeChannelCount();
+}
+
+template <class TWrappedReader>
+FString FDNAReader<TWrappedReader>::GetBlendShapeChannelName(uint16 Index) const
+{
+	return FString(ANSI_TO_TCHAR(ReaderPtr->getBlendShapeChannelName(Index).data()));
+}
+
+template <class TWrappedReader>
+TArrayView<const uint16> FDNAReader<TWrappedReader>::GetBlendShapeChannelIndicesForLOD(uint16 LOD) const
+{
+	const auto Indices = ReaderPtr->getBlendShapeChannelIndicesForLOD(LOD);
+	return TArrayView<const uint16>(Indices.data(), static_cast<int32>(Indices.size()));
+}
+
+template <class TWrappedReader>
+uint16 FDNAReader<TWrappedReader>::GetAnimatedMapCount() const
+{
+	return ReaderPtr->getAnimatedMapCount();
+}
+
+template <class TWrappedReader>
+FString FDNAReader<TWrappedReader>::GetAnimatedMapName(uint16 Index) const
+{
+	return FString(ANSI_TO_TCHAR(ReaderPtr->getAnimatedMapName(Index).data()));
+}
+
+template <class TWrappedReader>
+TArrayView<const uint16> FDNAReader<TWrappedReader>::GetAnimatedMapIndicesForLOD(uint16 LOD) const
+{
+	const auto Indices = ReaderPtr->getAnimatedMapIndicesForLOD(LOD);
+	return TArrayView<const uint16>(Indices.data(), static_cast<int32>(Indices.size()));
+}
+
+template <class TWrappedReader>
+uint16 FDNAReader<TWrappedReader>::GetMeshCount() const
+{
+	return ReaderPtr->getMeshCount();
+}
+
+template <class TWrappedReader>
+FString FDNAReader<TWrappedReader>::GetMeshName(uint16 Index) const
+{
+	return FString(ANSI_TO_TCHAR(ReaderPtr->getMeshName(Index).data()));
+}
+
+template <class TWrappedReader>
+TArrayView<const uint16> FDNAReader<TWrappedReader>::GetMeshIndicesForLOD(uint16 LOD) const
+{
+	const auto Indices = ReaderPtr->getMeshIndicesForLOD(LOD);
+	return TArrayView<const uint16>(Indices.data(), static_cast<int32>(Indices.size()));
+}
+
+template <class TWrappedReader>
+uint16 FDNAReader<TWrappedReader>::GetMeshBlendShapeChannelMappingCount() const
+{
+	return ReaderPtr->getMeshBlendShapeChannelMappingCount();
+}
+
+template <class TWrappedReader>
+FMeshBlendShapeChannelMapping FDNAReader<TWrappedReader>::GetMeshBlendShapeChannelMapping(uint16 Index) const
+{
+	const auto Mapping = ReaderPtr->getMeshBlendShapeChannelMapping(Index);
+	return FMeshBlendShapeChannelMapping{Mapping.meshIndex, Mapping.blendShapeChannelIndex};
+}
+
+template <class TWrappedReader>
+TArrayView<const uint16> FDNAReader<TWrappedReader>::GetMeshBlendShapeChannelMappingIndicesForLOD(uint16 LOD) const
+{
+	const auto Indices = ReaderPtr->getMeshBlendShapeChannelMappingIndicesForLOD(LOD);
+	return TArrayView<const uint16>(Indices.data(), static_cast<int32>(Indices.size()));
+}
+
+template <class TWrappedReader>
+FVector FDNAReader<TWrappedReader>::GetNeutralJointTranslation(uint16 Index) const
+{
+	const auto Translation = ReaderPtr->getNeutralJointTranslation(Index);
+	// X = X, Y = -Y, Z = Z
+	return FVector(Translation.x, -Translation.y, Translation.z);
+}
+
+template <class TWrappedReader>
+FVector FDNAReader<TWrappedReader>::GetNeutralJointRotation(uint16 Index) const
+{
+	const auto Translation = ReaderPtr->getNeutralJointRotation(Index);
+	// X = -Y, Y = -Z, Z = X
+	return FVector(-Translation.y, -Translation.z, Translation.x);
+}
+
+template <class TWrappedReader>
+TArrayView<const uint16> FDNAReader<TWrappedReader>::GetGUIToRawInputIndices() const
+{
+	const auto Indices = ReaderPtr->getGUIToRawInputIndices();
+	return TArrayView<const uint16>(Indices.data(), static_cast<int32>(Indices.size()));
+}
+
+template <class TWrappedReader>
+TArrayView<const uint16> FDNAReader<TWrappedReader>::GetGUIToRawOutputIndices() const
+{
+	const auto Indices = ReaderPtr->getGUIToRawOutputIndices();
+	return TArrayView<const uint16>(Indices.data(), static_cast<int32>(Indices.size()));
+}
+
+template <class TWrappedReader>
+TArrayView<const float> FDNAReader<TWrappedReader>::GetGUIToRawFromValues() const
+{
+	const auto Values = ReaderPtr->getGUIToRawFromValues();
+	return TArrayView<const float>(Values.data(), static_cast<int32>(Values.size()));
+}
+
+template <class TWrappedReader>
+TArrayView<const float> FDNAReader<TWrappedReader>::GetGUIToRawToValues() const
+{
+	const auto Values = ReaderPtr->getGUIToRawToValues();
+	return TArrayView<const float>(Values.data(), static_cast<int32>(Values.size()));
+}
+
+template <class TWrappedReader>
+TArrayView<const float> FDNAReader<TWrappedReader>::GetGUIToRawSlopeValues() const
+{
+	const auto Values = ReaderPtr->getGUIToRawSlopeValues();
+	return TArrayView<const float>(Values.data(), static_cast<int32>(Values.size()));
+}
+
+template <class TWrappedReader>
+TArrayView<const float> FDNAReader<TWrappedReader>::GetGUIToRawCutValues() const
+{
+	const auto Values = ReaderPtr->getGUIToRawCutValues();
+	return TArrayView<const float>(Values.data(), static_cast<int32>(Values.size()));
+}
+
+template <class TWrappedReader>
+uint16 FDNAReader<TWrappedReader>::GetPSDCount() const
+{
+	return ReaderPtr->getPSDCount();
+}
+
+template <class TWrappedReader>
+TArrayView<const uint16> FDNAReader<TWrappedReader>::GetPSDRowIndices() const
+{
+	const auto Indices = ReaderPtr->getPSDRowIndices();
+	return TArrayView<const uint16>(Indices.data(), static_cast<int32>(Indices.size()));
+}
+
+template <class TWrappedReader>
+TArrayView<const uint16> FDNAReader<TWrappedReader>::GetPSDColumnIndices() const
+{
+	const auto Indices = ReaderPtr->getPSDColumnIndices();
+	return TArrayView<const uint16>(Indices.data(), static_cast<int32>(Indices.size()));
+}
+
+template <class TWrappedReader>
+TArrayView<const float> FDNAReader<TWrappedReader>::GetPSDValues() const
+{
+	const auto Values = ReaderPtr->getPSDValues();
+	return TArrayView<const float>(Values.data(), static_cast<int32>(Values.size()));
+}
+
+template <class TWrappedReader>
+uint16 FDNAReader<TWrappedReader>::GetJointRowCount() const
+{
+	return ReaderPtr->getJointRowCount();
+}
+
+template <class TWrappedReader>
+uint16 FDNAReader<TWrappedReader>::GetJointColumnCount() const
+{
+	return ReaderPtr->getJointColumnCount();
+}
+
+template <class TWrappedReader>
+TArrayView<const uint16> FDNAReader<TWrappedReader>::GetJointVariableAttributeIndices(uint16 LOD) const
+{
+	const auto Indices = ReaderPtr->getJointVariableAttributeIndices(LOD);
+	return TArrayView<const uint16>(Indices.data(), static_cast<int32>(Indices.size()));
+}
+
+template <class TWrappedReader>
+uint16 FDNAReader<TWrappedReader>::GetJointGroupCount() const
+{
+	return ReaderPtr->getJointGroupCount();
+}
+
+template <class TWrappedReader>
+TArrayView<const uint16> FDNAReader<TWrappedReader>::GetJointGroupLODs(uint16 JointGroupIndex) const
+{
+	const auto LODs = ReaderPtr->getJointGroupLODs(JointGroupIndex);
+	return TArrayView<const uint16>(LODs.data(), static_cast<int32>(LODs.size()));
+}
+
+template <class TWrappedReader>
+TArrayView<const uint16> FDNAReader<TWrappedReader>::GetJointGroupInputIndices(uint16 JointGroupIndex) const
+{
+	const auto Indices = ReaderPtr->getJointGroupInputIndices(JointGroupIndex);
+	return TArrayView<const uint16>(Indices.data(), static_cast<int32>(Indices.size()));
+}
+
+template <class TWrappedReader>
+TArrayView<const uint16> FDNAReader<TWrappedReader>::GetJointGroupOutputIndices(uint16 JointGroupIndex) const
+{
+	const auto Indices = ReaderPtr->getJointGroupOutputIndices(JointGroupIndex);
+	return TArrayView<const uint16>(Indices.data(), static_cast<int32>(Indices.size()));
+}
+
+template <class TWrappedReader>
+TArrayView<const float> FDNAReader<TWrappedReader>::GetJointGroupValues(uint16 JointGroupIndex) const
+{
+	const auto Values = ReaderPtr->getJointGroupValues(JointGroupIndex);
+	return TArrayView<const float>(Values.data(), static_cast<int32>(Values.size()));
+}
+
+template <class TWrappedReader>
+TArrayView<const uint16> FDNAReader<TWrappedReader>::GetJointGroupJointIndices(uint16 JointGroupIndex) const
+{
+	const auto Indices = ReaderPtr->getJointGroupJointIndices(JointGroupIndex);
+	return TArrayView<const uint16>(Indices.data(), static_cast<int32>(Indices.size()));
+}
+
+template <class TWrappedReader>
+TArrayView<const uint16> FDNAReader<TWrappedReader>::GetBlendShapeChannelLODs() const
+{
+	const auto LODs = ReaderPtr->getBlendShapeChannelLODs();
+	return TArrayView<const uint16>(LODs.data(), static_cast<int32>(LODs.size()));
+}
+
+template <class TWrappedReader>
+TArrayView<const uint16> FDNAReader<TWrappedReader>::GetBlendShapeChannelOutputIndices() const
+{
+	const auto Indices = ReaderPtr->getBlendShapeChannelOutputIndices();
+	return TArrayView<const uint16>(Indices.data(), static_cast<int32>(Indices.size()));
+}
+
+template <class TWrappedReader>
+TArrayView<const uint16> FDNAReader<TWrappedReader>::GetBlendShapeChannelInputIndices() const
+{
+	const auto Indices = ReaderPtr->getBlendShapeChannelInputIndices();
+	return TArrayView<const uint16>(Indices.data(), static_cast<int32>(Indices.size()));
+}
+
+template <class TWrappedReader>
+TArrayView<const uint16> FDNAReader<TWrappedReader>::GetAnimatedMapLODs() const
+{
+	const auto LODs = ReaderPtr->getAnimatedMapLODs();
+	return TArrayView<const uint16>(LODs.data(), static_cast<int32>(LODs.size()));
+}
+
+template <class TWrappedReader>
+TArrayView<const uint16> FDNAReader<TWrappedReader>::GetAnimatedMapInputIndices() const
+{
+	const auto Indices = ReaderPtr->getAnimatedMapInputIndices();
+	return TArrayView<const uint16>(Indices.data(), static_cast<int32>(Indices.size()));
+}
+
+template <class TWrappedReader>
+TArrayView<const uint16> FDNAReader<TWrappedReader>::GetAnimatedMapOutputIndices() const
+{
+	const auto Indices = ReaderPtr->getAnimatedMapOutputIndices();
+	return TArrayView<const uint16>(Indices.data(), static_cast<int32>(Indices.size()));
+}
+
+template <class TWrappedReader>
+TArrayView<const float> FDNAReader<TWrappedReader>::GetAnimatedMapFromValues() const
+{
+	const auto Values = ReaderPtr->getAnimatedMapFromValues();
+	return TArrayView<const float>(Values.data(), static_cast<int32>(Values.size()));
+}
+
+template <class TWrappedReader>
+TArrayView<const float> FDNAReader<TWrappedReader>::GetAnimatedMapToValues() const
+{
+	const auto Values = ReaderPtr->getAnimatedMapToValues();
+	return TArrayView<const float>(Values.data(), static_cast<int32>(Values.size()));
+}
+
+template <class TWrappedReader>
+TArrayView<const float> FDNAReader<TWrappedReader>::GetAnimatedMapSlopeValues() const
+{
+	const auto Values = ReaderPtr->getAnimatedMapSlopeValues();
+	return TArrayView<const float>(Values.data(), static_cast<int32>(Values.size()));
+}
+
+template <class TWrappedReader>
+TArrayView<const float> FDNAReader<TWrappedReader>::GetAnimatedMapCutValues() const
+{
+	const auto Values = ReaderPtr->getAnimatedMapCutValues();
+	return TArrayView<const float>(Values.data(), static_cast<int32>(Values.size()));
+}
+
+
+template <class TWrappedReader>
+uint32 FDNAReader<TWrappedReader>::GetVertexPositionCount(uint16 MeshIndex) const
+{
+	return ReaderPtr->getVertexPositionCount(MeshIndex);
+}
+
+template <class TWrappedReader>
+FVector FDNAReader<TWrappedReader>::GetVertexPosition(uint16 MeshIndex, uint32 VertexIndex) const
+{
+	const auto Position = ReaderPtr->getVertexPosition(MeshIndex, VertexIndex);
+	// X = X, Y = Z, Z = Y
+	return FVector(Position.x, Position.z, Position.y);
+}
+
+template <class TWrappedReader>
+TArrayView<const float> FDNAReader<TWrappedReader>::GetVertexPositionXs(uint16 MeshIndex) const
+{
+	// X = X
+	const auto Values = ReaderPtr->getVertexPositionXs(MeshIndex);
+	return TArrayView<const float>(Values.data(), static_cast<int32>(Values.size()));
+}
+
+template <class TWrappedReader>
+TArrayView<const float> FDNAReader<TWrappedReader>::GetVertexPositionYs(uint16 MeshIndex) const
+{
+	// Y = Z
+	const auto Values = ReaderPtr->getVertexPositionZs(MeshIndex);
+	return TArrayView<const float>(Values.data(), static_cast<int32>(Values.size()));
+}
+
+template <class TWrappedReader>
+TArrayView<const float> FDNAReader<TWrappedReader>::GetVertexPositionZs(uint16 MeshIndex) const
+{
+	// Z = Y
+	const auto Values = ReaderPtr->getVertexPositionYs(MeshIndex);
+	return TArrayView<const float>(Values.data(), static_cast<int32>(Values.size()));
+}
+
+template <class TWrappedReader>
+uint32 FDNAReader<TWrappedReader>::GetVertexTextureCoordinateCount(uint16 MeshIndex) const
+{
+	return ReaderPtr->getVertexTextureCoordinateCount(MeshIndex);
+}
+
+template <class TWrappedReader>
+FTextureCoordinate FDNAReader<TWrappedReader>::GetVertexTextureCoordinate(uint16 MeshIndex, uint32 TextureCoordinateIndex) const
+{
+	const auto Coordinate = ReaderPtr->getVertexTextureCoordinate(MeshIndex, TextureCoordinateIndex);
+	return FTextureCoordinate{Coordinate.u, Coordinate.v};
+}
+
+template <class TWrappedReader>
+uint32 FDNAReader<TWrappedReader>::GetVertexNormalCount(uint16 MeshIndex) const
+{
+	return ReaderPtr->getVertexNormalCount(MeshIndex);
+}
+
+template <class TWrappedReader>
+FVector FDNAReader<TWrappedReader>::GetVertexNormal(uint16 MeshIndex, uint32 NormalIndex) const
+{
+	const auto Normal = ReaderPtr->getVertexNormal(MeshIndex, NormalIndex);
+	return FVector(Normal.x, Normal.y, Normal.z);
+}
+
+template <class TWrappedReader>
+uint32 FDNAReader<TWrappedReader>::GetFaceCount(uint16 MeshIndex) const
+{
+	return ReaderPtr->getFaceCount(MeshIndex);
+}
+
+template <class TWrappedReader>
+TArrayView<const uint32> FDNAReader<TWrappedReader>::GetFaceVertexLayoutIndices(uint16 MeshIndex, uint32 FaceIndex) const
+{
+	const auto Indices = ReaderPtr->getFaceVertexLayoutIndices(MeshIndex, FaceIndex);
+	return TArrayView<const uint32>(Indices.data(), static_cast<int32>(Indices.size()));
+}
+
+template <class TWrappedReader>
+uint32 FDNAReader<TWrappedReader>::GetVertexLayoutCount(uint16 MeshIndex) const
+{
+	return ReaderPtr->getVertexLayoutCount(MeshIndex);
+}
+
+template <class TWrappedReader>
+FVertexLayout FDNAReader<TWrappedReader>::GetVertexLayout(uint16 MeshIndex, uint32 LayoutIndex) const
+{
+	const auto Layout = ReaderPtr->getVertexLayout(MeshIndex, LayoutIndex);
+	return FVertexLayout{static_cast<int32>(Layout.position), static_cast<int32>(Layout.textureCoordinate), static_cast<int32>(Layout.normal)};
+}
+
+template <class TWrappedReader>
+uint16 FDNAReader<TWrappedReader>::GetMaximumInfluencePerVertex(uint16 MeshIndex) const
+{
+	return ReaderPtr->getMaximumInfluencePerVertex(MeshIndex);
+}
+
+template <class TWrappedReader>
+TArrayView<const float> FDNAReader<TWrappedReader>::GetSkinWeightsValues(uint16 MeshIndex, uint32 VertexIndex) const
+{
+	const auto Values = ReaderPtr->getSkinWeightsValues(MeshIndex, VertexIndex);
+	return TArrayView<const float>(Values.data(), static_cast<int32>(Values.size()));
+}
+
+template <class TWrappedReader>
+TArrayView<const uint16> FDNAReader<TWrappedReader>::GetSkinWeightsJointIndices(uint16 MeshIndex, uint32 VertexIndex) const
+{
+	const auto Indices = ReaderPtr->getSkinWeightsJointIndices(MeshIndex, VertexIndex);
+	return TArrayView<const uint16>(Indices.data(), static_cast<int32>(Indices.size()));
+}
+
+template <class TWrappedReader>
+uint16 FDNAReader<TWrappedReader>::GetBlendShapeTargetCount(uint16 MeshIndex) const
+{
+	return ReaderPtr->getBlendShapeTargetCount(MeshIndex);
+}
+
+template <class TWrappedReader>
+uint16 FDNAReader<TWrappedReader>::GetBlendShapeChannelIndex(uint16 MeshIndex, uint16 BlendShapeTargetIndex) const
+{
+	return ReaderPtr->getBlendShapeChannelIndex(MeshIndex, BlendShapeTargetIndex);
+}
+
+template <class TWrappedReader>
+uint32 FDNAReader<TWrappedReader>::GetBlendShapeTargetDeltaCount(uint16 MeshIndex, uint16 BlendShapeTargetIndex) const
+{
+	return ReaderPtr->getBlendShapeTargetDeltaCount(MeshIndex, BlendShapeTargetIndex);
+}
+
+template <class TWrappedReader>
+FVector FDNAReader<TWrappedReader>::GetBlendShapeTargetDelta(uint16 MeshIndex, uint16 BlendShapeTargetIndex, uint32 DeltaIndex) const
+{
+	const auto Delta = ReaderPtr->getBlendShapeTargetDelta(MeshIndex, BlendShapeTargetIndex, DeltaIndex);
+	// X = X, Y = Z, Z = Y
+	return FVector(Delta.x, Delta.z, Delta.y);
+}
+
+template <class TWrappedReader>
+TArrayView<const float> FDNAReader<TWrappedReader>::GetBlendShapeTargetDeltaXs(uint16 MeshIndex, uint16 BlendShapeTargetIndex) const
+{
+	// X = X
+	const auto Values = ReaderPtr->getBlendShapeTargetDeltaXs(MeshIndex, BlendShapeTargetIndex);
+	return TArrayView<const float>(Values.data(), static_cast<int32>(Values.size()));
+}
+
+template <class TWrappedReader>
+TArrayView<const float> FDNAReader<TWrappedReader>::GetBlendShapeTargetDeltaYs(uint16 MeshIndex, uint16 BlendShapeTargetIndex) const
+{
+	// Y = Z
+	const auto Values = ReaderPtr->getBlendShapeTargetDeltaZs(MeshIndex, BlendShapeTargetIndex);
+	return TArrayView<const float>(Values.data(), static_cast<int32>(Values.size()));
+}
+
+template <class TWrappedReader>
+TArrayView<const float> FDNAReader<TWrappedReader>::GetBlendShapeTargetDeltaZs(uint16 MeshIndex, uint16 BlendShapeTargetIndex) const
+{
+	// Z = Y
+	const auto Values = ReaderPtr->getBlendShapeTargetDeltaYs(MeshIndex, BlendShapeTargetIndex);
+	return TArrayView<const float>(Values.data(), static_cast<int32>(Values.size()));
+}
+
+template <class TWrappedReader>
+TArrayView<const uint32> FDNAReader<TWrappedReader>::GetBlendShapeTargetVertexIndices(uint16 MeshIndex, uint16 BlendShapeTargetIndex) const
+{
+	const auto Indices = ReaderPtr->getBlendShapeTargetVertexIndices(MeshIndex, BlendShapeTargetIndex);
+	return TArrayView<const uint32>(Indices.data(), static_cast<int32>(Indices.size()));
+}
diff --git a/Engine/Plugins/Animation/RigLogic/Source/RigLogicModule/Public/DNAReaderAdapterImpl.h b/Engine/Plugins/Animation/RigLogic/Source/RigLogicModule/Public/DNAReaderAdapterImpl.h
deleted file mode 100644
