Use const skeletal mesh when possible

diff --git a/Engine/Plugins/Animation/RigLogic/Source/RigLogicModule/Private/DNAToSkelMeshMap.cpp b/Engine/Plugins/Animation/RigLogic/Source/RigLogicModule/Private/DNAToSkelMeshMap.cpp
--- a/Engine/Plugins/Animation/RigLogic/Source/RigLogicModule/Private/DNAToSkelMeshMap.cpp
+++ b/Engine/Plugins/Animation/RigLogic/Source/RigLogicModule/Private/DNAToSkelMeshMap.cpp
@@ -187,7 +187,7 @@ void FDNAToSkelMeshMap::MapJoints(IDNAReader* DNAReader)
 
 	double StartTime = FPlatformTime::Seconds();
 
-	FReferenceSkeleton* RefSkeleton = &TargetSkelMesh->GetRefSkeleton();
+	const FReferenceSkeleton& RefSkeleton = TargetSkelMesh->GetRefSkeleton();
 	uint32 JointCount = DNAReader->GetJointCount();
 
 	// Map Joints to Bones.
@@ -197,7 +197,7 @@ void FDNAToSkelMeshMap::MapJoints(IDNAReader* DNAReader)
 	{
 		const FString boneNameFStr = DNAReader->GetJointName(JntIndex);
 		const FName BoneName(*boneNameFStr, 0, FNAME_Add, false);
-		const int32 BoneIndex = RefSkeleton->FindBoneIndex(BoneName);
+		const int32 BoneIndex = RefSkeleton.FindBoneIndex(BoneName);
 
 		// BoneIndex can be INDEX_NONE;
 		// We can safely put it into the map with other indices, it will be handled in the Evaluate method.
diff --git a/Engine/Plugins/Experimental/ControlRig/Source/ControlRig/Private/ControlRigComponent.cpp b/Engine/Plugins/Experimental/ControlRig/Source/ControlRig/Private/ControlRigComponent.cpp
--- a/Engine/Plugins/Experimental/ControlRig/Source/ControlRig/Private/ControlRigComponent.cpp
+++ b/Engine/Plugins/Experimental/ControlRig/Source/ControlRig/Private/ControlRigComponent.cpp
@@ -388,9 +388,9 @@ void UControlRigComponent::AddMappedSkeletalMesh(USkeletalMeshComponent* Skeleta
 	{
 		if (UControlRig* CR = SetupControlRigIfRequired())
 		{
-			if (USkeletalMesh* SkeletalMesh = SkeletalMeshComponent->SkeletalMesh)
+			if (const USkeletalMesh* SkeletalMesh = SkeletalMeshComponent->SkeletalMesh)
 			{
-				if (USkeleton* Skeleton = SkeletalMesh->GetSkeleton())
+				if (const USkeleton* Skeleton = SkeletalMesh->GetSkeleton())
 				{
 					for (const FRigBone& RigBone : CR->GetBoneHierarchy())
 					{
diff --git a/Engine/Source/Runtime/ClothingSystemRuntimeCommon/Private/ClothingAsset.cpp b/Engine/Source/Runtime/ClothingSystemRuntimeCommon/Private/ClothingAsset.cpp
--- a/Engine/Source/Runtime/ClothingSystemRuntimeCommon/Private/ClothingAsset.cpp
+++ b/Engine/Source/Runtime/ClothingSystemRuntimeCommon/Private/ClothingAsset.cpp
@@ -663,7 +663,7 @@ void UClothingAssetCommon::CalculateReferenceBoneIndex()
 	// List of valid paths to the root bone from each weighted bone
 	TArray<BoneIndexArray> PathsToRoot;
 	
-	USkeletalMesh* OwnerMesh = Cast<USkeletalMesh>(GetOuter());
+	const USkeletalMesh* OwnerMesh = Cast<USkeletalMesh>(GetOuter());
 
 	if(OwnerMesh)
 	{
