Fixing PVS 7.7 Issues:

warning V1028: Possible overflow. Consider casting operands, not the result.

diff --git a/Engine/Plugins/Animation/RigLogic/Source/RigLogicLib/Private/dna/LODConstraint.cpp b/Engine/Plugins/Animation/RigLogic/Source/RigLogicLib/Private/dna/LODConstraint.cpp
--- a/Engine/Plugins/Animation/RigLogic/Source/RigLogicLib/Private/dna/LODConstraint.cpp
+++ b/Engine/Plugins/Animation/RigLogic/Source/RigLogicLib/Private/dna/LODConstraint.cpp
@@ -19,7 +19,7 @@ namespace dna {
 
 LODConstraint::LODConstraint(std::uint16_t maxLOD, std::uint16_t minLOD, MemoryResource* memRes) : lods{memRes} {
     assert(maxLOD <= minLOD);
-    lods.resize(static_cast<std::size_t>(minLOD - maxLOD + 1));
+    lods.resize(static_cast<std::size_t>(minLOD - maxLOD) + 1);
     std::iota(lods.begin(), lods.end(), maxLOD);
 }
 
diff --git a/Engine/Plugins/Experimental/GeometryProcessing/Source/DynamicMesh/Public/DynamicVertexAttribute.h b/Engine/Plugins/Experimental/GeometryProcessing/Source/DynamicMesh/Public/DynamicVertexAttribute.h
--- a/Engine/Plugins/Experimental/GeometryProcessing/Source/DynamicMesh/Public/DynamicVertexAttribute.h
+++ b/Engine/Plugins/Experimental/GeometryProcessing/Source/DynamicMesh/Public/DynamicVertexAttribute.h
@@ -292,7 +292,7 @@ public:
 		{
 			return;
 		}
-		size_t NeededSize = size_t((VertexID+1) * AttribDimension);
+		size_t NeededSize = (size_t(VertexID)+1) * AttribDimension;
 		if (NeededSize > AttribValues.Num())
 		{
 			AttribValues.Resize(NeededSize, GetDefaultAttributeValue());
diff --git a/Engine/Source/Editor/VirtualTexturingEditor/Private/RuntimeVirtualTextureBuildStreamingMips.cpp b/Engine/Source/Editor/VirtualTexturingEditor/Private/RuntimeVirtualTextureBuildStreamingMips.cpp
--- a/Engine/Source/Editor/VirtualTexturingEditor/Private/RuntimeVirtualTextureBuildStreamingMips.cpp
+++ b/Engine/Source/Editor/VirtualTexturingEditor/Private/RuntimeVirtualTextureBuildStreamingMips.cpp
@@ -92,7 +92,7 @@ namespace
 		for (int32 y = 0; y < TileSize; y++)
 		{
 			memcpy(
-				DestPixels + (SIZE_T)DestStride * (SIZE_T)(DestPos[1] + y) + DestPos[0],
+				DestPixels + (SIZE_T)DestStride * (SIZE_T)DestPos[1] + y + DestPos[0],
 				SrcPixels + SrcStride * y,
 				TileSize * sizeof(T));
 		}
diff --git a/Engine/Source/Runtime/PakFile/Private/IPlatformFilePak.cpp b/Engine/Source/Runtime/PakFile/Private/IPlatformFilePak.cpp
--- a/Engine/Source/Runtime/PakFile/Private/IPlatformFilePak.cpp
+++ b/Engine/Source/Runtime/PakFile/Private/IPlatformFilePak.cpp
@@ -4281,7 +4281,7 @@ bool FPakProcessedReadRequest::CheckCompletion(const FPakEntry& FileEntry, int32
 	}
 	{
 		int64 BlockStart = int64(BlockIndex) * int64(FileEntry.CompressionBlockSize);
-		int64 BlockEnd = int64(BlockIndex + 1) * int64(FileEntry.CompressionBlockSize);
+		int64 BlockEnd = (int64(BlockIndex) + 1) * int64(FileEntry.CompressionBlockSize);
 		if (Offset >= BlockEnd || Offset + BytesToRead <= BlockStart)
 		{
 			return false;
