RigLogic: Fixed unity-related compile errors

#rnx

diff --git a/Engine/Plugins/Animation/RigLogic/Source/RigLogicLib/Private/riglogic/utils/Extd.h b/Engine/Plugins/Animation/RigLogic/Source/RigLogicLib/Private/riglogic/utils/Extd.h
--- a/Engine/Plugins/Animation/RigLogic/Source/RigLogicLib/Private/riglogic/utils/Extd.h
+++ b/Engine/Plugins/Animation/RigLogic/Source/RigLogicLib/Private/riglogic/utils/Extd.h
@@ -59,6 +59,17 @@ inline void filter(std::vector<T, Args...>& source, Predicate pred) {
     }), source.end());
 }
 
+template<class TContainer, class Predicate, typename ... Args>
+inline void filter(TContainer& source, Predicate pred) {
+    using value_type = typename TContainer::value_type;
+    auto newEnd = std::remove_if(std::begin(source), std::end(source), [&source, &pred](const value_type& value) {
+        const auto index = static_cast<std::size_t>(&value - source.data());
+        return !pred(value, index);
+    });
+    const auto newSize = static_cast<std::size_t>(std::distance(source.begin(), newEnd));
+    source.resize(newSize);
+}
+
 namespace impl {
 
 enum class LUTStrategy {
diff --git a/Engine/Plugins/Animation/RigLogic/Source/RigLogicLibTest/Private/rltests/blendshapes/BlendShapeFixtures.cpp b/Engine/Plugins/Animation/RigLogic/Source/RigLogicLibTest/Private/rltests/blendshapes/BlendShapeFixtures.cpp
--- a/Engine/Plugins/Animation/RigLogic/Source/RigLogicLibTest/Private/rltests/blendshapes/BlendShapeFixtures.cpp
+++ b/Engine/Plugins/Animation/RigLogic/Source/RigLogicLibTest/Private/rltests/blendshapes/BlendShapeFixtures.cpp
@@ -5,7 +5,7 @@
 #include "riglogic/TypeDefs.h"
 #include "riglogic/blendshapes/BlendShapes.h"
 
-rl4::BlendShapes createTestBlendShapes(rl4::MemoryResource* memRes) {
+rl4::BlendShapes BlendShapeFixturesTest::createTestBlendShapes(rl4::MemoryResource* memRes) {
     return rl4::BlendShapes{
         rl4::Vector<std::uint16_t>{std::begin(LODs), std::end(LODs), memRes},
         rl4::Vector<std::uint16_t>{std::begin(blendShapeInputIndices), std::end(blendShapeInputIndices), memRes},
diff --git a/Engine/Plugins/Animation/RigLogic/Source/RigLogicLibTest/Private/rltests/blendshapes/BlendShapeFixtures.h b/Engine/Plugins/Animation/RigLogic/Source/RigLogicLibTest/Private/rltests/blendshapes/BlendShapeFixtures.h
--- a/Engine/Plugins/Animation/RigLogic/Source/RigLogicLibTest/Private/rltests/blendshapes/BlendShapeFixtures.h
+++ b/Engine/Plugins/Animation/RigLogic/Source/RigLogicLibTest/Private/rltests/blendshapes/BlendShapeFixtures.h
@@ -17,6 +17,8 @@
     #pragma warning(pop)
 #endif
 
+namespace BlendShapeFixturesTest {
+
 static const std::size_t LODCount = 4ul;
 static const std::uint16_t LOD0 = 10u;
 static const std::uint16_t LOD1 = 7u;
@@ -33,3 +35,5 @@ static const std::uint16_t blendShapeOutputIndices[] = {0, 2, 1, 3, 4, 5, 6, 7,
 static const std::array<float, 10ul> blendShapeOutputs = {1.0f, 4.0f, 2.0f, 5.0f, 7.0f, 8.0f, 9.0f, 11.0f, 13.0f, 15.0f};
 
 rl4::BlendShapes createTestBlendShapes(rl4::MemoryResource* memRes);
+
+}
diff --git a/Engine/Plugins/Animation/RigLogic/Source/RigLogicLibTest/Private/rltests/blendshapes/TestBlendShapes.cpp b/Engine/Plugins/Animation/RigLogic/Source/RigLogicLibTest/Private/rltests/blendshapes/TestBlendShapes.cpp
--- a/Engine/Plugins/Animation/RigLogic/Source/RigLogicLibTest/Private/rltests/blendshapes/TestBlendShapes.cpp
+++ b/Engine/Plugins/Animation/RigLogic/Source/RigLogicLibTest/Private/rltests/blendshapes/TestBlendShapes.cpp
@@ -26,6 +26,8 @@ class BlendShapesTest : public ::testing::TestWithParam<std::uint16_t> {
 }  // namespace
 
 TEST_P(BlendShapesTest, Calculate) {
+    using namespace BlendShapeFixturesTest;
+
     // Test input parameter - Which LOD level to test
     const auto lod = GetParam();
 
