Delete obsolete files and fix lost changes in build scripts due to merge #rb damir.jugovic

diff --git a/Engine/Plugins/Animation/RigLogic/Source/RigLogicEditor/RigLogicEditor.Build.cs b/Engine/Plugins/Animation/RigLogic/Source/RigLogicEditor/RigLogicEditor.Build.cs
--- a/Engine/Plugins/Animation/RigLogic/Source/RigLogicEditor/RigLogicEditor.Build.cs
+++ b/Engine/Plugins/Animation/RigLogic/Source/RigLogicEditor/RigLogicEditor.Build.cs
@@ -26,6 +26,7 @@ namespace UnrealBuildTool.Rules
                     "Engine",
                     "ControlRig",
                     "UnrealEd",
+                    "EditorFramework",
                     "MainFrame",
                     "RigLogicModule",
                     "RigLogicLib",
diff --git a/Engine/Plugins/Animation/RigLogic/Source/RigLogicLib/RigLogicLib.Build.cs b/Engine/Plugins/Animation/RigLogic/Source/RigLogicLib/RigLogicLib.Build.cs
--- a/Engine/Plugins/Animation/RigLogic/Source/RigLogicLib/RigLogicLib.Build.cs
+++ b/Engine/Plugins/Animation/RigLogic/Source/RigLogicLib/RigLogicLib.Build.cs
@@ -26,6 +26,7 @@ public class RigLogicLib : ModuleRules
         if (Target.Type == TargetType.Editor)
         {
             PublicDependencyModuleNames.Add("UnrealEd");
+            PublicDependencyModuleNames.Add("EditorFramework");
         }
 
         Type = ModuleType.CPlusPlus;
diff --git a/Engine/Plugins/Animation/RigLogic/Source/RigLogicModule/Private/DNAToSkelMeshMap.cpp b/Engine/Plugins/Animation/RigLogic/Source/RigLogicModule/Private/DNAToSkelMeshMap.cpp
deleted file mode 100644
diff --git a/Engine/Plugins/Animation/RigLogic/Source/RigLogicModule/Private/RigUnit_RigLogic.cpp b/Engine/Plugins/Animation/RigLogic/Source/RigLogicModule/Private/RigUnit_RigLogic.cpp
--- a/Engine/Plugins/Animation/RigLogic/Source/RigLogicModule/Private/RigUnit_RigLogic.cpp
+++ b/Engine/Plugins/Animation/RigLogic/Source/RigLogicModule/Private/RigUnit_RigLogic.cpp
@@ -6,7 +6,6 @@
 #include "DNAReader.h"
 #include "RigInstance.h"
 #include "RigLogic.h"
-#include "SkelMeshDNAUtils.h"
 #include "Components/SkeletalMeshComponent.h"
 #include "Math/TransformNonVectorized.h"
 #include "Units/RigUnitContext.h"
diff --git a/Engine/Plugins/Animation/RigLogic/Source/RigLogicModule/Private/SkelMeshDNAReader.cpp b/Engine/Plugins/Animation/RigLogic/Source/RigLogicModule/Private/SkelMeshDNAReader.cpp
deleted file mode 100644
diff --git a/Engine/Plugins/Animation/RigLogic/Source/RigLogicModule/Private/SkelMeshDNAReader.h b/Engine/Plugins/Animation/RigLogic/Source/RigLogicModule/Private/SkelMeshDNAReader.h
deleted file mode 100644
diff --git a/Engine/Plugins/Animation/RigLogic/Source/RigLogicModule/Private/SkelMeshDNAUtils.cpp b/Engine/Plugins/Animation/RigLogic/Source/RigLogicModule/Private/SkelMeshDNAUtils.cpp
deleted file mode 100644
diff --git a/Engine/Plugins/Animation/RigLogic/Source/RigLogicModule/Public/DNAToSkelMeshMap.h b/Engine/Plugins/Animation/RigLogic/Source/RigLogicModule/Public/DNAToSkelMeshMap.h
deleted file mode 100644
diff --git a/Engine/Plugins/Animation/RigLogic/Source/RigLogicModule/Public/SkelMeshDNAUtils.h b/Engine/Plugins/Animation/RigLogic/Source/RigLogicModule/Public/SkelMeshDNAUtils.h
deleted file mode 100644
diff --git a/Engine/Plugins/Animation/RigLogic/Source/RigLogicModule/RigLogicModule.Build.cs b/Engine/Plugins/Animation/RigLogic/Source/RigLogicModule/RigLogicModule.Build.cs
--- a/Engine/Plugins/Animation/RigLogic/Source/RigLogicModule/RigLogicModule.Build.cs
+++ b/Engine/Plugins/Animation/RigLogic/Source/RigLogicModule/RigLogicModule.Build.cs
@@ -29,19 +29,18 @@ namespace UnrealBuildTool.Rules
                     "ControlRig",
                     "RigLogicLib",
                     "RigVM",
-                    "Projects",
-					"RenderCore",
-					"RHI"
+                    "Projects"
                 }
             );
 
             if (Target.Type == TargetType.Editor)
             {
                 PublicDependencyModuleNames.Add("UnrealEd");
-				PublicDependencyModuleNames.Add("MessageLog");
-			}
+                PublicDependencyModuleNames.Add("EditorFramework");
+                PublicDependencyModuleNames.Add("MessageLog");
+            }
 
-			PrivateDependencyModuleNames.AddRange(
+            PrivateDependencyModuleNames.AddRange(
                 new string[]
                 {
                     "AnimationCore"
