Fix deprecation warning

diff --git a/Engine/Plugins/Animation/RigLogic/Source/RigLogicModule/Private/RigUnit_RigLogic.cpp b/Engine/Plugins/Animation/RigLogic/Source/RigLogicModule/Private/RigUnit_RigLogic.cpp
--- a/Engine/Plugins/Animation/RigLogic/Source/RigLogicModule/Private/RigUnit_RigLogic.cpp
+++ b/Engine/Plugins/Animation/RigLogic/Source/RigLogicModule/Private/RigUnit_RigLogic.cpp
@@ -359,7 +359,7 @@ FRigUnit_RigLogic_Execute()
 				{
 					return;
 				}
-				Data.CurrentLOD = Data.SkelMeshComponent->PredictedLODLevel;
+				Data.CurrentLOD = Data.SkelMeshComponent->GetPredictedLODLevel();
 				Data.CalculateRigLogic(Hierarchy);
 
 				//Filing a struct so we can call the same method for updating joints from tests
