Delete obsolete files and fix UE4 references #jira UE-111051 UE-111055 UE-111077 #rb trivial

diff --git a/Engine/Plugins/Animation/RigLogic/Source/RigLogicModule/Public/DNAReader.h b/Engine/Plugins/Animation/RigLogic/Source/RigLogicModule/Public/DNAReader.h
--- a/Engine/Plugins/Animation/RigLogic/Source/RigLogicModule/Public/DNAReader.h
+++ b/Engine/Plugins/Animation/RigLogic/Source/RigLogicModule/Public/DNAReader.h
@@ -149,7 +149,7 @@ public:
 };
 
 /**
-	@brief UE4 interface for DNA Reader wrappers.
+	@brief UE interface for DNA Reader wrappers.
 */
 class IDNAReader : public IBehaviorReader, public IGeometryReader
 {
