Merge CL 15752516 from //UE5/Release-Engine-Staging to Main (//UE5/Main) to fix jira #UE-112099

diff --git a/Engine/Plugins/Animation/RigLogic/Source/RigLogicModule/Private/RigUnit_RigLogic.cpp b/Engine/Plugins/Animation/RigLogic/Source/RigLogicModule/Private/RigUnit_RigLogic.cpp
--- a/Engine/Plugins/Animation/RigLogic/Source/RigLogicModule/Private/RigUnit_RigLogic.cpp
+++ b/Engine/Plugins/Animation/RigLogic/Source/RigLogicModule/Private/RigUnit_RigLogic.cpp
@@ -79,9 +79,9 @@ bool FRigUnit_RigLogic_Data::IsRigLogicInitialized()
 
 void FRigUnit_RigLogic_Data::InitializeRigLogic(const URigHierarchy* InHierarchy)
 {
-	if (!SharedRigRuntimeContext->BehaviorReader.IsValid() || SharedRigRuntimeContext->BehaviorReader->GetJointCount() == 0u)
+	if ((SharedRigRuntimeContext == nullptr) || !SharedRigRuntimeContext->BehaviorReader.IsValid() || SharedRigRuntimeContext->BehaviorReader->GetJointCount() == 0u)
 	{
-		UE_LOG(LogRigLogicUnit, Warning, TEXT("Empty DNA file detected, abort initialization."));
+		UE_LOG(LogRigLogicUnit, Warning, TEXT("No valid DNA file found, abort initialization."));
 		return;
 	}
 
