Fix warnings about uninitialized members in RigLogic #jira UE-113320

diff --git a/Engine/Plugins/Animation/RigLogic/Source/RigLogicModule/Public/DNACommon.h b/Engine/Plugins/Animation/RigLogic/Source/RigLogicModule/Public/DNACommon.h
--- a/Engine/Plugins/Animation/RigLogic/Source/RigLogicModule/Public/DNACommon.h
+++ b/Engine/Plugins/Animation/RigLogic/Source/RigLogicModule/Public/DNACommon.h
@@ -69,6 +69,14 @@ struct FCoordinateSystem
 {
 	GENERATED_BODY()
 
+	FCoordinateSystem() : XAxis(), YAxis(), ZAxis()
+	{
+	}
+
+	FCoordinateSystem(EDirection XAxis, EDirection YAxis, EDirection ZAxis) : XAxis(XAxis), YAxis(YAxis), ZAxis(ZAxis)
+	{
+	}
+
 	UPROPERTY(BlueprintReadOnly, Category = "RigLogic")
 	EDirection XAxis;
 	UPROPERTY(BlueprintReadOnly, Category = "RigLogic")
@@ -82,6 +90,14 @@ struct FMeshBlendShapeChannelMapping
 {
 	GENERATED_BODY()
 
+	FMeshBlendShapeChannelMapping() : MeshIndex(), BlendShapeChannelIndex()
+	{
+	}
+
+	FMeshBlendShapeChannelMapping(int32 MeshIndex, int32 BlendShapeChannelIndex) : MeshIndex(MeshIndex), BlendShapeChannelIndex(BlendShapeChannelIndex)
+	{
+	}
+
 	UPROPERTY(BlueprintReadOnly, Category = "RigLogic")
 	int32 MeshIndex;
 	UPROPERTY(BlueprintReadOnly, Category = "RigLogic")
@@ -93,6 +109,14 @@ struct FTextureCoordinate
 {
 	GENERATED_BODY()
 
+	FTextureCoordinate() : U(), V()
+	{
+	}
+
+	FTextureCoordinate(float U, float V) : U(U), V(V)
+	{
+	}
+
 	UPROPERTY(BlueprintReadOnly, Category = "RigLogic")
 	float U;
 	UPROPERTY(BlueprintReadOnly, Category = "RigLogic")
@@ -104,6 +128,14 @@ struct FVertexLayout
 {
 	GENERATED_BODY()
 
+	FVertexLayout() : Position(), TextureCoordinate(), Normal()
+	{
+	}
+
+	FVertexLayout(int32 Position, int32 TextureCoordinate, int32 Normal) : Position(Position), TextureCoordinate(TextureCoordinate), Normal(Normal)
+	{
+	}
+
 	UPROPERTY(BlueprintReadOnly, Category = "RigLogic")
 	int32 Position;
 	UPROPERTY(BlueprintReadOnly, Category = "RigLogic")
