Added caching of RigUnit_RigLogicData mappings, which would allow for multiple rigs with same rig definition and topology version to share mappings.

diff --git a/Engine/Plugins/Animation/RigLogic/Source/RigLogicModule/Private/RigUnit_RigLogic.cpp b/Engine/Plugins/Animation/RigLogic/Source/RigLogicModule/Private/RigUnit_RigLogic.cpp
--- a/Engine/Plugins/Animation/RigLogic/Source/RigLogicModule/Private/RigUnit_RigLogic.cpp
+++ b/Engine/Plugins/Animation/RigLogic/Source/RigLogicModule/Private/RigUnit_RigLogic.cpp
@@ -39,6 +39,14 @@ static FString ConstructCurveName(const FString& NameToSplit, const FString& For
 	return CurveName;
 }
 
+FRigDefinitionMapping::FRigDefinitionMapping(const IBehaviorReader* DNABehavior, const URigHierarchy* InHierarchy) : FRigDefinitionMapping()
+{
+	MapInputCurveIndices(DNABehavior, InHierarchy);
+	MapJoints(DNABehavior, InHierarchy);
+	MapMaskMultipliers(DNABehavior, InHierarchy);
+	MapMorphTargets(DNABehavior, InHierarchy);
+}
+
 FRigUnit_RigLogic_Data::FRigUnit_RigLogic_Data()
 : SkelMeshComponent(nullptr)
 , SharedRigRuntimeContext(nullptr)
@@ -62,12 +70,7 @@ FRigUnit_RigLogic_Data& FRigUnit_RigLogic_Data::operator=(const FRigUnit_RigLogi
 	SkelMeshComponent = Other.SkelMeshComponent;
 	SharedRigRuntimeContext = nullptr;
 	RigInstance = nullptr;
-	InputCurveIndices = Other.InputCurveIndices;
-	HierarchyBoneIndices = Other.HierarchyBoneIndices;
-	MorphTargetCurveIndices = Other.MorphTargetCurveIndices;
-	BlendShapeIndices = Other.BlendShapeIndices;
-	CurveElementIndicesForAnimMaps = Other.CurveElementIndicesForAnimMaps;
-	RigLogicIndicesForAnimMaps = Other.RigLogicIndicesForAnimMaps;
+	Mappings = Other.Mappings;
 	CurrentLOD = Other.CurrentLOD;
 	return *this;
 }
@@ -77,6 +80,8 @@ bool FRigUnit_RigLogic_Data::IsRigLogicInitialized()
 	return (SharedRigRuntimeContext != nullptr) && SharedRigRuntimeContext->RigLogic.IsValid() && RigInstance.IsValid();
 }
 
+TDiscardableKeyValueCache< TTuple<FString, uint16>, TSharedPtr<const FRigDefinitionMapping>> FRigUnit_RigLogic_Data::RigDefintionMappingsCache{};
+
 void FRigUnit_RigLogic_Data::InitializeRigLogic(const URigHierarchy* InHierarchy)
 {
 	if ((SharedRigRuntimeContext == nullptr) || !SharedRigRuntimeContext->BehaviorReader.IsValid() || SharedRigRuntimeContext->BehaviorReader->GetJointCount() == 0u)
@@ -97,17 +102,26 @@ void FRigUnit_RigLogic_Data::InitializeRigLogic(const URigHierarchy* InHierarchy
 		RigInstance = MakeUnique<FRigInstance>(SharedRigRuntimeContext->RigLogic.Get());
 		RigInstance->SetLOD(CurrentLOD);
 
-		MapJoints(InHierarchy);
-		MapInputCurveIndices(InHierarchy);
-		MapMorphTargets(InHierarchy);
-		MapMaskMultipliers(InHierarchy);
+		const IBehaviorReader* DNABehavior = SharedRigRuntimeContext->BehaviorReader.Get();
+		auto DBName = DNABehavior->GetDBName();
+		uint16 TopologyVersion = InHierarchy->GetTopologyVersion();
+		TTuple<FString, uint16> CacheKey{ DBName, TopologyVersion };
+		auto& Cache = FRigUnit_RigLogic_Data::RigDefintionMappingsCache;
+
+		using FLockFLags = TDiscardableKeyValueCache< FString, TSharedPtr<FRigDefinitionMapping>>::LockFlags;
+		uint32 LockFlags = Cache.ApplyLock(0, FLockFLags::ReadLock);
+		if (!Cache.Find(CacheKey, Mappings, LockFlags | FLockFLags::WriteLockOnAddFail, LockFlags)) 
+		{
+			Mappings = MakeShared<const FRigDefinitionMapping>(DNABehavior, InHierarchy);
+			Cache.Add(CacheKey, Mappings, LockFlags);
+		}
+		Cache.Unlock(LockFlags);
 	}
 }
 
 //maps indices of input curves from dna file to control rig curves
-void FRigUnit_RigLogic_Data::MapInputCurveIndices(const URigHierarchy* InHierarchy)
+void FRigDefinitionMapping::MapInputCurveIndices(const IBehaviorReader* DNABehavior, const URigHierarchy* InHierarchy)
 {
-	const IBehaviorReader* DNABehavior = SharedRigRuntimeContext->BehaviorReader.Get();
 	const uint32 ControlCount = DNABehavior->GetRawControlCount();
 	InputCurveIndices.Reset(ControlCount);
 	for (uint32_t ControlIndex = 0; ControlIndex < ControlCount; ++ControlIndex)
@@ -124,16 +138,15 @@ void FRigUnit_RigLogic_Data::MapInputCurveIndices(const URigHierarchy* InHierarc
 	}
 }
 
-void FRigUnit_RigLogic_Data::MapJoints(const URigHierarchy* Hierarchy)
+void FRigDefinitionMapping::MapJoints(const IBehaviorReader* DNABehavior, const URigHierarchy* InHierarchy)
 {
-	const IBehaviorReader* DNABehavior = SharedRigRuntimeContext->BehaviorReader.Get();
 	const uint16 JointCount = DNABehavior->GetJointCount();
 	HierarchyBoneIndices.Reset(JointCount);
 	for (uint16 JointIndex = 0; JointIndex < JointCount ; ++JointIndex)
 	{
 		const FString RLJointName = DNABehavior->GetJointName(JointIndex);
 		const FName JointFName = FName(*RLJointName);
-		const int32 BoneIndex = Hierarchy->GetIndex(FRigElementKey(JointFName, ERigElementType::Bone));
+		const int32 BoneIndex = InHierarchy->GetIndex(FRigElementKey(JointFName, ERigElementType::Bone));
 		HierarchyBoneIndices.Add(BoneIndex);
 	}
 }
@@ -159,9 +172,8 @@ void FRigUnit_RigLogic_Data::CacheVariableJointIndices()
 	}
 }
 
-void FRigUnit_RigLogic_Data::MapMorphTargets(const URigHierarchy* InHierarchy)
+void FRigDefinitionMapping::MapMorphTargets(const IBehaviorReader* DNABehavior, const URigHierarchy* InHierarchy)
 {
-	const IBehaviorReader* DNABehavior = SharedRigRuntimeContext->BehaviorReader.Get();
 	const uint16 LODCount = DNABehavior->GetLODCount();
 
 	MorphTargetCurveIndices.Reset();
@@ -191,9 +203,8 @@ void FRigUnit_RigLogic_Data::MapMorphTargets(const URigHierarchy* InHierarchy)
 	}
 }
 
-void FRigUnit_RigLogic_Data::MapMaskMultipliers(const URigHierarchy* InHierarchy)
+void FRigDefinitionMapping::MapMaskMultipliers(const IBehaviorReader* DNABehavior, const URigHierarchy* InHierarchy)
 {
-	const IBehaviorReader* DNABehavior = SharedRigRuntimeContext->BehaviorReader.Get();
 	const uint16 LODCount = DNABehavior->GetLODCount();
 	CurveElementIndicesForAnimMaps.Reset();
 	CurveElementIndicesForAnimMaps.AddDefaulted(LODCount);
@@ -231,7 +242,7 @@ void FRigUnit_RigLogic_Data::CalculateRigLogic(const URigHierarchy* InHierarchy)
 	const int32 RawControlCount = RigInstance->GetRawControlCount();
 	for (int32 ControlIndex = 0; ControlIndex < RawControlCount; ++ControlIndex)
 	{
-		const uint32 CurveIndex = InputCurveIndices[ControlIndex];
+		const uint32 CurveIndex = Mappings->InputCurveIndices[ControlIndex];
 		const float Value = InHierarchy->GetCurveValue(CurveIndex);
 		RigInstance->SetRawControl(ControlIndex, Value);
 	}
@@ -242,7 +253,7 @@ void FRigUnit_RigLogic_Data::UpdateJoints(URigHierarchy* Hierarchy, const FRigUn
 {
 	for (const uint16 JointIndex : SharedRigRuntimeContext->VariableJointIndices[CurrentLOD].Values)
 	{
-		const int32 BoneIndex = HierarchyBoneIndices[JointIndex];
+		const int32 BoneIndex = Mappings->HierarchyBoneIndices[JointIndex];
 		if (BoneIndex != INDEX_NONE)
 		{
 			const FTransform& Neutral = JointUpdateParams.NeutralJointTransforms[JointIndex];
@@ -261,16 +272,16 @@ void FRigUnit_RigLogic_Data::UpdateJoints(URigHierarchy* Hierarchy, const FRigUn
 void FRigUnit_RigLogic_Data::UpdateBlendShapeCurves(URigHierarchy* InHierarchy, TArrayView<const float> BlendShapeValues)
 {
 	// set output blend shapes
-	if (BlendShapeIndices.IsValidIndex(CurrentLOD) && MorphTargetCurveIndices.IsValidIndex(CurrentLOD))
+	if (Mappings->BlendShapeIndices.IsValidIndex(CurrentLOD) && Mappings->MorphTargetCurveIndices.IsValidIndex(CurrentLOD))
 	{
-		const uint32 BlendShapePerLODCount = static_cast<uint32>(BlendShapeIndices[CurrentLOD].Values.Num());
+		const uint32 BlendShapePerLODCount = static_cast<uint32>(Mappings->BlendShapeIndices[CurrentLOD].Values.Num());
 
-		if (ensure(BlendShapePerLODCount == MorphTargetCurveIndices[CurrentLOD].Values.Num()))
+		if (ensure(BlendShapePerLODCount == Mappings->MorphTargetCurveIndices[CurrentLOD].Values.Num()))
 		{
 			for (uint32 MeshBlendIndex = 0; MeshBlendIndex < BlendShapePerLODCount; MeshBlendIndex++)
 			{
-				const int32 BlendShapeIndex = BlendShapeIndices[CurrentLOD].Values[MeshBlendIndex];
-				const int32 MorphTargetCurveIndex = MorphTargetCurveIndices[CurrentLOD].Values[MeshBlendIndex];
+				const int32 BlendShapeIndex = Mappings->BlendShapeIndices[CurrentLOD].Values[MeshBlendIndex];
+				const int32 MorphTargetCurveIndex = Mappings->MorphTargetCurveIndices[CurrentLOD].Values[MeshBlendIndex];
 				if (MorphTargetCurveIndex != INDEX_NONE)
 				{
 					const float Value = BlendShapeValues[BlendShapeIndex];
@@ -290,13 +301,13 @@ void FRigUnit_RigLogic_Data::UpdateAnimMapCurves(URigHierarchy* InHierarchy, TAr
 	// set output mask multipliers
 	// In case curves are not imported yet into CL, AnimatedMapsCurveIndices will be empty, so we need to check
 	// array bounds before trying to access it:
-	if (RigLogicIndicesForAnimMaps.IsValidIndex(CurrentLOD) && CurveElementIndicesForAnimMaps.IsValidIndex(CurrentLOD))
+	if (Mappings->RigLogicIndicesForAnimMaps.IsValidIndex(CurrentLOD) && Mappings->CurveElementIndicesForAnimMaps.IsValidIndex(CurrentLOD))
 	{
-		const uint32 AnimMapPerLODCount = RigLogicIndicesForAnimMaps[CurrentLOD].Values.Num();
+		const uint32 AnimMapPerLODCount = Mappings->RigLogicIndicesForAnimMaps[CurrentLOD].Values.Num();
 		for (uint32 AnimMapIndexForLOD = 0; AnimMapIndexForLOD < AnimMapPerLODCount; ++AnimMapIndexForLOD)
 		{
-			const int32 RigLogicAnimMapIndex = RigLogicIndicesForAnimMaps[CurrentLOD].Values[AnimMapIndexForLOD];
-			const int32 InHierarchyAnimMapIndex = CurveElementIndicesForAnimMaps[CurrentLOD].Values[AnimMapIndexForLOD];
+			const int32 RigLogicAnimMapIndex = Mappings->RigLogicIndicesForAnimMaps[CurrentLOD].Values[AnimMapIndexForLOD];
+			const int32 InHierarchyAnimMapIndex = Mappings->CurveElementIndicesForAnimMaps[CurrentLOD].Values[AnimMapIndexForLOD];
 			if (InHierarchyAnimMapIndex != INDEX_NONE)
 			{
 				const float Value = AnimMapOutputs[RigLogicAnimMapIndex];
diff --git a/Engine/Plugins/Animation/RigLogic/Source/RigLogicModule/Private/RigUnit_RigLogic_Test.cpp b/Engine/Plugins/Animation/RigLogic/Source/RigLogicModule/Private/RigUnit_RigLogic_Test.cpp
--- a/Engine/Plugins/Animation/RigLogic/Source/RigLogicModule/Private/RigUnit_RigLogic_Test.cpp
+++ b/Engine/Plugins/Animation/RigLogic/Source/RigLogicModule/Private/RigUnit_RigLogic_Test.cpp
@@ -56,7 +56,7 @@ TStrongObjectPtr<URigHierarchy> FRigUnit_RigLogic::TestAccessor::CreateCurveCont
 
 void FRigUnit_RigLogic::TestAccessor::Exec_MapInputCurve(URigHierarchy* TestHierarchy)
 {
-	Unit->Data.MapInputCurveIndices(TestHierarchy); //inside the method so we can access data, which is a private member
+	Unit->Data.InitializeRigLogic(TestHierarchy); //inside the method so we can access data, which is a private member
 }
 
 /** ====== Map Joints ===== **/
@@ -88,7 +88,7 @@ TSharedPtr<TestBehaviorReader> FRigUnit_RigLogic::TestAccessor::CreateBehaviorRe
 
 void FRigUnit_RigLogic::TestAccessor::Exec_MapJoints(URigHierarchy* TestHierarchy)
 {
-	Unit->Data.MapJoints(TestHierarchy);
+	Unit->Data.InitializeRigLogic(TestHierarchy); //inside the method so we can access data, which is a private member
 }
 
 /** ====== Map Morph Targets ===== **/
@@ -149,7 +149,7 @@ TStrongObjectPtr<URigHierarchy> FRigUnit_RigLogic::TestAccessor::CreateCurveCont
 void FRigUnit_RigLogic::TestAccessor::Exec_MapMorphTargets(URigHierarchy* TestCurveContainer)
 {
 	//put into a separate method so we can access the private Data member
-	Unit->Data.MapMorphTargets(TestCurveContainer); 
+	Unit->Data.InitializeRigLogic(TestCurveContainer);
 }
 
 /** ====== Map Mask Multipliers ===== **/
@@ -165,7 +165,7 @@ TSharedPtr<TestBehaviorReader> FRigUnit_RigLogic::TestAccessor::CreateBehaviorRe
 
 void FRigUnit_RigLogic::TestAccessor::Exec_MapMaskMultipliers(URigHierarchy* TestHierarchy)
 {
-	Unit->Data.MapMaskMultipliers(TestHierarchy); //inside the method so we can access data, which is a private member
+	Unit->Data.InitializeRigLogic(TestHierarchy); //inside the method so we can access data, which is a private member
 }
 
 void FRigUnit_RigLogic::TestAccessor::AddToTransformArray(float* InArray, FTransform& Transform)
@@ -288,18 +288,18 @@ IMPLEMENT_RIGUNIT_AUTOMATION_TEST(FRigUnit_RigLogic)
 		Test.GetData()->RigInstance != nullptr &&
 
 		// Check joints
-		Test.GetData()->HierarchyBoneIndices.Num() > 0 &&
+		Test.GetData()->Mappings->HierarchyBoneIndices.Num() > 0 &&
 
 		// Check input curves
-		Test.GetData()->InputCurveIndices.Num() > 0 &&
+		Test.GetData()->Mappings->InputCurveIndices.Num() > 0 &&
 
 		// Check morph targets
-		Test.GetData()->MorphTargetCurveIndices.Num() > 0 &&
-		Test.GetData()->BlendShapeIndices.Num() > 0 &&
+		Test.GetData()->Mappings->MorphTargetCurveIndices.Num() > 0 &&
+		Test.GetData()->Mappings->BlendShapeIndices.Num() > 0 &&
 
 		// Check mask multipliers
-		Test.GetData()->RigLogicIndicesForAnimMaps.Num() > 0 &&
-		Test.GetData()->CurveElementIndicesForAnimMaps.Num() > 0,
+		Test.GetData()->Mappings->RigLogicIndicesForAnimMaps.Num() > 0 &&
+		Test.GetData()->Mappings->CurveElementIndicesForAnimMaps.Num() > 0,
 
 		TEXT("InitAndExecute failed to initialize rig logic.")
 	);
@@ -316,8 +316,8 @@ IMPLEMENT_RIGUNIT_AUTOMATION_TEST(FRigUnit_RigLogic)
 	Test.Exec_MapInputCurve(TestCurveContainerNameMismatch.Get());
 	//Assert
 	AddErrorIfFalse(
-		Test.GetData()->InputCurveIndices.Num() == 1 &&
-		Test.GetData()->InputCurveIndices[0] == INDEX_NONE,
+		Test.GetData()->Mappings->InputCurveIndices.Num() == 1 &&
+		Test.GetData()->Mappings->InputCurveIndices[0] == INDEX_NONE,
 		TEXT("MapInputCurve - ValidReader CurveContainerWithNameMismatch")
 	);
 
@@ -331,7 +331,7 @@ IMPLEMENT_RIGUNIT_AUTOMATION_TEST(FRigUnit_RigLogic)
 	Test.Exec_MapInputCurve(TestCurveContainerValid.Get());
 	//Assert
 	AddErrorIfFalse(
-		Test.GetData()->InputCurveIndices.Num() == 0,
+		Test.GetData()->Mappings->InputCurveIndices.Num() == 0,
 		TEXT("MapInputCurve - EmptyReader ValidCurveContainer")
 	);
 
@@ -344,8 +344,8 @@ IMPLEMENT_RIGUNIT_AUTOMATION_TEST(FRigUnit_RigLogic)
 	Test.Exec_MapInputCurve(TestCurveContainerEmpty.Get());
 	//Assert
 	AddErrorIfFalse(
-		Test.GetData()->InputCurveIndices.Num() == 1 &&
-		Test.GetData()->InputCurveIndices[0] == INDEX_NONE,
+		Test.GetData()->Mappings->InputCurveIndices.Num() == 1 &&
+		Test.GetData()->Mappings->InputCurveIndices[0] == INDEX_NONE,
 		TEXT("MapInputCurve - ValidReader EmptyCurveContainer")
 	);
 
@@ -359,7 +359,7 @@ IMPLEMENT_RIGUNIT_AUTOMATION_TEST(FRigUnit_RigLogic)
 	Test.Exec_MapInputCurve(TestCurveContainerValid.Get());
 	//Assert
 	AddErrorIfFalse(
-		Test.GetData()->InputCurveIndices.Num() == 0,
+		Test.GetData()->Mappings->InputCurveIndices.Num() == 0,
 		TEXT("MapInputCurve - InvalidReader ValidCurveContainer")
 	);
 
@@ -372,8 +372,8 @@ IMPLEMENT_RIGUNIT_AUTOMATION_TEST(FRigUnit_RigLogic)
 	Test.Exec_MapInputCurve(TestCurveContainerValid.Get());
 	//Assert
 	AddErrorIfFalse(
-		Test.GetData()->InputCurveIndices.Num() == 1 &&
-		Test.GetData()->InputCurveIndices[0] == 0,
+		Test.GetData()->Mappings->InputCurveIndices.Num() == 1 &&
+		Test.GetData()->Mappings->InputCurveIndices[0] == 0,
 		TEXT("MapInputCurve - Valid Inputs")
 	);
 
@@ -388,7 +388,7 @@ IMPLEMENT_RIGUNIT_AUTOMATION_TEST(FRigUnit_RigLogic)
 	Test.Exec_MapJoints(TestHierarchyEmpty.Get());
 	//Assert
 	AddErrorIfFalse(
-		Test.GetData()->HierarchyBoneIndices.Num() == 0,
+		Test.GetData()->Mappings->HierarchyBoneIndices.Num() == 0,
 		TEXT("MapJoints - Empty Inputs")
 	);
 
@@ -399,7 +399,7 @@ IMPLEMENT_RIGUNIT_AUTOMATION_TEST(FRigUnit_RigLogic)
 	//Test
 	Test.Exec_MapJoints(TestHierarchyTwoBones.Get());
 	AddErrorIfFalse(
-		Test.GetData()->HierarchyBoneIndices.Num() == 0,
+		Test.GetData()->Mappings->HierarchyBoneIndices.Num() == 0,
 		TEXT("MapJoints - EmptyReader TwoBones")
 	);
 
@@ -410,17 +410,17 @@ IMPLEMENT_RIGUNIT_AUTOMATION_TEST(FRigUnit_RigLogic)
 	//Test
 	Test.Exec_MapJoints(TestHierarchyEmpty.Get());
 	AddErrorIfFalse(
-		Test.GetData()->HierarchyBoneIndices.Num() == 2,
+		Test.GetData()->Mappings->HierarchyBoneIndices.Num() == 2,
 		TEXT("MapJoints - TwoJoints NoBones - expected 2 bone indices")
 	);
 	AddErrorIfFalse(
-		Test.GetData()->HierarchyBoneIndices.Num() == 2 && //repeated condition to prevent crash
-		Test.GetData()->HierarchyBoneIndices[0] == INDEX_NONE,
+		Test.GetData()->Mappings->HierarchyBoneIndices.Num() == 2 && //repeated condition to prevent crash
+		Test.GetData()->Mappings->HierarchyBoneIndices[0] == INDEX_NONE,
 		TEXT("MapJoints - TwoJoints NoBones - Expected joint 0 index to be NONE")
 	);
 	AddErrorIfFalse(
-		Test.GetData()->HierarchyBoneIndices.Num() == 2 && //repeated condition to prevent crash
-		Test.GetData()->HierarchyBoneIndices[1] == INDEX_NONE,
+		Test.GetData()->Mappings->HierarchyBoneIndices.Num() == 2 && //repeated condition to prevent crash
+		Test.GetData()->Mappings->HierarchyBoneIndices[1] == INDEX_NONE,
 		TEXT("MapJoints - TwoJoints NoBones - Expected joint 1 index to be NONE")
 	);
 
@@ -432,17 +432,17 @@ IMPLEMENT_RIGUNIT_AUTOMATION_TEST(FRigUnit_RigLogic)
 	Test.Exec_MapJoints(TestHierarchyTwoBones.Get());
 	//Assert
 	AddErrorIfFalse(
-		Test.GetData()->HierarchyBoneIndices.Num() == 2,
+		Test.GetData()->Mappings->HierarchyBoneIndices.Num() == 2,
 		TEXT("MapJoints - TwoJoints TwoBones - Expected 2 bone indices")
 	);
 	AddErrorIfFalse(
-		Test.GetData()->HierarchyBoneIndices.Num() == 2 && //prevent crash
-		Test.GetData()->HierarchyBoneIndices[0] == 0,
+		Test.GetData()->Mappings->HierarchyBoneIndices.Num() == 2 && //prevent crash
+		Test.GetData()->Mappings->HierarchyBoneIndices[0] == 0,
 		TEXT("MapJoints - TwoJoints TwoBones - Expected bone 0 index to be 0")
 	);
 	AddErrorIfFalse(
-		Test.GetData()->HierarchyBoneIndices.Num() == 2 && //prevent crash
-		Test.GetData()->HierarchyBoneIndices[1] == 1,
+		Test.GetData()->Mappings->HierarchyBoneIndices.Num() == 2 && //prevent crash
+		Test.GetData()->Mappings->HierarchyBoneIndices[1] == 1,
 		TEXT("MapJoints - TwoJoints TwoBones - Expected bone index 1 index to be 1")
 	);
 
@@ -459,12 +459,12 @@ IMPLEMENT_RIGUNIT_AUTOMATION_TEST(FRigUnit_RigLogic)
 	Test.Exec_MapMorphTargets(TestMorphTargetNameMismatch.Get());
 	//Assert
 	AddErrorIfFalse(
-		Test.GetData()->BlendShapeIndices.Num() == 1 && //LOD 0
-		Test.GetData()->MorphTargetCurveIndices.Num() == 1 && //LOD 0
-		Test.GetData()->BlendShapeIndices[0].Values.Num() == 1 && //at least one blendshape exists
-		Test.GetData()->BlendShapeIndices[0].Values[0] == 0 && //has index 0
-		Test.GetData()->MorphTargetCurveIndices[0].Values.Num() == 1 && //but morph target corresponding to that blendshape
-		Test.GetData()->MorphTargetCurveIndices[0].Values[0] == INDEX_NONE, //wasn't found
+		Test.GetData()->Mappings->BlendShapeIndices.Num() == 1 && //LOD 0
+		Test.GetData()->Mappings->MorphTargetCurveIndices.Num() == 1 && //LOD 0
+		Test.GetData()->Mappings->BlendShapeIndices[0].Values.Num() == 1 && //at least one blendshape exists
+		Test.GetData()->Mappings->BlendShapeIndices[0].Values[0] == 0 && //has index 0
+		Test.GetData()->Mappings->MorphTargetCurveIndices[0].Values.Num() == 1 && //but morph target corresponding to that blendshape
+		Test.GetData()->Mappings->MorphTargetCurveIndices[0].Values[0] == INDEX_NONE, //wasn't found
 		TEXT("MapMorphTargets - ValidReader MorphTargetWithNameMismatch")
 	);
 
@@ -478,8 +478,8 @@ IMPLEMENT_RIGUNIT_AUTOMATION_TEST(FRigUnit_RigLogic)
 	Test.Exec_MapMorphTargets(TestMorphTargetCurveValid.Get());
 	//Assert
 	AddErrorIfFalse(
-		Test.GetData()->MorphTargetCurveIndices.Num() == 0 &&
-		Test.GetData()->BlendShapeIndices.Num() == 0,
+		Test.GetData()->Mappings->MorphTargetCurveIndices.Num() == 0 &&
+		Test.GetData()->Mappings->BlendShapeIndices.Num() == 0,
 		TEXT("MapMorphTargets - EmptyReader ValidMorphTargetCurve")
 	);
 
@@ -493,9 +493,9 @@ IMPLEMENT_RIGUNIT_AUTOMATION_TEST(FRigUnit_RigLogic)
 	Test.Exec_MapMorphTargets(TestMorphTargetCurveValid.Get());
 	//Assert
 	AddErrorIfFalse(
-		Test.GetData()->BlendShapeIndices.Num() == 1 && //LOD 0 exists
-		Test.GetData()->BlendShapeIndices[0].Values.Num() == 0 && //but no blend shapes mapped
-		Test.GetData()->MorphTargetCurveIndices[0].Values.Num() == 0, //or morph targets
+		Test.GetData()->Mappings->BlendShapeIndices.Num() == 1 && //LOD 0 exists
+		Test.GetData()->Mappings->BlendShapeIndices[0].Values.Num() == 0 && //but no blend shapes mapped
+		Test.GetData()->Mappings->MorphTargetCurveIndices[0].Values.Num() == 0, //or morph targets
 		TEXT("MapMorphTargets - NoBlendShapes ValidMorphTargetCurve")
 	);
 
@@ -508,11 +508,11 @@ IMPLEMENT_RIGUNIT_AUTOMATION_TEST(FRigUnit_RigLogic)
 	Test.Exec_MapMorphTargets(TestCurveContainerEmpty.Get());
 	//Assert
 	AddErrorIfFalse(
-		Test.GetData()->BlendShapeIndices.Num() == 1 && //LOD 0
-		Test.GetData()->BlendShapeIndices[0].Values.Num() == 1 && //one blend shape
-		Test.GetData()->BlendShapeIndices[0].Values[0] == 0 && //of index 0
-		Test.GetData()->MorphTargetCurveIndices[0].Values.Num() == 1 && //we put in a morph target index corresponding to it
-		Test.GetData()->MorphTargetCurveIndices[0].Values[0] == INDEX_NONE, //but just to signal that it is not found
+		Test.GetData()->Mappings->BlendShapeIndices.Num() == 1 && //LOD 0
+		Test.GetData()->Mappings->BlendShapeIndices[0].Values.Num() == 1 && //one blend shape
+		Test.GetData()->Mappings->BlendShapeIndices[0].Values[0] == 0 && //of index 0
+		Test.GetData()->Mappings->MorphTargetCurveIndices[0].Values.Num() == 1 && //we put in a morph target index corresponding to it
+		Test.GetData()->Mappings->MorphTargetCurveIndices[0].Values[0] == INDEX_NONE, //but just to signal that it is not found
 		TEXT("MapMorphTargets - ValidReader EmptyCurveContainer")
 	);
 
@@ -525,11 +525,11 @@ IMPLEMENT_RIGUNIT_AUTOMATION_TEST(FRigUnit_RigLogic)
 	Test.Exec_MapMorphTargets(TestMorphTargetCurveValid.Get());
 	//Assert
 	AddErrorIfFalse(
-		Test.GetData()->BlendShapeIndices.Num() == 1 && //LOD 0
-		Test.GetData()->BlendShapeIndices[0].Values.Num() == 1 && //one blend shape (empty named)
-		Test.GetData()->BlendShapeIndices[0].Values[0] == 0 && //of index 0
-		Test.GetData()->MorphTargetCurveIndices[0].Values.Num() == 1 && //we put in a morph target index corresponding to it
-		Test.GetData()->MorphTargetCurveIndices[0].Values[0] == INDEX_NONE, //but just to signal that it is not found
+		Test.GetData()->Mappings->BlendShapeIndices.Num() == 1 && //LOD 0
+		Test.GetData()->Mappings->BlendShapeIndices[0].Values.Num() == 1 && //one blend shape (empty named)
+		Test.GetData()->Mappings->BlendShapeIndices[0].Values[0] == 0 && //of index 0
+		Test.GetData()->Mappings->MorphTargetCurveIndices[0].Values.Num() == 1 && //we put in a morph target index corresponding to it
+		Test.GetData()->Mappings->MorphTargetCurveIndices[0].Values[0] == INDEX_NONE, //but just to signal that it is not found
 		TEXT("MapMorphTargets - InvalidReader ValidMorphTargetCurve")
 	);
 
@@ -542,11 +542,11 @@ IMPLEMENT_RIGUNIT_AUTOMATION_TEST(FRigUnit_RigLogic)
 	Test.Exec_MapMorphTargets(TestMorphTargetCurvesInvalid.Get());
 	//Assert
 	AddErrorIfFalse(
-		Test.GetData()->BlendShapeIndices.Num() == 1 && //LOD 0
-		Test.GetData()->BlendShapeIndices[0].Values.Num() == 1 && //one blend shape
-		Test.GetData()->BlendShapeIndices[0].Values[0] == 0 && //of index 0
-		Test.GetData()->MorphTargetCurveIndices[0].Values.Num() == 1 && //we put in a morph target index corresponding to it
-		Test.GetData()->MorphTargetCurveIndices[0].Values[0] == INDEX_NONE, //but just to signal that it is not found
+		Test.GetData()->Mappings->BlendShapeIndices.Num() == 1 && //LOD 0
+		Test.GetData()->Mappings->BlendShapeIndices[0].Values.Num() == 1 && //one blend shape
+		Test.GetData()->Mappings->BlendShapeIndices[0].Values[0] == 0 && //of index 0
+		Test.GetData()->Mappings->MorphTargetCurveIndices[0].Values.Num() == 1 && //we put in a morph target index corresponding to it
+		Test.GetData()->Mappings->MorphTargetCurveIndices[0].Values[0] == INDEX_NONE, //but just to signal that it is not found
 		TEXT("MapMorphTargets - ValidReader InvalidMorphTargetCurve")
 	);
 
@@ -559,12 +559,12 @@ IMPLEMENT_RIGUNIT_AUTOMATION_TEST(FRigUnit_RigLogic)
 	Test.Exec_MapMorphTargets(TestMorphTargetCurveValid.Get());
 	//Assert
 	AddErrorIfFalse(
-		Test.GetData()->BlendShapeIndices.Num() == 1 && //LOD 0
-		Test.GetData()->MorphTargetCurveIndices.Num() == 1 && //LOD 0
-		Test.GetData()->BlendShapeIndices[0].Values.Num() == 1 && //at least one blendshape exists
-		Test.GetData()->BlendShapeIndices[0].Values[0] == 0 && //has index 0
-		Test.GetData()->MorphTargetCurveIndices[0].Values.Num() == 1 && //morph target corresponding to that blendshape exists
-		Test.GetData()->MorphTargetCurveIndices[0].Values[0] == 0, //and actually points to the right index
+		Test.GetData()->Mappings->BlendShapeIndices.Num() == 1 && //LOD 0
+		Test.GetData()->Mappings->MorphTargetCurveIndices.Num() == 1 && //LOD 0
+		Test.GetData()->Mappings->BlendShapeIndices[0].Values.Num() == 1 && //at least one blendshape exists
+		Test.GetData()->Mappings->BlendShapeIndices[0].Values[0] == 0 && //has index 0
+		Test.GetData()->Mappings->MorphTargetCurveIndices[0].Values.Num() == 1 && //morph target corresponding to that blendshape exists
+		Test.GetData()->Mappings->MorphTargetCurveIndices[0].Values[0] == 0, //and actually points to the right index
 		TEXT("MapMorphTargets - ValidReader ValidTestMorphTarget")
 	);
 
@@ -587,28 +587,28 @@ IMPLEMENT_RIGUNIT_AUTOMATION_TEST(FRigUnit_RigLogic)
 	Test.Exec_MapMorphTargets(TestMorphTargetTwoCurves.Get());
 	//Assert
 	AddErrorIfFalse(
-		Test.GetData()->BlendShapeIndices.Num() == 2 &&  //2 LODs
-		Test.GetData()->MorphTargetCurveIndices.Num() == 2,
+		Test.GetData()->Mappings->BlendShapeIndices.Num() == 2 &&  //2 LODs
+		Test.GetData()->Mappings->MorphTargetCurveIndices.Num() == 2,
 		TEXT("MapMorphTargets LOD0(AB) LOD1(A) - Expected 2 LODs for both blendshapes and morph targets")
 	);
 
 	AddErrorIfFalse(
-		Test.GetData()->BlendShapeIndices.Num() == 2 &&         //condition repeated for crash prevention
-		Test.GetData()->BlendShapeIndices[0].Values.Num() == 2 && //two blendshapes at LOD 0
-		Test.GetData()->BlendShapeIndices[0].Values[0] == 0 &&  // A
-		Test.GetData()->BlendShapeIndices[0].Values[1] == 1 &&  // B
-		Test.GetData()->BlendShapeIndices[1].Values.Num() == 1 && //one blendshape at LOD 1
-		Test.GetData()->BlendShapeIndices[1].Values[0] == 0,    // A
+		Test.GetData()->Mappings->BlendShapeIndices.Num() == 2 &&         //condition repeated for crash prevention
+		Test.GetData()->Mappings->BlendShapeIndices[0].Values.Num() == 2 && //two blendshapes at LOD 0
+		Test.GetData()->Mappings->BlendShapeIndices[0].Values[0] == 0 &&  // A
+		Test.GetData()->Mappings->BlendShapeIndices[0].Values[1] == 1 &&  // B
+		Test.GetData()->Mappings->BlendShapeIndices[1].Values.Num() == 1 && //one blendshape at LOD 1
+		Test.GetData()->Mappings->BlendShapeIndices[1].Values[0] == 0,    // A
 			TEXT("MapMorphTargets LOD0(AB) LOD1(A) - resulting blendshape indices not correct")
 	);
 
 	AddErrorIfFalse(
-		Test.GetData()->MorphTargetCurveIndices.Num() == 2 &&         //condition repeated for crash prevention
-		Test.GetData()->MorphTargetCurveIndices[0].Values.Num() == 2 && //two morph targets at LOD 0
-		Test.GetData()->MorphTargetCurveIndices[0].Values[0] == 0 &&  // A
-		Test.GetData()->MorphTargetCurveIndices[0].Values[1] == 1 &&  // B
-		Test.GetData()->MorphTargetCurveIndices[1].Values.Num() == 1 && //one morph target at LOD 1
-		Test.GetData()->MorphTargetCurveIndices[1].Values[0] == 0,    // A
+		Test.GetData()->Mappings->MorphTargetCurveIndices.Num() == 2 &&         //condition repeated for crash prevention
+		Test.GetData()->Mappings->MorphTargetCurveIndices[0].Values.Num() == 2 && //two morph targets at LOD 0
+		Test.GetData()->Mappings->MorphTargetCurveIndices[0].Values[0] == 0 &&  // A
+		Test.GetData()->Mappings->MorphTargetCurveIndices[0].Values[1] == 1 &&  // B
+		Test.GetData()->Mappings->MorphTargetCurveIndices[1].Values.Num() == 1 && //one morph target at LOD 1
+		Test.GetData()->Mappings->MorphTargetCurveIndices[1].Values[0] == 0,    // A
 		TEXT("MapMorphTargets LOD0(AB) LOD1(A) - resulting morph target indices not correct")
 	);
 
@@ -625,25 +625,25 @@ IMPLEMENT_RIGUNIT_AUTOMATION_TEST(FRigUnit_RigLogic)
 	Test.Exec_MapMorphTargets(TestMorphTargetTwoCurves.Get());
 	//Assert
 	AddErrorIfFalse(
-		Test.GetData()->BlendShapeIndices.Num() == 1 &&  //1 LOD
-		Test.GetData()->MorphTargetCurveIndices.Num() == 1, //1 LOD
+		Test.GetData()->Mappings->BlendShapeIndices.Num() == 1 &&  //1 LOD
+		Test.GetData()->Mappings->MorphTargetCurveIndices.Num() == 1, //1 LOD
 		TEXT("MapMorphTargets LOD0(AB) LOD1(-) - Expected 1 LOD for both blendshapes and morph targets")
 	);
 
 
 	AddErrorIfFalse(
-		Test.GetData()->BlendShapeIndices.Num() == 1 &&         //condition repeated for crash prevention
-		Test.GetData()->BlendShapeIndices[0].Values.Num() == 2 && //two blendshapes at LOD 0
-		Test.GetData()->BlendShapeIndices[0].Values[0] == 0 &&  // A
-		Test.GetData()->BlendShapeIndices[0].Values[1] == 1,    // B
+		Test.GetData()->Mappings->BlendShapeIndices.Num() == 1 &&         //condition repeated for crash prevention
+		Test.GetData()->Mappings->BlendShapeIndices[0].Values.Num() == 2 && //two blendshapes at LOD 0
+		Test.GetData()->Mappings->BlendShapeIndices[0].Values[0] == 0 &&  // A
+		Test.GetData()->Mappings->BlendShapeIndices[0].Values[1] == 1,    // B
 		TEXT("MapMorphTargets LOD0(AB) LOD1(-) - Resulting blendshapes not correct")
 	);
 
 	AddErrorIfFalse(
-		Test.GetData()->MorphTargetCurveIndices.Num() == 1 &&         //condition repeated for crash prevention
-		Test.GetData()->MorphTargetCurveIndices[0].Values.Num() == 2 && //two morph targets at LOD 0
-		Test.GetData()->MorphTargetCurveIndices[0].Values[0] == 0 &&  // A
-		Test.GetData()->MorphTargetCurveIndices[0].Values[1] == 1,    // B
+		Test.GetData()->Mappings->MorphTargetCurveIndices.Num() == 1 &&         //condition repeated for crash prevention
+		Test.GetData()->Mappings->MorphTargetCurveIndices[0].Values.Num() == 2 && //two morph targets at LOD 0
+		Test.GetData()->Mappings->MorphTargetCurveIndices[0].Values[0] == 0 &&  // A
+		Test.GetData()->Mappings->MorphTargetCurveIndices[0].Values[1] == 1,    // B
 		TEXT("MapMorphTargets LOD0(AB) LOD1(-) - Resulting morph targets not correct")
 	);
 
@@ -659,26 +659,26 @@ IMPLEMENT_RIGUNIT_AUTOMATION_TEST(FRigUnit_RigLogic)
 	Test.Exec_MapMorphTargets(TestMorphTargetTwoCurves.Get());
 	//Assert
 	AddErrorIfFalse(
-		Test.GetData()->BlendShapeIndices.Num() == 2 &&  //2 LODs
-		Test.GetData()->MorphTargetCurveIndices.Num() == 2, //2 LODs
+		Test.GetData()->Mappings->BlendShapeIndices.Num() == 2 &&  //2 LODs
+		Test.GetData()->Mappings->MorphTargetCurveIndices.Num() == 2, //2 LODs
 		TEXT("MapMorphTargets LOD0(A) LOD1(B) - Expected 2 LODs for both blendshapes and morph targets")
 	);
 
 	AddErrorIfFalse(
-		Test.GetData()->BlendShapeIndices.Num() == 2 &&         //condition repeated for crash prevention
-		Test.GetData()->BlendShapeIndices[0].Values.Num() == 1 && //1 blendshape at LOD 0
-		Test.GetData()->BlendShapeIndices[0].Values[0] == 0 &&  // A
-		Test.GetData()->BlendShapeIndices[1].Values.Num() == 1 && //1 blendshape at LOD 1
-		Test.GetData()->BlendShapeIndices[1].Values[0] == 1,    // B
+		Test.GetData()->Mappings->BlendShapeIndices.Num() == 2 &&         //condition repeated for crash prevention
+		Test.GetData()->Mappings->BlendShapeIndices[0].Values.Num() == 1 && //1 blendshape at LOD 0
+		Test.GetData()->Mappings->BlendShapeIndices[0].Values[0] == 0 &&  // A
+		Test.GetData()->Mappings->BlendShapeIndices[1].Values.Num() == 1 && //1 blendshape at LOD 1
+		Test.GetData()->Mappings->BlendShapeIndices[1].Values[0] == 1,    // B
 		TEXT("MapMorphTargets LOD0(A) LOD1(B) - Resulting blendshape indices not correct")
 	);
 
 	AddErrorIfFalse(
-		Test.GetData()->MorphTargetCurveIndices.Num() == 2 &&         //condition repeated for crash prevention
-		Test.GetData()->MorphTargetCurveIndices[0].Values.Num() == 1 && //1 morph target at LOD 0
-		Test.GetData()->MorphTargetCurveIndices[0].Values[0] == 0 &&  // A
-		Test.GetData()->MorphTargetCurveIndices[1].Values.Num() == 1 && //1 morph target at LOD 1
-		Test.GetData()->MorphTargetCurveIndices[1].Values[0] == 1,    // B
+		Test.GetData()->Mappings->MorphTargetCurveIndices.Num() == 2 &&         //condition repeated for crash prevention
+		Test.GetData()->Mappings->MorphTargetCurveIndices[0].Values.Num() == 1 && //1 morph target at LOD 0
+		Test.GetData()->Mappings->MorphTargetCurveIndices[0].Values[0] == 0 &&  // A
+		Test.GetData()->Mappings->MorphTargetCurveIndices[1].Values.Num() == 1 && //1 morph target at LOD 1
+		Test.GetData()->Mappings->MorphTargetCurveIndices[1].Values[0] == 1,    // B
 		TEXT("MapMorphTargets LOD0(A) LOD1(B) - Resulting morph target indices not correct")
 	);
 
@@ -695,12 +695,12 @@ IMPLEMENT_RIGUNIT_AUTOMATION_TEST(FRigUnit_RigLogic)
 	Test.Exec_MapMaskMultipliers(TestCurveContainerForAnimMapsNameMismatch.Get());
 	//Assert
 	AddErrorIfFalse(
-		Test.GetData()->CurveElementIndicesForAnimMaps.Num() == 1 &&
-		Test.GetData()->CurveElementIndicesForAnimMaps[0].Values.Num() == 1 &&
-		Test.GetData()->CurveElementIndicesForAnimMaps[0].Values[0] == INDEX_NONE &&
-		Test.GetData()->RigLogicIndicesForAnimMaps.Num() == 1 &&
-		Test.GetData()->RigLogicIndicesForAnimMaps[0].Values.Num() == 1 &&
-		Test.GetData()->RigLogicIndicesForAnimMaps[0].Values[0] == 0,
+		Test.GetData()->Mappings->CurveElementIndicesForAnimMaps.Num() == 1 &&
+		Test.GetData()->Mappings->CurveElementIndicesForAnimMaps[0].Values.Num() == 1 &&
+		Test.GetData()->Mappings->CurveElementIndicesForAnimMaps[0].Values[0] == INDEX_NONE &&
+		Test.GetData()->Mappings->RigLogicIndicesForAnimMaps.Num() == 1 &&
+		Test.GetData()->Mappings->RigLogicIndicesForAnimMaps[0].Values.Num() == 1 &&
+		Test.GetData()->Mappings->RigLogicIndicesForAnimMaps[0].Values[0] == 0,
 		TEXT("MapMaskMultipliers - ValidReader ValidAnimatedMapNameMismatch")
 	);
 
@@ -713,8 +713,8 @@ IMPLEMENT_RIGUNIT_AUTOMATION_TEST(FRigUnit_RigLogic)
 	Test.Exec_MapMaskMultipliers(TestCurveContainerForAnimMapsValid.Get());
 	//Assert
 	AddErrorIfFalse(
-		Test.GetData()->CurveElementIndicesForAnimMaps.Num() == 0 &&
-		Test.GetData()->RigLogicIndicesForAnimMaps.Num() == 0,
+		Test.GetData()->Mappings->CurveElementIndicesForAnimMaps.Num() == 0 &&
+		Test.GetData()->Mappings->RigLogicIndicesForAnimMaps.Num() == 0,
 		TEXT("MapMaskMultipliers - EmptyReader ValidAnimatedMap")
 	);
 
@@ -725,12 +725,12 @@ IMPLEMENT_RIGUNIT_AUTOMATION_TEST(FRigUnit_RigLogic)
 	Test.Exec_MapMaskMultipliers(TestCurveContainerEmpty.Get());
 	//Assert
 	AddErrorIfFalse(
-		Test.GetData()->CurveElementIndicesForAnimMaps.Num() == 1 &&
-		Test.GetData()->CurveElementIndicesForAnimMaps[0].Values.Num() == 1 &&
-		Test.GetData()->CurveElementIndicesForAnimMaps[0].Values[0] == INDEX_NONE &&
-		Test.GetData()->RigLogicIndicesForAnimMaps.Num() == 1 &&
-		Test.GetData()->RigLogicIndicesForAnimMaps[0].Values.Num() == 1 &&
-		Test.GetData()->RigLogicIndicesForAnimMaps[0].Values[0] == 0,
+		Test.GetData()->Mappings->CurveElementIndicesForAnimMaps.Num() == 1 &&
+		Test.GetData()->Mappings->CurveElementIndicesForAnimMaps[0].Values.Num() == 1 &&
+		Test.GetData()->Mappings->CurveElementIndicesForAnimMaps[0].Values[0] == INDEX_NONE &&
+		Test.GetData()->Mappings->RigLogicIndicesForAnimMaps.Num() == 1 &&
+		Test.GetData()->Mappings->RigLogicIndicesForAnimMaps[0].Values.Num() == 1 &&
+		Test.GetData()->Mappings->RigLogicIndicesForAnimMaps[0].Values[0] == 0,
 		TEXT("MapMaskMultipliers - ValidReader EmptyCurveContainer")
 	);
 
@@ -743,12 +743,12 @@ IMPLEMENT_RIGUNIT_AUTOMATION_TEST(FRigUnit_RigLogic)
 	Test.Exec_MapMaskMultipliers(TestCurveContainerForAnimMapsValid.Get());
 	//Assert
 	AddErrorIfFalse(
-		Test.GetData()->CurveElementIndicesForAnimMaps.Num() == 1 &&
-		Test.GetData()->CurveElementIndicesForAnimMaps[0].Values.Num() == 1 &&
-		Test.GetData()->CurveElementIndicesForAnimMaps[0].Values[0] == 0 &&
-		Test.GetData()->RigLogicIndicesForAnimMaps.Num() == 1 &&
-		Test.GetData()->RigLogicIndicesForAnimMaps[0].Values.Num() == 1 &&
-		Test.GetData()->RigLogicIndicesForAnimMaps[0].Values[0] == 0,
+		Test.GetData()->Mappings->CurveElementIndicesForAnimMaps.Num() == 1 &&
+		Test.GetData()->Mappings->CurveElementIndicesForAnimMaps[0].Values.Num() == 1 &&
+		Test.GetData()->Mappings->CurveElementIndicesForAnimMaps[0].Values[0] == 0 &&
+		Test.GetData()->Mappings->RigLogicIndicesForAnimMaps.Num() == 1 &&
+		Test.GetData()->Mappings->RigLogicIndicesForAnimMaps[0].Values.Num() == 1 &&
+		Test.GetData()->Mappings->RigLogicIndicesForAnimMaps[0].Values[0] == 0,
 		TEXT("MapMaskMultipliers - Valid Inputs")
 	);
 
diff --git a/Engine/Plugins/Animation/RigLogic/Source/RigLogicModule/Public/RigUnit_RigLogic.h b/Engine/Plugins/Animation/RigLogic/Source/RigLogicModule/Public/RigUnit_RigLogic.h
--- a/Engine/Plugins/Animation/RigLogic/Source/RigLogicModule/Public/RigUnit_RigLogic.h
+++ b/Engine/Plugins/Animation/RigLogic/Source/RigLogicModule/Public/RigUnit_RigLogic.h
@@ -7,6 +7,7 @@
 
 #include "ControlRig.h"
 #include "DNAAsset.h"
+#include "Containers/DiscardableKeyValueCache.h"
 #include "Materials/Material.h"
 #include "Units/RigUnit.h"
 
@@ -47,28 +48,15 @@
 	FTransformArrayView DeltaTransforms; //the result of rig logic calculations
 };
 
-/* The work data used by the FRigUnit_RigLogic */
+/* A helper struct used inside of the RigUnit_RigLogic to store mappings between skelMesh and DNA. */
 USTRUCT(meta = (DocumentationPolicy = "Strict"))
-struct FRigUnit_RigLogic_Data
+struct FRigDefinitionMapping 
 {
+	
 	GENERATED_BODY()
 
-	FRigUnit_RigLogic_Data();
- 	~FRigUnit_RigLogic_Data();
- 	FRigUnit_RigLogic_Data(const FRigUnit_RigLogic_Data& Other);
- 	FRigUnit_RigLogic_Data& operator=(const FRigUnit_RigLogic_Data& Other);
+	FRigDefinitionMapping() = default;
+	FRigDefinitionMapping(const IBehaviorReader* DNABehavior, const URigHierarchy* InHierarchy);
 
-	/** Cached Skeletal Mesh Component **/
-	UPROPERTY(transient)
-	TWeakObjectPtr<USkeletalMeshComponent> SkelMeshComponent;
-
-	FSharedRigRuntimeContext* SharedRigRuntimeContext;
-
-	/** RigInstance is a thin class, containing only character instance specific data
-	  * i.e. the output buffers where RigLogic writes the results of computations, specific
-	  * to each instance of a rig.
-	**/
-	TUniquePtr<FRigInstance> RigInstance;
-
 	/** Mapping RL indices to UE indices
 	  * Note: we use int32 instead of uint32 to allow storing INDEX_NONE for missing elements
@@ -81,9 +69,12 @@
 
 	/** RL joint index to ControlRig's hierarchy bone index **/
 	UPROPERTY(transient)
-	TArray<int32> HierarchyBoneIndices;
+	TArray<int32> HierarchyBoneIndices;	/** RL mesh blend shape index to ControlRig's output blendshape curve index for each LOD **/
 
-	/** RL mesh blend shape index to ControlRig's output blendshape curve index for each LOD **/
+	/** Mapping RL indices to UE indices
+	  * Note: we use int32 instead of uint32 to allow storing INDEX_NONE for missing elements
+	  * if value is valid, it is cast to appropriate uint type
+	**/
 	UPROPERTY(transient)
 	TArray<FRigUnit_RigLogic_IntArray> MorphTargetCurveIndices;
 
@@ -99,6 +90,45 @@
 	UPROPERTY(transient)
 	TArray<FRigUnit_RigLogic_IntArray> RigLogicIndicesForAnimMaps;
 
+private:
+	/** Makes a map of input curve indices from DNA file to the control rig curves **/
+	void MapInputCurveIndices(const IBehaviorReader* DNABehavior, const URigHierarchy* InHierarchy);
+	/** Uses names to map joint indices from DNA file to the indices of bones in control rig hierarchy **/
+	void MapJoints(const IBehaviorReader* DNABehavior, const URigHierarchy* Hierarchy);
+	/** Uses names of blend shapes and meshes in DNA file, for all LODs, to map their indices to the indices of
+	  * morph target curves in the curve container; curve name format is <mesh>__<blendshape> **/
+	void MapMorphTargets(const IBehaviorReader* DNABehavior, const URigHierarchy* InHierarchy);
+	/** Uses names to map mask multiplier indices from DNA file, for all LODs, to the indices of curves in the
+	  * control rig's curve container **/
+	void MapMaskMultipliers(const IBehaviorReader* DNABehavior, const URigHierarchy* InHierarchy);
+
+};
+
+/* The work data used by the FRigUnit_RigLogic */
+USTRUCT(meta = (DocumentationPolicy = "Strict"))
+struct FRigUnit_RigLogic_Data
+{
+	GENERATED_BODY()
+
+	FRigUnit_RigLogic_Data();
+ 	~FRigUnit_RigLogic_Data();
+ 	FRigUnit_RigLogic_Data(const FRigUnit_RigLogic_Data& Other);
+ 	FRigUnit_RigLogic_Data& operator=(const FRigUnit_RigLogic_Data& Other);
+
+	/** Cached Skeletal Mesh Component **/
+	UPROPERTY(transient)
+	TWeakObjectPtr<USkeletalMeshComponent> SkelMeshComponent;
+
+	FSharedRigRuntimeContext* SharedRigRuntimeContext;
+
+	/** RigInstance is a thin class, containing only character instance specific data
+	  * i.e. the output buffers where RigLogic writes the results of computations, specific
+	  * to each instance of a rig.
+	**/
+	TUniquePtr<FRigInstance> RigInstance;
+
+	TSharedPtr<const FRigDefinitionMapping> Mappings;
+
 	/** LOD for which the model is rendered **/
 	UPROPERTY(transient)
 	uint32 CurrentLOD;
@@ -108,18 +138,8 @@
 	bool IsRigLogicInitialized();
 	void InitializeRigLogic(const URigHierarchy* InHierarchy);
 
-	/** Makes a map of input curve indices from DNA file to the control rig curves **/
-	void MapInputCurveIndices(const URigHierarchy* InHierarchy);
-	/** Uses names to map joint indices from DNA file to the indices of bones in control rig hierarchy **/
-	void MapJoints(const URigHierarchy* Hierarchy);
 	/** Cache the joint indices that change per each LOD **/
 	void CacheVariableJointIndices();
-	/** Uses names of blend shapes and meshes in DNA file, for all LODs, to map their indices to the indices of
-	  * morph target curves in the curve container; curve name format is <mesh>__<blendshape> **/
-	void MapMorphTargets(const URigHierarchy* InHierarchy);
-	/** Uses names to map mask multiplier indices from DNA file, for all LODs, to the indices of curves in the
-	  * control rig's curve container **/
-	void MapMaskMultipliers(const URigHierarchy* InHierarchy);
 
 	/** Calculates joint positions, orientation and scale based on inputs curves of the control rig **/
 	void CalculateRigLogic(const URigHierarchy* InHierarchy);
@@ -129,6 +149,8 @@
 	void UpdateBlendShapeCurves(URigHierarchy* InHierarchy, TArrayView<const float> BlendShapeValues);
 	/** Updates anim map curve values based on values of input curves of the control rig **/
 	void UpdateAnimMapCurves(URigHierarchy* InHierarchy, TArrayView<const float> AnimMapOutputs);
+private:
+	static TDiscardableKeyValueCache< TTuple<FString, uint16>, TSharedPtr<const FRigDefinitionMapping>> RigDefintionMappingsCache;
 };
 
 
