Fix buffer read overrun in RigLogic test suite.

#rnx

diff --git a/Engine/Plugins/Animation/RigLogic/Source/RigLogicLibTest/Private/tersetests/archives/binary/TestInputArchive.cpp b/Engine/Plugins/Animation/RigLogic/Source/RigLogicLibTest/Private/tersetests/archives/binary/TestInputArchive.cpp
--- a/Engine/Plugins/Animation/RigLogic/Source/RigLogicLibTest/Private/tersetests/archives/binary/TestInputArchive.cpp
+++ b/Engine/Plugins/Animation/RigLogic/Source/RigLogicLibTest/Private/tersetests/archives/binary/TestInputArchive.cpp
@@ -24,8 +24,8 @@
 
 TEST(BinaryInputArchiveTest, StreamExtractionOperator) {
     tersetests::FakeStream stream;
-    unsigned char bytes[] = {0x01, 0x02};
-    stream.write(reinterpret_cast<char*>(bytes), 2ul);
+    unsigned char bytes[2 * sizeof(char)] = {0x01, 0x02};
+    stream.write(reinterpret_cast<char*>(bytes), sizeof(bytes));
     stream.seek(0ul);
 
     terse::BinaryInputArchive<tersetests::FakeStream> archive(&stream);
@@ -39,7 +39,7 @@ TEST(BinaryInputArchiveTest, StreamExtractionOperator) {
 
 TEST(BinaryInputArchiveTest, ArchiveOffset) {
     tersetests::FakeStream stream;
-    unsigned char bytes[] = {
+    unsigned char bytes[33ul] = {
         0x00, 0x00, 0x00, 0x19,  // layerStart offset
         0x00, 0x00, 0x00, 0x19,  // first offset
         0x00, 0x00, 0x00, 0x1d,  // second offset
@@ -51,7 +51,7 @@ TEST(BinaryInputArchiveTest, ArchiveOffset) {
         0x00, 0x00, 0x04, 0x00  // secondInt = 1024
 
     };
-    stream.write(reinterpret_cast<char*>(bytes), 33ul);
+    stream.write(reinterpret_cast<char*>(bytes), sizeof(bytes));
     stream.seek(0ul);
 
     terse::BinaryInputArchive<tersetests::FakeStream> archive(&stream);
@@ -110,7 +110,7 @@ struct SectionSkipper {
 
 TEST(BinaryInputArchiveTest, ArchiveOffsetRepositionsStream) {
     tersetests::FakeStream stream;
-    unsigned char bytes[] = {
+    unsigned char bytes[22ul] = {
         0x00, 0x00, 0x00, 0x12,  // afterValue offset
         0x00, 0x00, 0x04, 0xd2,  // a = 1234
         0x00, 0x00, 0x00, 0x04,  // b size = 4
@@ -119,7 +119,7 @@ TEST(BinaryInputArchiveTest, ArchiveOffsetRepositionsStream) {
         0x00, 0x00, 0x04, 0x00  // afterValue = 1024
 
     };
-    stream.write(reinterpret_cast<char*>(bytes), 22ul);
+    stream.write(reinterpret_cast<char*>(bytes), sizeof(bytes));
     stream.seek(0ul);
 
     terse::BinaryInputArchive<tersetests::FakeStream> archive(&stream);
@@ -140,8 +140,8 @@ TEST(BinaryInputArchiveTest, ArchiveOffsetRepositionsStream) {
 
 TEST(BinaryInputArchiveTest, PrimitiveTypeDeserialization) {
     tersetests::FakeStream stream;
-    unsigned char bytes[] = {0x00, 0x00, 0x04, 0xd2};  // 1234
-    stream.write(reinterpret_cast<char*>(bytes), sizeof(std::int32_t));
+    unsigned char bytes[sizeof(std::int32_t)] = {0x00, 0x00, 0x04, 0xd2};  // 1234
+    stream.write(reinterpret_cast<char*>(bytes), sizeof(bytes));
     stream.seek(0ul);
 
     terse::BinaryInputArchive<tersetests::FakeStream> archive(&stream);
@@ -153,7 +153,7 @@ TEST(BinaryInputArchiveTest, PrimitiveTypeDeserialization) {
 
 TEST(BinaryInputArchiveTest, ComplexTypeDeserialization) {
     tersetests::FakeStream stream;
-    unsigned char bytes[] = {
+    unsigned char bytes[sizeof(ComplexType)] = {
         0x7f,  // 127
         0xff,  // 255
         0x7f, 0xff,  // 32767
@@ -161,7 +161,7 @@ TEST(BinaryInputArchiveTest, ComplexTypeDeserialization) {
         0x7f, 0xff, 0xff, 0xff,  // 2147483647
         0xff, 0xff, 0xff, 0xff  // 4294967295
     };
-    stream.write(reinterpret_cast<char*>(bytes), 14ul);
+    stream.write(reinterpret_cast<char*>(bytes), sizeof(bytes));
     stream.seek(0ul);
 
     terse::BinaryInputArchive<tersetests::FakeStream> archive(&stream);
@@ -178,11 +178,11 @@ TEST(BinaryInputArchiveTest, ComplexTypeDeserialization) {
 
 TEST(BinaryInputArchiveTest, FreeSerialize) {
     tersetests::FakeStream stream;
-    unsigned char bytes[] = {
+    unsigned char bytes[sizeof(SerializableByFreeSerialize)] = {
         0xff, 0xff, 0xff, 0xff,  // 4294967295
         0xff, 0xff  // 65535
     };
-    stream.write(reinterpret_cast<char*>(bytes), 6ul);
+    stream.write(reinterpret_cast<char*>(bytes), sizeof(bytes));
     stream.seek(0ul);
 
     terse::BinaryInputArchive<tersetests::FakeStream> archive(&stream);
@@ -195,11 +195,11 @@ TEST(BinaryInputArchiveTest, FreeSerialize) {
 
 TEST(BinaryInputArchiveTest, FreeLoad) {
     tersetests::FakeStream stream;
-    unsigned char bytes[] = {
+    unsigned char bytes[sizeof(SerializableByFreeLoadSave)] = {
         0xff, 0xff, 0xff, 0xff,  // 4294967295
         0xff, 0xff  // 65535
     };
-    stream.write(reinterpret_cast<char*>(bytes), 6ul);
+    stream.write(reinterpret_cast<char*>(bytes), sizeof(bytes));
     stream.seek(0ul);
 
     terse::BinaryInputArchive<tersetests::FakeStream> archive(&stream);
@@ -212,8 +212,8 @@ TEST(BinaryInputArchiveTest, FreeLoad) {
 
 TEST(BinaryInputArchiveTest, StringDeserialization) {
     tersetests::FakeStream stream;
-    unsigned char size[] = {0x00, 0x00, 0x00, 0x08};  // 8
-    unsigned char bytes[] = {
+    unsigned char size[sizeof(std::uint32_t)] = {0x00, 0x00, 0x00, 0x08};  // 8
+    unsigned char bytes[8] = {
         0x61,  // 'a'
         0x62,  // 'b'
         0x63,  // 'c'
@@ -223,8 +223,8 @@ TEST(BinaryInputArchiveTest, StringDeserialization) {
         0x67,  // 'g'
         0x68  // 'h'
     };
-    stream.write(reinterpret_cast<char*>(size), 4ul);
-    stream.write(reinterpret_cast<char*>(bytes), 8ul);
+    stream.write(reinterpret_cast<char*>(size), sizeof(size));
+    stream.write(reinterpret_cast<char*>(bytes), sizeof(bytes));
     stream.seek(0ul);
 
     terse::BinaryInputArchive<tersetests::FakeStream> archive(&stream);
@@ -236,16 +236,16 @@ TEST(BinaryInputArchiveTest, StringDeserialization) {
 
 TEST(BinaryInputArchiveTest, VectorOfInt8TypeDeserialization) {
     tersetests::FakeStream stream;
-    unsigned char size[] = {0x00, 0x00, 0x00, 0x05};  // 5
-    unsigned char bytes[] = {
+    unsigned char size[sizeof(std::uint32_t)] = {0x00, 0x00, 0x00, 0x05};  // 5
+    unsigned char bytes[5] = {
         0x80,  // -128
         0x9c,  // -100
         0xff,  // -1
         0x0f,  // 15
         0x7f  // 127
     };
-    stream.write(reinterpret_cast<char*>(size), 4ul);
-    stream.write(reinterpret_cast<char*>(bytes), 5ul);
+    stream.write(reinterpret_cast<char*>(size), sizeof(size));
+    stream.write(reinterpret_cast<char*>(bytes), sizeof(bytes));
     stream.seek(0ul);
 
     terse::BinaryInputArchive<tersetests::FakeStream> archive(&stream);
@@ -258,16 +258,16 @@ TEST(BinaryInputArchiveTest, VectorOfInt8TypeDeserialization) {
 
 TEST(BinaryInputArchiveTest, VectorOfInt16TypeDeserialization) {
     tersetests::FakeStream stream;
-    unsigned char size[] = {0x00, 0x00, 0x00, 0x05};  // 5
-    unsigned char bytes[] = {
+    unsigned char size[sizeof(std::uint32_t)] = {0x00, 0x00, 0x00, 0x05};  // 5
+    unsigned char bytes[5 * sizeof(std::uint16_t)] = {
         0x80, 0x00,  // -32768
         0xff, 0xf9,  // -7
         0x00, 0x0f,  // 15
         0x00, 0xff,  // 255
         0x7f, 0xff  // 32767
     };
-    stream.write(reinterpret_cast<char*>(size), 4ul);
-    stream.write(reinterpret_cast<char*>(bytes), 10ul);
+    stream.write(reinterpret_cast<char*>(size), sizeof(size));
+    stream.write(reinterpret_cast<char*>(bytes), sizeof(bytes));
     stream.seek(0ul);
 
     terse::BinaryInputArchive<tersetests::FakeStream> archive(&stream);
@@ -281,8 +281,8 @@ TEST(BinaryInputArchiveTest, VectorOfInt16TypeDeserialization) {
 
 TEST(BinaryInputArchiveTest, VectorOfInt32TypeDeserialization) {
     tersetests::FakeStream stream;
-    unsigned char size[] = {0x00, 0x00, 0x00, 0x08};  // 8
-    unsigned char bytes[] = {
+    unsigned char size[sizeof(std::uint32_t)] = {0x00, 0x00, 0x00, 0x08};  // 8
+    unsigned char bytes[8 * sizeof(std::int32_t)] = {
         0x80, 0x00, 0x00, 0x00,  // -2147483648
         0xff, 0xff, 0xff, 0xf0,  // -16
         0x00, 0x00, 0x00, 0x80,  // 128
@@ -292,8 +292,8 @@ TEST(BinaryInputArchiveTest, VectorOfInt32TypeDeserialization) {
         0x00, 0x02, 0x00, 0x00,  // 131072
         0x7f, 0xff, 0xff, 0xff  // 2147483647
     };
-    stream.write(reinterpret_cast<char*>(size), 4ul);
-    stream.write(reinterpret_cast<char*>(bytes), 32ul);
+    stream.write(reinterpret_cast<char*>(size), sizeof(size));
+    stream.write(reinterpret_cast<char*>(bytes), sizeof(bytes));
     stream.seek(0ul);
 
     terse::BinaryInputArchive<tersetests::FakeStream> archive(&stream);
@@ -307,16 +307,16 @@ TEST(BinaryInputArchiveTest, VectorOfInt32TypeDeserialization) {
 
 TEST(BinaryInputArchiveTest, VectorOfInt64TypeDeserialization) {
     tersetests::FakeStream stream;
-    unsigned char size[] = {0x00, 0x00, 0x00, 0x05};  // 5
-    unsigned char bytes[] = {
+    unsigned char size[sizeof(std::int32_t)] = {0x00, 0x00, 0x00, 0x05};  // 5
+    unsigned char bytes[5 * sizeof(std::int64_t)] = {
         0x80, 0x00, 0x00, 0x00, 0x00, 0x00, 0x00, 0x00,  // -9223372036854775808
         0xff, 0xff, 0xff, 0xff, 0xff, 0xff, 0xff, 0x00,  // -256
         0xff, 0xff, 0xff, 0xff, 0xff, 0xff, 0xff, 0xff,  // -1
         0x00, 0x00, 0x00, 0x00, 0x00, 0x00, 0x00, 0xff,  // 255
         0x7f, 0xff, 0xff, 0xff, 0xff, 0xff, 0xff, 0xff  // 9223372036854775807
     };
-    stream.write(reinterpret_cast<char*>(size), 4ul);
-    stream.write(reinterpret_cast<char*>(bytes), 40ul);
+    stream.write(reinterpret_cast<char*>(size), sizeof(size));
+    stream.write(reinterpret_cast<char*>(bytes), sizeof(bytes));
     stream.seek(0ul);
 
     terse::BinaryInputArchive<tersetests::FakeStream> archive(&stream);
@@ -330,16 +330,16 @@ TEST(BinaryInputArchiveTest, VectorOfInt64TypeDeserialization) {
 
 TEST(BinaryInputArchiveTest, VectorOfFloatTypeDeserialization) {
     tersetests::FakeStream stream;
-    unsigned char size[] = {0x00, 0x00, 0x00, 0x05};  // 5
-    unsigned char bytes[] = {
+    unsigned char size[sizeof(std::uint32_t)] = {0x00, 0x00, 0x00, 0x05};  // 5
+    unsigned char bytes[5 * sizeof(float)] = {
         0x3f, 0x99, 0x99, 0x9a,  // 1.2f
         0x40, 0x48, 0xf5, 0xc3,  // 3.14f
         0x3d, 0x4c, 0xcc, 0xcd,  // 0.05f
         0x3f, 0x3d, 0x70, 0xa4,  // 0.74f
         0x43, 0x2, 0xe6, 0x66,  // 130.9f
     };
-    stream.write(reinterpret_cast<char*>(size), 4ul);
-    stream.write(reinterpret_cast<char*>(bytes), 24ul);
+    stream.write(reinterpret_cast<char*>(size), sizeof(size));
+    stream.write(reinterpret_cast<char*>(bytes), sizeof(bytes));
     stream.seek(0ul);
 
     terse::BinaryInputArchive<tersetests::FakeStream> archive(&stream);
@@ -352,16 +352,16 @@ TEST(BinaryInputArchiveTest, VectorOfFloatTypeDeserialization) {
 
 TEST(BinaryInputArchiveTest, VectorOfDoubleTypeDeserialization) {
     tersetests::FakeStream stream;
-    unsigned char size[] = {0x00, 0x00, 0x00, 0x05};  // 5
-    unsigned char bytes[] = {
+    unsigned char size[sizeof(std::uint32_t)] = {0x00, 0x00, 0x00, 0x05};  // 5
+    unsigned char bytes[5 * sizeof(double)] = {
         0xc0, 0x24, 0x00, 0x00, 0x00, 0x00, 0x00, 0x00,
         0x40, 0x09, 0x1e, 0xb8, 0x51, 0xeb, 0x85, 0x1f,
         0x3f, 0xa9, 0x99, 0x99, 0x99, 0x99, 0x99, 0x9a,
         0x3f, 0xe7, 0xae, 0x14, 0x7a, 0xe1, 0x47, 0xae,
         0x40, 0x60, 0x5c, 0xcc, 0xcc, 0xcc, 0xcc, 0xcd
     };
-    stream.write(reinterpret_cast<char*>(size), 4ul);
-    stream.write(reinterpret_cast<char*>(bytes), 40ul);
+    stream.write(reinterpret_cast<char*>(size), sizeof(size));
+    stream.write(reinterpret_cast<char*>(bytes), sizeof(bytes));
     stream.seek(0ul);
 
     terse::BinaryInputArchive<tersetests::FakeStream> archive(&stream);
@@ -374,8 +374,8 @@ TEST(BinaryInputArchiveTest, VectorOfDoubleTypeDeserialization) {
 
 TEST(BinaryInputArchiveTest, VectorOfComplexTypeDeserialization) {
     tersetests::FakeStream stream;
-    unsigned char size[] = {0x00, 0x00, 0x00, 0x02};  // 2
-    unsigned char bytes[] = {
+    unsigned char size[sizeof(std::uint32_t)] = {0x00, 0x00, 0x00, 0x02};  // 2
+    unsigned char bytes[2 * sizeof(ComplexType)] = {
         0x7f,  // 127
         0xff,  // 255
         0x7f, 0xff,  // 32767
@@ -389,8 +389,8 @@ TEST(BinaryInputArchiveTest, VectorOfComplexTypeDeserialization) {
         0x7f, 0xff, 0xff, 0xff,  // 2147483647
         0xff, 0xff, 0xff, 0xff  // 4294967295
     };
-    stream.write(reinterpret_cast<char*>(size), 4ul);
-    stream.write(reinterpret_cast<char*>(bytes), 28ul);
+    stream.write(reinterpret_cast<char*>(size), sizeof(size));
+    stream.write(reinterpret_cast<char*>(bytes), sizeof(bytes));
     stream.seek(0ul);
 
     terse::BinaryInputArchive<tersetests::FakeStream> archive(&stream);
@@ -424,10 +424,10 @@ struct INeedAMemoryResource {
 
 TEST(BinaryInputArchiveTest, VectorOfTypeConstructibleFromMemoryResourceDeserialization) {
     tersetests::FakeStream stream;
-    unsigned char size[] = {0x00, 0x00, 0x00, 0x01};  // 1
-    unsigned char bytes[] = {0x00, 0x00, 0x04, 0xd2};  // 1234
-    stream.write(reinterpret_cast<char*>(size), 4ul);
-    stream.write(reinterpret_cast<char*>(bytes), sizeof(std::int32_t));
+    unsigned char size[sizeof(std::uint32_t)] = {0x00, 0x00, 0x00, 0x01};  // 1
+    unsigned char bytes[sizeof(std::int32_t)] = {0x00, 0x00, 0x04, 0xd2};  // 1234
+    stream.write(reinterpret_cast<char*>(size), sizeof(size));
+    stream.write(reinterpret_cast<char*>(bytes), sizeof(bytes));
     stream.seek(0ul);
 
     terse::BinaryInputArchive<tersetests::FakeStream> archive(&stream);
@@ -445,9 +445,9 @@ TEST(BinaryInputArchiveTest, VectorOfTypeConstructibleFromMemoryResourceDeserial
 
 TEST(BinaryInputArchiveTest, VectorOfStringDeserialization) {
     tersetests::FakeStream stream;
-    unsigned char vectorSize[] = {0x00, 0x00, 0x00, 0x02};  // 2
-    unsigned char stringSize[] = {0x00, 0x00, 0x00, 0x08};  // 8
-    unsigned char stringBytes[] = {
+    unsigned char vectorSize[sizeof(std::uint32_t)] = {0x00, 0x00, 0x00, 0x02};  // 2
+    unsigned char stringSize[sizeof(std::uint32_t)] = {0x00, 0x00, 0x00, 0x08};  // 8
+    unsigned char stringBytes[8ul] = {
         0x61,  // 'a'
         0x62,  // 'b'
         0x63,  // 'c'
@@ -457,13 +457,13 @@ TEST(BinaryInputArchiveTest, VectorOfStringDeserialization) {
         0x67,  // 'g'
         0x68  // 'h'
     };
-    stream.write(reinterpret_cast<char*>(vectorSize), 4ul);
+    stream.write(reinterpret_cast<char*>(vectorSize), sizeof(vectorSize));
     // Write first string
-    stream.write(reinterpret_cast<char*>(stringSize), 4ul);
-    stream.write(reinterpret_cast<char*>(stringBytes), 8ul);
+    stream.write(reinterpret_cast<char*>(stringSize), sizeof(stringSize));
+    stream.write(reinterpret_cast<char*>(stringBytes), sizeof(stringBytes));
     // Write second string
-    stream.write(reinterpret_cast<char*>(stringSize), 4ul);
-    stream.write(reinterpret_cast<char*>(stringBytes), 8ul);
+    stream.write(reinterpret_cast<char*>(stringSize), sizeof(stringSize));
+    stream.write(reinterpret_cast<char*>(stringBytes), sizeof(stringBytes));
 
     stream.seek(0ul);
 
@@ -479,23 +479,23 @@ TEST(BinaryInputArchiveTest, VectorOfStringDeserialization) {
 
 TEST(BinaryInputArchiveTest, VectorOfPairOfStringsDeserialization) {
     tersetests::FakeStream stream;
-    unsigned char vectorSize[] = {0x00, 0x00, 0x00, 0x02};  // 2
-    unsigned char stringSize[] = {0x00, 0x00, 0x00, 0x04};  // 4
-    unsigned char keyBytes[] = {0x61, 0x62, 0x63, 0x64};  // "abcd"
-    unsigned char valueBytes[] = {0x65, 0x66, 0x67, 0x68};  // "efgh"
-    stream.write(reinterpret_cast<char*>(vectorSize), 4ul);
+    unsigned char vectorSize[sizeof(std::uint32_t)] = {0x00, 0x00, 0x00, 0x02};  // 2
+    unsigned char stringSize[sizeof(std::uint32_t)] = {0x00, 0x00, 0x00, 0x04};  // 4
+    unsigned char keyBytes[4] = {0x61, 0x62, 0x63, 0x64};  // "abcd"
+    unsigned char valueBytes[4] = {0x65, 0x66, 0x67, 0x68};  // "efgh"
+    stream.write(reinterpret_cast<char*>(vectorSize), sizeof(vectorSize));
     // Write key
-    stream.write(reinterpret_cast<char*>(stringSize), 4ul);
-    stream.write(reinterpret_cast<char*>(keyBytes), 4ul);
+    stream.write(reinterpret_cast<char*>(stringSize), sizeof(stringSize));
+    stream.write(reinterpret_cast<char*>(keyBytes), sizeof(keyBytes));
     // Write value
-    stream.write(reinterpret_cast<char*>(stringSize), 4ul);
-    stream.write(reinterpret_cast<char*>(valueBytes), 4ul);
+    stream.write(reinterpret_cast<char*>(stringSize), sizeof(stringSize));
+    stream.write(reinterpret_cast<char*>(valueBytes), sizeof(valueBytes));
     // Write key
-    stream.write(reinterpret_cast<char*>(stringSize), 4ul);
-    stream.write(reinterpret_cast<char*>(keyBytes), 4ul);
+    stream.write(reinterpret_cast<char*>(stringSize), sizeof(stringSize));
+    stream.write(reinterpret_cast<char*>(keyBytes), sizeof(keyBytes));
     // Write value
-    stream.write(reinterpret_cast<char*>(stringSize), 4ul);
-    stream.write(reinterpret_cast<char*>(valueBytes), 4ul);
+    stream.write(reinterpret_cast<char*>(stringSize), sizeof(stringSize));
+    stream.write(reinterpret_cast<char*>(valueBytes), sizeof(valueBytes));
 
     stream.seek(0ul);
 
@@ -512,21 +512,21 @@ TEST(BinaryInputArchiveTest, VectorOfPairOfStringsDeserialization) {
 
 TEST(BinaryInputArchiveTest, VectorOfVectorOfPrimitiveTypeDeserialization) {
     tersetests::FakeStream stream;
-    unsigned char outerVectorSize[] = {0x00, 0x00, 0x00, 0x02};  // 2
-    unsigned char innerVectorSize[] = {0x00, 0x00, 0x00, 0x04};  // 4
-    unsigned char innerVectorBytes[] = {
+    unsigned char outerVectorSize[sizeof(std::uint32_t)] = {0x00, 0x00, 0x00, 0x02};  // 2
+    unsigned char innerVectorSize[sizeof(std::uint32_t)] = {0x00, 0x00, 0x00, 0x04};  // 4
+    unsigned char innerVectorBytes[4 * sizeof(std::uint16_t)] = {
         0x00, 0x04,  // 4
         0x00, 0x10,  // 16
         0x01, 0x00,  // 256
         0x10, 0x00  // 4096
     };
-    stream.write(reinterpret_cast<char*>(outerVectorSize), 4ul);
+    stream.write(reinterpret_cast<char*>(outerVectorSize), sizeof(outerVectorSize));
     // Write first inner vector
-    stream.write(reinterpret_cast<char*>(innerVectorSize), 4ul);
-    stream.write(reinterpret_cast<char*>(innerVectorBytes), 8ul);
+    stream.write(reinterpret_cast<char*>(innerVectorSize), sizeof(innerVectorSize));
+    stream.write(reinterpret_cast<char*>(innerVectorBytes), sizeof(innerVectorBytes));
     // Write second inner vector
-    stream.write(reinterpret_cast<char*>(innerVectorSize), 4ul);
-    stream.write(reinterpret_cast<char*>(innerVectorBytes), 8ul);
+    stream.write(reinterpret_cast<char*>(innerVectorSize), sizeof(innerVectorSize));
+    stream.write(reinterpret_cast<char*>(innerVectorBytes), sizeof(innerVectorBytes));
     stream.seek(0ul);
 
     pma::AlignedMemoryResource memRes;
@@ -553,7 +553,7 @@ TEST(BinaryInputArchiveTest, VectorOfVectorOfPrimitiveTypeDeserialization) {
 
 TEST(BinaryInputArchiveTest, MixedNestedTypes) {
     tersetests::FakeStream stream;
-    unsigned char bytes[] = {
+    unsigned char bytes[68ul] = {
         0x00, 0x00, 0x00, 0x01,  // 1
         0x00, 0x00, 0x00, 0x04,  // 4
         0x72, 0x6f, 0x6f, 0x74,  // "root"
@@ -567,7 +567,7 @@ TEST(BinaryInputArchiveTest, MixedNestedTypes) {
         0x00, 0x00, 0x00, 0x03,  // 3
         0x3c, 0x23, 0xd7, 0x0a, 0x3f, 0x7d, 0x70, 0xa4, 0x40, 0x48, 0xf5, 0xc3  // 0.01f, 0.99f, 3.14f
     };
-    stream.write(reinterpret_cast<char*>(bytes), 68ul);
+    stream.write(reinterpret_cast<char*>(bytes), sizeof(bytes));
     stream.seek(0ul);
 
     pma::AlignedMemoryResource memRes;
