Remove obsolete configurations from RigLogicLibTest build script #rb none #robomerge 5.0

diff --git a/Engine/Plugins/Animation/RigLogic/Source/RigLogicLibTest/RigLogicLibTest.Build.cs b/Engine/Plugins/Animation/RigLogic/Source/RigLogicLibTest/RigLogicLibTest.Build.cs
--- a/Engine/Plugins/Animation/RigLogic/Source/RigLogicLibTest/RigLogicLibTest.Build.cs
+++ b/Engine/Plugins/Animation/RigLogic/Source/RigLogicLibTest/RigLogicLibTest.Build.cs
@@ -9,11 +9,10 @@ public class RigLogicLibTest : ModuleRules
 {
     public RigLogicLibTest(ReadOnlyTargetRules Target) : base(Target)
     {
-        if (Target.Platform == UnrealTargetPlatform.Win64 ||
-            Target.Platform == UnrealTargetPlatform.Linux ||
-            Target.Platform == UnrealTargetPlatform.Mac)
+        if (Target.Platform == UnrealTargetPlatform.Win32 || Target.Platform == UnrealTargetPlatform.Win64)
         {
             PrivateDefinitions.Add("RL_BUILD_WITH_SSE=1");
+            PublicDefinitions.Add("GTEST_OS_WINDOWS=1");
         }
 
         string RigLogicLibPath = Path.GetFullPath(Path.Combine(ModuleDirectory, "../RigLogicLib"));
@@ -36,26 +35,5 @@ public class RigLogicLibTest : ModuleRules
                 "GoogleTest"
             }
         );
-
-        if (Target.Platform == UnrealTargetPlatform.Win64)
-        {
-            PublicDefinitions.Add("GTEST_OS_WINDOWS=1");
-        }
-        else if (Target.Platform == UnrealTargetPlatform.Mac)
-        {
-            PublicDefinitions.Add("GTEST_OS_MAC=1");
-        }
-        else if (Target.Platform == UnrealTargetPlatform.IOS || Target.Platform == UnrealTargetPlatform.TVOS)
-        {
-            PublicDefinitions.Add("GTEST_OS_IOS=1");
-        }
-        else if (Target.Platform == UnrealTargetPlatform.Android)
-        {
-            PublicDefinitions.Add("GTEST_OS_LINUX_ANDROID=1");
-        }
-        else if (Target.IsInPlatformGroup(UnrealPlatformGroup.Unix) || Target.Platform == UnrealTargetPlatform.PS4)
-        {
-            PublicDefinitions.Add("GTEST_OS_LINUX=1");
-        }
     }
 }
