Fix obsolete platform access #rb none #robomerge 5.0

diff --git a/Engine/Plugins/Animation/RigLogic/Source/RigLogicLibTest/RigLogicLibTest.Build.cs b/Engine/Plugins/Animation/RigLogic/Source/RigLogicLibTest/RigLogicLibTest.Build.cs
--- a/Engine/Plugins/Animation/RigLogic/Source/RigLogicLibTest/RigLogicLibTest.Build.cs
+++ b/Engine/Plugins/Animation/RigLogic/Source/RigLogicLibTest/RigLogicLibTest.Build.cs
@@ -9,7 +9,7 @@ public class RigLogicLibTest : ModuleRules
 {
     public RigLogicLibTest(ReadOnlyTargetRules Target) : base(Target)
     {
-        if (Target.Platform == UnrealTargetPlatform.Win32 || Target.Platform == UnrealTargetPlatform.Win64)
+        if (Target.Platform == UnrealTargetPlatform.Win64)
         {
             PrivateDefinitions.Add("RL_BUILD_WITH_SSE=1");
             PublicDefinitions.Add("GTEST_OS_WINDOWS=1");
