RigLogic: Remove #include cycle

diff --git a/Engine/Plugins/Animation/RigLogic/Source/RigLogicLib/Private/riglogic/joints/bpcm/strategies/AVX.h b/Engine/Plugins/Animation/RigLogic/Source/RigLogicLib/Private/riglogic/joints/bpcm/strategies/AVX.h
--- a/Engine/Plugins/Animation/RigLogic/Source/RigLogicLib/Private/riglogic/joints/bpcm/strategies/AVX.h
+++ b/Engine/Plugins/Animation/RigLogic/Source/RigLogicLib/Private/riglogic/joints/bpcm/strategies/AVX.h
@@ -3,7 +3,7 @@
 #pragma once
 
 #include "riglogic/joints/bpcm/strategies/Block8.h"
-#include "riglogic/system/simd/SIMD.h"
+#include "trimd/TRiMD.h"
 
 namespace rl4 {
 
diff --git a/Engine/Plugins/Animation/RigLogic/Source/RigLogicLib/Private/riglogic/joints/bpcm/strategies/SSE.h b/Engine/Plugins/Animation/RigLogic/Source/RigLogicLib/Private/riglogic/joints/bpcm/strategies/SSE.h
--- a/Engine/Plugins/Animation/RigLogic/Source/RigLogicLib/Private/riglogic/joints/bpcm/strategies/SSE.h
+++ b/Engine/Plugins/Animation/RigLogic/Source/RigLogicLib/Private/riglogic/joints/bpcm/strategies/SSE.h
@@ -3,7 +3,7 @@
 #pragma once
 
 #include "riglogic/joints/bpcm/strategies/Block4.h"
-#include "riglogic/system/simd/SIMD.h"
+#include "trimd/TRiMD.h"
 
 namespace rl4 {
 
diff --git a/Engine/Plugins/Animation/RigLogic/Source/RigLogicLib/Private/riglogic/joints/bpcm/strategies/Scalar.h b/Engine/Plugins/Animation/RigLogic/Source/RigLogicLib/Private/riglogic/joints/bpcm/strategies/Scalar.h
--- a/Engine/Plugins/Animation/RigLogic/Source/RigLogicLib/Private/riglogic/joints/bpcm/strategies/Scalar.h
+++ b/Engine/Plugins/Animation/RigLogic/Source/RigLogicLib/Private/riglogic/joints/bpcm/strategies/Scalar.h
@@ -3,7 +3,7 @@
 #pragma once
 
 #include "riglogic/joints/bpcm/strategies/Block4.h"
-#include "riglogic/system/simd/SIMD.h"
+#include "trimd/TRiMD.h"
 
 namespace rl4 {
 
diff --git a/Engine/Plugins/Animation/RigLogic/Source/RigLogicLib/Private/riglogic/system/simd/SIMD.h b/Engine/Plugins/Animation/RigLogic/Source/RigLogicLib/Private/riglogic/system/simd/SIMD.h
--- a/Engine/Plugins/Animation/RigLogic/Source/RigLogicLib/Private/riglogic/system/simd/SIMD.h
+++ b/Engine/Plugins/Animation/RigLogic/Source/RigLogicLib/Private/riglogic/system/simd/SIMD.h
@@ -14,4 +14,3 @@
     #define TRIMD_ENABLE_SSE
 #endif  // RL_BUILD_WITH_SSE
 
-#include <trimd/TRiMD.h>
