Added caching of RigUnit_RigLogicData mappings, which would allow for multiple rigs with same dna and topology version to share mappings.

diff --git a/Engine/Plugins/Animation/RigLogic/Source/RigLogicModule/Private/DNAAsset.cpp b/Engine/Plugins/Animation/RigLogic/Source/RigLogicModule/Private/DNAAsset.cpp
--- a/Engine/Plugins/Animation/RigLogic/Source/RigLogicModule/Private/DNAAsset.cpp
+++ b/Engine/Plugins/Animation/RigLogic/Source/RigLogicModule/Private/DNAAsset.cpp
@@ -7,10 +7,11 @@
 #include "DNAAssetCustomVersion.h"
 #include "DNAReaderAdapter.h"
 #include "FMemoryResource.h"
+#include "Hasher.h"
 #include "RigLogicMemoryStream.h"
 
 #if WITH_EDITORONLY_DATA
-    #include "EditorFramework/AssetImportData.h"
+#include "EditorFramework/AssetImportData.h"
 #endif
 #include "Engine/AssetUserData.h"
 #include "Serialization/BufferArchive.h"
@@ -83,6 +84,25 @@ static void InvalidateSharedRigRuntimeContext(FSharedRigRuntimeContext& Context)
 {
 	Context.RigLogic = nullptr;
 	Context.VariableJointIndices.Reset();
+	Context.DNAHash = 0u;
+}
+
+static uint32 HashDescriptor(const IDescriptorReader* Reader)
+{
+	Hasher DescriptorHasher;
+	DescriptorHasher.Update(Reader->GetName());
+	DescriptorHasher.Update(Reader->GetArchetype());
+	DescriptorHasher.Update(Reader->GetGender());
+	DescriptorHasher.Update(Reader->GetAge());
+	DescriptorHasher.Update(Reader->GetMetaDataCount());
+	DescriptorHasher.Update(Reader->GetTranslationUnit());
+	DescriptorHasher.Update(Reader->GetRotationUnit());
+	DescriptorHasher.Update(Reader->GetCoordinateSystem());
+	DescriptorHasher.Update(Reader->GetLODCount());
+	DescriptorHasher.Update(Reader->GetDBMaxLOD());
+	DescriptorHasher.Update(Reader->GetDBComplexity());
+	DescriptorHasher.Update(Reader->GetDBName());
+	return DescriptorHasher.GetHash();
 }
 
 UDNAAsset::UDNAAsset() :
@@ -98,22 +118,22 @@ bool UDNAAsset::Init(const FString& DNAFilename)
 	}
 
 	this->DNAFileName = DNAFilename; //memorize for re-import
-	
+
 	if (!FPaths::FileExists(DNAFilename))
 	{
 		UE_LOG(LogDNAAsset, Error, TEXT("DNA file %s doesn't exist!"), *DNAFilename);
 		return false;
 	}
-	
+
 	// Temporary buffer for the DNA file
 	TArray<uint8> TempFileBuffer;
-	
+
 	if (!FFileHelper::LoadFileToArray(TempFileBuffer, *DNAFilename)) //load entire DNA file into the array
 	{
 		UE_LOG(LogDNAAsset, Error, TEXT("Couldn't read DNA file %s!"), *DNAFilename);
 		return false;
 	}
-	
+
 	// Load run-time data (behavior) from whole-DNA buffer into BehaviorReader
 	InvalidateSharedRigRuntimeContext(Context);
 	Context.BehaviorReader = ReadDNAFromBuffer(&TempFileBuffer, EDNADataLayer::Behavior, 0u); //0u = MaxLOD
@@ -121,6 +141,7 @@ bool UDNAAsset::Init(const FString& DNAFilename)
 	{
 		return false;
 	}
+	Context.DNAHash = HashDescriptor(Context.BehaviorReader.Get());
 
 #if WITH_EDITORONLY_DATA
 	//We use geometry part of the data in MHC only (for updating the SkeletalMesh with
@@ -153,12 +174,13 @@ void UDNAAsset::Serialize(FArchive& Ar)
 		{
 			InvalidateSharedRigRuntimeContext(Context);
 
-			FArchiveMemoryStream BehaviorStream{&Ar};
+			FArchiveMemoryStream BehaviorStream{ &Ar };
 			Context.BehaviorReader = ReadDNAFromStream(&BehaviorStream, EDNADataLayer::Behavior, 0u); //0u = max LOD
+			Context.DNAHash = HashDescriptor(Context.BehaviorReader.Get());
 
 			// Geometry data is always present (even if only as an empty placeholder), just so the uasset
 			// format remains consistent between editor and non-editor builds
-			FArchiveMemoryStream GeometryStream{&Ar};
+			FArchiveMemoryStream GeometryStream{ &Ar };
 			auto Reader = ReadDNAFromStream(&GeometryStream, EDNADataLayer::Geometry, 0u); //0u = max LOD
 #if WITH_EDITORONLY_DATA
 			// Geometry data is discarded unless in Editor
@@ -170,14 +192,14 @@ void UDNAAsset::Serialize(FArchive& Ar)
 		{
 			TSharedPtr<IDNAReader> EmptyDNA = CreateEmptyDNA(AVG_EMPTY_SIZE);
 			IDNAReader* BehaviorReaderPtr = (Context.BehaviorReader.IsValid() ? static_cast<IDNAReader*>(Context.BehaviorReader.Get()) : EmptyDNA.Get());
-			FArchiveMemoryStream BehaviorStream{&Ar};
+			FArchiveMemoryStream BehaviorStream{ &Ar };
 			WriteDNAToStream(BehaviorReaderPtr, EDNADataLayer::Behavior, &BehaviorStream);
 
 			// When cooking (or when there was no Geometry data available), an empty DNA structure is written
 			// into the stream, serving as a placeholder just so uasset files can be conveniently loaded
 			// regardless if they were cooked or prepared for in-editor work
 			IDNAReader* GeometryReaderPtr = (Context.GeometryReader.IsValid() && !Ar.IsCooking() ? static_cast<IDNAReader*>(Context.GeometryReader.Get()) : EmptyDNA.Get());
-			FArchiveMemoryStream GeometryStream{&Ar};
+			FArchiveMemoryStream GeometryStream{ &Ar };
 			WriteDNAToStream(GeometryReaderPtr, EDNADataLayer::Geometry, &GeometryStream);
 		}
 	}
@@ -187,6 +209,7 @@ void UDNAAsset::SetBehaviorReader(TSharedPtr<IDNAReader> SourceDNAReader)
 {
 	InvalidateSharedRigRuntimeContext(Context);
 	Context.BehaviorReader = CopyDNALayer(SourceDNAReader.Get(), EDNADataLayer::Behavior, AVG_BEHAVIOR_SIZE);
+	Context.DNAHash = HashDescriptor(Context.BehaviorReader.Get());
 }
 
 void UDNAAsset::SetGeometryReader(TSharedPtr<IDNAReader> SourceDNAReader)
diff --git a/Engine/Plugins/Animation/RigLogic/Source/RigLogicModule/Private/Hasher.h b/Engine/Plugins/Animation/RigLogic/Source/RigLogicModule/Private/Hasher.h
new file mode 100644
--- /dev/null
+++ b/Engine/Plugins/Animation/RigLogic/Source/RigLogicModule/Private/Hasher.h
@@ -0,0 +1,36 @@
+// Copyright Epic Games, Inc. All Rights Reserved.
+#pragma once
+#include "CoreMinimal.h"
+class Hasher
+{
+public:
+	explicit Hasher(uint32 InitialValue = 0ul) : HashedValue{InitialValue}
+	{
+	}
+	uint32 GetHash() const
+	{
+		return HashedValue;
+	}
+	template <typename T>
+	void Update(T Value)
+	{
+		Update(reinterpret_cast<const char*>(&Value), sizeof(T));
+	}
+	template <typename T>
+	void Update(TArrayView<T> Values)
+	{
+		Update(reinterpret_cast<const char*>(Values.Data()), Values.Num());
+	}
+private:
+	void Update(const char* Bytes, size_t Length)
+	{
+		uint32 NewHash = HashedValue;
+		for (size_t i = 0ul; i < Length; ++i)
+		{
+			NewHash = Bytes[i] + (NewHash << 6) + (NewHash << 16) - NewHash;
+		}
+		HashedValue = NewHash;
+	}
+private:
+	uint32 HashedValue;
+};
\ No newline at end of file
diff --git a/Engine/Plugins/Animation/RigLogic/Source/RigLogicModule/Private/RigUnit_RigLogic.cpp b/Engine/Plugins/Animation/RigLogic/Source/RigLogicModule/Private/RigUnit_RigLogic.cpp
--- a/Engine/Plugins/Animation/RigLogic/Source/RigLogicModule/Private/RigUnit_RigLogic.cpp
+++ b/Engine/Plugins/Animation/RigLogic/Source/RigLogicModule/Private/RigUnit_RigLogic.cpp
@@ -39,11 +39,19 @@ static FString ConstructCurveName(const FString& NameToSplit, const FString& For
 	return CurveName;
 }
 
+FRigDefinitionMapping::FRigDefinitionMapping(const IBehaviorReader* DNABehavior, const URigHierarchy* InHierarchy) : FRigDefinitionMapping()
+{
+	MapInputCurveIndices(DNABehavior, InHierarchy);
+	MapJoints(DNABehavior, InHierarchy);
+	MapMaskMultipliers(DNABehavior, InHierarchy);
+	MapMorphTargets(DNABehavior, InHierarchy);
+}
+
 FRigUnit_RigLogic_Data::FRigUnit_RigLogic_Data()
-: SkelMeshComponent(nullptr)
-, SharedRigRuntimeContext(nullptr)
-, RigInstance(nullptr)
-, CurrentLOD(0)
+	: SkelMeshComponent(nullptr)
+	, SharedRigRuntimeContext(nullptr)
+	, RigInstance(nullptr)
+	, CurrentLOD(0)
 {
 }
 
@@ -62,12 +70,7 @@ FRigUnit_RigLogic_Data& FRigUnit_RigLogic_Data::operator=(const FRigUnit_RigLogi
 	SkelMeshComponent = Other.SkelMeshComponent;
 	SharedRigRuntimeContext = nullptr;
 	RigInstance = nullptr;
-	InputCurveIndices = Other.InputCurveIndices;
-	HierarchyBoneIndices = Other.HierarchyBoneIndices;
-	MorphTargetCurveIndices = Other.MorphTargetCurveIndices;
-	BlendShapeIndices = Other.BlendShapeIndices;
-	CurveElementIndicesForAnimMaps = Other.CurveElementIndicesForAnimMaps;
-	RigLogicIndicesForAnimMaps = Other.RigLogicIndicesForAnimMaps;
+	Mappings = Other.Mappings;
 	CurrentLOD = Other.CurrentLOD;
 	return *this;
 }
@@ -77,6 +80,8 @@ bool FRigUnit_RigLogic_Data::IsRigLogicInitialized()
 	return (SharedRigRuntimeContext != nullptr) && SharedRigRuntimeContext->RigLogic.IsValid() && RigInstance.IsValid();
 }
 
+TDiscardableKeyValueCache<uint32, TSharedPtr<const FRigDefinitionMapping>> FRigUnit_RigLogic_Data::RigDefintionMappingsCache{};
+
 void FRigUnit_RigLogic_Data::InitializeRigLogic(const URigHierarchy* InHierarchy)
 {
 	if ((SharedRigRuntimeContext == nullptr) || !SharedRigRuntimeContext->BehaviorReader.IsValid() || SharedRigRuntimeContext->BehaviorReader->GetJointCount() == 0u)
@@ -97,17 +102,25 @@ void FRigUnit_RigLogic_Data::InitializeRigLogic(const URigHierarchy* InHierarchy
 		RigInstance = MakeUnique<FRigInstance>(SharedRigRuntimeContext->RigLogic.Get());
 		RigInstance->SetLOD(CurrentLOD);
 
-		MapJoints(InHierarchy);
-		MapInputCurveIndices(InHierarchy);
-		MapMorphTargets(InHierarchy);
-		MapMaskMultipliers(InHierarchy);
+		const IBehaviorReader* DNABehavior = SharedRigRuntimeContext->BehaviorReader.Get();
+		const uint16 TopologyVersion = InHierarchy->GetTopologyVersion();
+		const uint32 CacheKey = SharedRigRuntimeContext->DNAHash + TopologyVersion;
+		auto& Cache = FRigUnit_RigLogic_Data::RigDefintionMappingsCache;
+
+		using FLockFLags = TDiscardableKeyValueCache< FString, TSharedPtr<FRigDefinitionMapping>>::LockFlags;
+		uint32 LockFlags = Cache.ApplyLock(0, FLockFLags::ReadLock);
+		if (!Cache.Find(CacheKey, Mappings, LockFlags | FLockFLags::WriteLockOnAddFail, LockFlags))
+		{
+			Mappings = MakeShared<const FRigDefinitionMapping>(DNABehavior, InHierarchy);
+			Cache.Add(CacheKey, Mappings, LockFlags);
+		}
+		Cache.Unlock(LockFlags);
 	}
 }
 
 //maps indices of input curves from dna file to control rig curves
-void FRigUnit_RigLogic_Data::MapInputCurveIndices(const URigHierarchy* InHierarchy)
+void FRigDefinitionMapping::MapInputCurveIndices(const IBehaviorReader* DNABehavior, const URigHierarchy* InHierarchy)
 {
-	const IBehaviorReader* DNABehavior = SharedRigRuntimeContext->BehaviorReader.Get();
 	const uint32 ControlCount = DNABehavior->GetRawControlCount();
 	InputCurveIndices.Reset(ControlCount);
 	for (uint32_t ControlIndex = 0; ControlIndex < ControlCount; ++ControlIndex)
@@ -124,16 +137,15 @@ void FRigUnit_RigLogic_Data::MapInputCurveIndices(const URigHierarchy* InHierarc
 	}
 }
 
-void FRigUnit_RigLogic_Data::MapJoints(const URigHierarchy* Hierarchy)
+void FRigDefinitionMapping::MapJoints(const IBehaviorReader* DNABehavior, const URigHierarchy* InHierarchy)
 {
-	const IBehaviorReader* DNABehavior = SharedRigRuntimeContext->BehaviorReader.Get();
 	const uint16 JointCount = DNABehavior->GetJointCount();
 	HierarchyBoneIndices.Reset(JointCount);
-	for (uint16 JointIndex = 0; JointIndex < JointCount ; ++JointIndex)
+	for (uint16 JointIndex = 0; JointIndex < JointCount; ++JointIndex)
 	{
 		const FString RLJointName = DNABehavior->GetJointName(JointIndex);
 		const FName JointFName = FName(*RLJointName);
-		const int32 BoneIndex = Hierarchy->GetIndex(FRigElementKey(JointFName, ERigElementType::Bone));
+		const int32 BoneIndex = InHierarchy->GetIndex(FRigElementKey(JointFName, ERigElementType::Bone));
 		HierarchyBoneIndices.Add(BoneIndex);
 	}
 }
@@ -159,9 +171,8 @@ void FRigUnit_RigLogic_Data::CacheVariableJointIndices()
 	}
 }
 
-void FRigUnit_RigLogic_Data::MapMorphTargets(const URigHierarchy* InHierarchy)
+void FRigDefinitionMapping::MapMorphTargets(const IBehaviorReader* DNABehavior, const URigHierarchy* InHierarchy)
 {
-	const IBehaviorReader* DNABehavior = SharedRigRuntimeContext->BehaviorReader.Get();
 	const uint16 LODCount = DNABehavior->GetLODCount();
 
 	MorphTargetCurveIndices.Reset();
@@ -175,7 +186,7 @@ void FRigUnit_RigLogic_Data::MapMorphTargets(const URigHierarchy* InHierarchy)
 			DNABehavior->GetMeshBlendShapeChannelMappingIndicesForLOD(LodIndex);
 		MorphTargetCurveIndices[LodIndex].Values.Reserve(BlendShapeChannelIndicesForLOD.Num());
 		BlendShapeIndices[LodIndex].Values.Reserve(BlendShapeChannelIndicesForLOD.Num());
-		for (uint16 MappingIndex: BlendShapeChannelIndicesForLOD)
+		for (uint16 MappingIndex : BlendShapeChannelIndicesForLOD)
 		{
 			const FMeshBlendShapeChannelMapping Mapping = DNABehavior->GetMeshBlendShapeChannelMapping(MappingIndex);
 			const uint16 BlendShapeIndex = Mapping.BlendShapeChannelIndex;
@@ -191,9 +202,8 @@ void FRigUnit_RigLogic_Data::MapMorphTargets(const URigHierarchy* InHierarchy)
 	}
 }
 
-void FRigUnit_RigLogic_Data::MapMaskMultipliers(const URigHierarchy* InHierarchy)
+void FRigDefinitionMapping::MapMaskMultipliers(const IBehaviorReader* DNABehavior, const URigHierarchy* InHierarchy)
 {
-	const IBehaviorReader* DNABehavior = SharedRigRuntimeContext->BehaviorReader.Get();
 	const uint16 LODCount = DNABehavior->GetLODCount();
 	CurveElementIndicesForAnimMaps.Reset();
 	CurveElementIndicesForAnimMaps.AddDefaulted(LODCount);
@@ -202,11 +212,11 @@ void FRigUnit_RigLogic_Data::MapMaskMultipliers(const URigHierarchy* InHierarchy
 	RigLogicIndicesForAnimMaps.AddDefaulted(LODCount);
 
 	for (uint16 LodIndex = 0; LodIndex < LODCount; ++LodIndex)
-	{		
+	{
 		TArrayView<const uint16> AnimMapIndicesPerLOD = DNABehavior->GetAnimatedMapIndicesForLOD(LodIndex);
 		CurveElementIndicesForAnimMaps[LodIndex].Values.Reserve(AnimMapIndicesPerLOD.Num());
 		RigLogicIndicesForAnimMaps[LodIndex].Values.Reserve(AnimMapIndicesPerLOD.Num());
-		for (uint16 AnimMapIndexPerLOD: AnimMapIndicesPerLOD)
+		for (uint16 AnimMapIndexPerLOD : AnimMapIndicesPerLOD)
 		{
 			const FString AnimMapNameFStr = DNABehavior->GetAnimatedMapName(AnimMapIndexPerLOD);
 			const FString MaskMultiplierNameStr = ConstructCurveName(AnimMapNameFStr, TEXT("<obj>_<attr>"));
@@ -231,7 +241,7 @@ void FRigUnit_RigLogic_Data::CalculateRigLogic(const URigHierarchy* InHierarchy)
 	const int32 RawControlCount = RigInstance->GetRawControlCount();
 	for (int32 ControlIndex = 0; ControlIndex < RawControlCount; ++ControlIndex)
 	{
-		const uint32 CurveIndex = InputCurveIndices[ControlIndex];
+		const uint32 CurveIndex = Mappings->InputCurveIndices[ControlIndex];
 		const float Value = InHierarchy->GetCurveValue(CurveIndex);
 		RigInstance->SetRawControl(ControlIndex, Value);
 	}
@@ -243,7 +253,7 @@ void FRigUnit_RigLogic_Data::UpdateJoints(URigHierarchy* Hierarchy, const FRigUn
 	FRigHierarchyValidityBracket ValidityBracket(Hierarchy);
 	for (const uint16 JointIndex : SharedRigRuntimeContext->VariableJointIndices[CurrentLOD].Values)
 	{
-		const int32 BoneIndex = HierarchyBoneIndices[JointIndex];
+		const int32 BoneIndex = Mappings->HierarchyBoneIndices[JointIndex];
 		if (BoneIndex != INDEX_NONE)
 		{
 			const FTransform& Neutral = JointUpdateParams.NeutralJointTransforms[JointIndex];
@@ -262,16 +272,16 @@ void FRigUnit_RigLogic_Data::UpdateJoints(URigHierarchy* Hierarchy, const FRigUn
 void FRigUnit_RigLogic_Data::UpdateBlendShapeCurves(URigHierarchy* InHierarchy, TArrayView<const float> BlendShapeValues)
 {
 	// set output blend shapes
-	if (BlendShapeIndices.IsValidIndex(CurrentLOD) && MorphTargetCurveIndices.IsValidIndex(CurrentLOD))
+	if (Mappings->BlendShapeIndices.IsValidIndex(CurrentLOD) && Mappings->MorphTargetCurveIndices.IsValidIndex(CurrentLOD))
 	{
-		const uint32 BlendShapePerLODCount = static_cast<uint32>(BlendShapeIndices[CurrentLOD].Values.Num());
+		const uint32 BlendShapePerLODCount = static_cast<uint32>(Mappings->BlendShapeIndices[CurrentLOD].Values.Num());
 
-		if (ensure(BlendShapePerLODCount == MorphTargetCurveIndices[CurrentLOD].Values.Num()))
+		if (ensure(BlendShapePerLODCount == Mappings->MorphTargetCurveIndices[CurrentLOD].Values.Num()))
 		{
 			for (uint32 MeshBlendIndex = 0; MeshBlendIndex < BlendShapePerLODCount; MeshBlendIndex++)
 			{
-				const int32 BlendShapeIndex = BlendShapeIndices[CurrentLOD].Values[MeshBlendIndex];
-				const int32 MorphTargetCurveIndex = MorphTargetCurveIndices[CurrentLOD].Values[MeshBlendIndex];
+				const int32 BlendShapeIndex = Mappings->BlendShapeIndices[CurrentLOD].Values[MeshBlendIndex];
+				const int32 MorphTargetCurveIndex = Mappings->MorphTargetCurveIndices[CurrentLOD].Values[MeshBlendIndex];
 				if (MorphTargetCurveIndex != INDEX_NONE)
 				{
 					const float Value = BlendShapeValues[BlendShapeIndex];
@@ -291,13 +301,13 @@ void FRigUnit_RigLogic_Data::UpdateAnimMapCurves(URigHierarchy* InHierarchy, TAr
 	// set output mask multipliers
 	// In case curves are not imported yet into CL, AnimatedMapsCurveIndices will be empty, so we need to check
 	// array bounds before trying to access it:
-	if (RigLogicIndicesForAnimMaps.IsValidIndex(CurrentLOD) && CurveElementIndicesForAnimMaps.IsValidIndex(CurrentLOD))
+	if (Mappings->RigLogicIndicesForAnimMaps.IsValidIndex(CurrentLOD) && Mappings->CurveElementIndicesForAnimMaps.IsValidIndex(CurrentLOD))
 	{
-		const uint32 AnimMapPerLODCount = RigLogicIndicesForAnimMaps[CurrentLOD].Values.Num();
+		const uint32 AnimMapPerLODCount = Mappings->RigLogicIndicesForAnimMaps[CurrentLOD].Values.Num();
 		for (uint32 AnimMapIndexForLOD = 0; AnimMapIndexForLOD < AnimMapPerLODCount; ++AnimMapIndexForLOD)
 		{
-			const int32 RigLogicAnimMapIndex = RigLogicIndicesForAnimMaps[CurrentLOD].Values[AnimMapIndexForLOD];
-			const int32 InHierarchyAnimMapIndex = CurveElementIndicesForAnimMaps[CurrentLOD].Values[AnimMapIndexForLOD];
+			const int32 RigLogicAnimMapIndex = Mappings->RigLogicIndicesForAnimMaps[CurrentLOD].Values[AnimMapIndexForLOD];
+			const int32 InHierarchyAnimMapIndex = Mappings->CurveElementIndicesForAnimMaps[CurrentLOD].Values[AnimMapIndexForLOD];
 			if (InHierarchyAnimMapIndex != INDEX_NONE)
 			{
 				const float Value = AnimMapOutputs[RigLogicAnimMapIndex];
@@ -308,7 +318,7 @@ void FRigUnit_RigLogic_Data::UpdateAnimMapCurves(URigHierarchy* InHierarchy, TAr
 	else
 	{
 		UE_LOG(LogRigLogicUnit, Warning, TEXT("Invalid LOD Index for the AnimationMaps. Ensure your curve is set up correctly!"));
-	}	
+	}
 }
 
 FSharedRigRuntimeContext* FRigUnit_RigLogic::GetSharedRigRuntimeContext(USkeletalMesh* SkelMesh)
@@ -326,62 +336,62 @@ FRigUnit_RigLogic_Execute()
 {
 	QUICK_SCOPE_CYCLE_COUNTER(STAT_RigUnit_RigLogic_Execute);
 
- 	URigHierarchy* Hierarchy = ExecuteContext.Hierarchy;
+	URigHierarchy* Hierarchy = ExecuteContext.Hierarchy;
 	if (Hierarchy)
 	{
 		switch (Context.State)
 		{
-			case EControlRigState::Init:
+		case EControlRigState::Init:
+		{
+			if (!Data.SkelMeshComponent.IsValid())
 			{
-				if (!Data.SkelMeshComponent.IsValid())
-				{
-					Data.SkelMeshComponent = Context.DataSourceRegistry->RequestSource<USkeletalMeshComponent>(UControlRig::OwnerComponent);
-					// In normal execution, Data.SkelMeshComponent will be nullptr at the beginning
-					// however, during unit testing we cannot fetch it from DataSourceRegistry 
-					// in that case, a mock version will be inserted into Data by unit test beforehand
-				}
-				if (!Data.SkelMeshComponent.IsValid() || Data.SkelMeshComponent->SkeletalMesh == nullptr)
-				{
-					return;
-				}
-				Data.CurrentLOD = Data.SkelMeshComponent->GetPredictedLODLevel();
-
-				// Fetch shared runtime context of rig from DNAAsset
-				Data.SharedRigRuntimeContext = GetSharedRigRuntimeContext(Data.SkelMeshComponent->SkeletalMesh);
-				// Context is initialized with a BehaviorReader, which can be imported into SkeletalMesh from DNA file
-				// or overwritten by GeneSplicer when making a new character
-				Data.InitializeRigLogic(Hierarchy);
-				break; 
+				Data.SkelMeshComponent = Context.DataSourceRegistry->RequestSource<USkeletalMeshComponent>(UControlRig::OwnerComponent);
+				// In normal execution, Data.SkelMeshComponent will be nullptr at the beginning
+				// however, during unit testing we cannot fetch it from DataSourceRegistry 
+				// in that case, a mock version will be inserted into Data by unit test beforehand
 			}
-			case EControlRigState::Update:
+			if (!Data.SkelMeshComponent.IsValid() || Data.SkelMeshComponent->SkeletalMesh == nullptr)
 			{
-				// Fetch shared runtime context of rig from DNAAsset
-				if (!Data.SkelMeshComponent.IsValid() || !Data.IsRigLogicInitialized() || Hierarchy == nullptr)
-				{
-					return;
-				}
-				Data.CurrentLOD = Data.SkelMeshComponent->GetPredictedLODLevel();
-				Data.CalculateRigLogic(Hierarchy);
+				return;
+			}
+			Data.CurrentLOD = Data.SkelMeshComponent->GetPredictedLODLevel();
+
+			// Fetch shared runtime context of rig from DNAAsset
+			Data.SharedRigRuntimeContext = GetSharedRigRuntimeContext(Data.SkelMeshComponent->SkeletalMesh);
+			// Context is initialized with a BehaviorReader, which can be imported into SkeletalMesh from DNA file
+			// or overwritten by GeneSplicer when making a new character
+			Data.InitializeRigLogic(Hierarchy);
+			break;
+		}
+		case EControlRigState::Update:
+		{
+			// Fetch shared runtime context of rig from DNAAsset
+			if (!Data.SkelMeshComponent.IsValid() || !Data.IsRigLogicInitialized() || Hierarchy == nullptr)
+			{
+				return;
+			}
+			Data.CurrentLOD = Data.SkelMeshComponent->GetPredictedLODLevel();
+			Data.CalculateRigLogic(Hierarchy);
 
-				//Filing a struct so we can call the same method for updating joints from tests
-				FRigUnit_RigLogic_JointUpdateParams JointUpdateParamsTemp
-				(
-					Data.SharedRigRuntimeContext->RigLogic->GetNeutralJointValues(),
-					Data.RigInstance->GetJointOutputs()
-				);
+			//Filing a struct so we can call the same method for updating joints from tests
+			FRigUnit_RigLogic_JointUpdateParams JointUpdateParamsTemp
+			(
+				Data.SharedRigRuntimeContext->RigLogic->GetNeutralJointValues(),
+				Data.RigInstance->GetJointOutputs()
+			);
 
-				Data.UpdateJoints(Hierarchy, JointUpdateParamsTemp);
+			Data.UpdateJoints(Hierarchy, JointUpdateParamsTemp);
 
-				TArrayView<const float> BlendShapeValues = Data.RigInstance->GetBlendShapeOutputs();
-				Data.UpdateBlendShapeCurves(Hierarchy, BlendShapeValues);
+			TArrayView<const float> BlendShapeValues = Data.RigInstance->GetBlendShapeOutputs();
+			Data.UpdateBlendShapeCurves(Hierarchy, BlendShapeValues);
 
-				TArrayView<const float> AnimMapOutputs = Data.RigInstance->GetAnimatedMapOutputs();
-				Data.UpdateAnimMapCurves(Hierarchy, AnimMapOutputs);
-			}
-			default:
-			{
-				break;
-			}
+			TArrayView<const float> AnimMapOutputs = Data.RigInstance->GetAnimatedMapOutputs();
+			Data.UpdateAnimMapCurves(Hierarchy, AnimMapOutputs);
+		}
+		default:
+		{
+			break;
+		}
 		}
 	}
 }
diff --git a/Engine/Plugins/Animation/RigLogic/Source/RigLogicModule/Private/RigUnit_RigLogic_Test.cpp b/Engine/Plugins/Animation/RigLogic/Source/RigLogicModule/Private/RigUnit_RigLogic_Test.cpp
--- a/Engine/Plugins/Animation/RigLogic/Source/RigLogicModule/Private/RigUnit_RigLogic_Test.cpp
+++ b/Engine/Plugins/Animation/RigLogic/Source/RigLogicModule/Private/RigUnit_RigLogic_Test.cpp
@@ -1,7 +1,7 @@
 // Copyright Epic Games, Inc. All Rights Reserved.
 
 #if WITH_DEV_AUTOMATION_TESTS
-	#include "RigUnit_RigLogic_Test.h"
+#include "RigUnit_RigLogic_Test.h"
 #endif
 
 #include "RigUnit_RigLogic.h"
@@ -55,7 +55,7 @@ TStrongObjectPtr<URigHierarchy> FRigUnit_RigLogic::TestAccessor::CreateCurveCont
 
 void FRigUnit_RigLogic::TestAccessor::Exec_MapInputCurve(URigHierarchy* TestHierarchy)
 {
-	Unit->Data.MapInputCurveIndices(TestHierarchy); //inside the method so we can access data, which is a private member
+	Unit->Data.InitializeRigLogic(TestHierarchy); //inside the method so we can access data, which is a private member
 }
 
 /** ====== Map Joints ===== **/
@@ -86,7 +86,7 @@ TSharedPtr<TestBehaviorReader> FRigUnit_RigLogic::TestAccessor::CreateBehaviorRe
 
 void FRigUnit_RigLogic::TestAccessor::Exec_MapJoints(URigHierarchy* TestHierarchy)
 {
-	Unit->Data.MapJoints(TestHierarchy);
+	Unit->Data.InitializeRigLogic(TestHierarchy); //inside the method so we can access data, which is a private member
 }
 
 /** ====== Map Morph Targets ===== **/
@@ -145,7 +145,7 @@ TStrongObjectPtr<URigHierarchy> FRigUnit_RigLogic::TestAccessor::CreateCurveCont
 void FRigUnit_RigLogic::TestAccessor::Exec_MapMorphTargets(URigHierarchy* TestCurveContainer)
 {
 	//put into a separate method so we can access the private Data member
-	Unit->Data.MapMorphTargets(TestCurveContainer); 
+	Unit->Data.InitializeRigLogic(TestCurveContainer);
 }
 
 /** ====== Map Mask Multipliers ===== **/
@@ -161,13 +161,13 @@ TSharedPtr<TestBehaviorReader> FRigUnit_RigLogic::TestAccessor::CreateBehaviorRe
 
 void FRigUnit_RigLogic::TestAccessor::Exec_MapMaskMultipliers(URigHierarchy* TestHierarchy)
 {
-	Unit->Data.MapMaskMultipliers(TestHierarchy); //inside the method so we can access data, which is a private member
+	Unit->Data.InitializeRigLogic(TestHierarchy); //inside the method so we can access data, which is a private member
 }
 
 void FRigUnit_RigLogic::TestAccessor::AddToTransformArray(float* InArray, FTransform& Transform)
 {
 	uint32 FirstAttributeIndex = 0;
-	
+
 	FVector Rotation = Transform.GetRotation().Euler();
 	InArray[FirstAttributeIndex + 0] = Rotation.X;
 	InArray[FirstAttributeIndex + 1] = Rotation.Y;
@@ -185,7 +185,7 @@ void FRigUnit_RigLogic::TestAccessor::AddToTransformArray(float* InArray, FTrans
 }
 
 
-FTransformArrayView FRigUnit_RigLogic::TestAccessor::CreateTwoJointNeutralTransforms(float *InValueArray)
+FTransformArrayView FRigUnit_RigLogic::TestAccessor::CreateTwoJointNeutralTransforms(float* InValueArray)
 {
 	float* Transform1Ptr = InValueArray;
 
@@ -229,7 +229,7 @@ TArrayView<const uint16> FRigUnit_RigLogic::TestAccessor::CreateTwoJointVariable
 		InVariableAttributeIndices[9 + 6] = 9 + 6;
 		InVariableAttributeIndices[9 + 7] = 9 + 7;
 		InVariableAttributeIndices[9 + 8] = 9 + 8;
-	
+
 		return TArrayView<const uint16>(InVariableAttributeIndices, 2 * FRigUnit_RigLogic::TestAccessor::MAX_ATTRS_PER_JOINT);
 	}
 
@@ -284,18 +284,18 @@ IMPLEMENT_RIGUNIT_AUTOMATION_TEST(FRigUnit_RigLogic)
 		Test.GetData()->RigInstance != nullptr &&
 
 		// Check joints
-		Test.GetData()->HierarchyBoneIndices.Num() > 0 &&
+		Test.GetData()->Mappings->HierarchyBoneIndices.Num() > 0 &&
 
 		// Check input curves
-		Test.GetData()->InputCurveIndices.Num() > 0 &&
+		Test.GetData()->Mappings->InputCurveIndices.Num() > 0 &&
 
 		// Check morph targets
-		Test.GetData()->MorphTargetCurveIndices.Num() > 0 &&
-		Test.GetData()->BlendShapeIndices.Num() > 0 &&
+		Test.GetData()->Mappings->MorphTargetCurveIndices.Num() > 0 &&
+		Test.GetData()->Mappings->BlendShapeIndices.Num() > 0 &&
 
 		// Check mask multipliers
-		Test.GetData()->RigLogicIndicesForAnimMaps.Num() > 0 &&
-		Test.GetData()->CurveElementIndicesForAnimMaps.Num() > 0,
+		Test.GetData()->Mappings->RigLogicIndicesForAnimMaps.Num() > 0 &&
+		Test.GetData()->Mappings->CurveElementIndicesForAnimMaps.Num() > 0,
 
 		TEXT("InitAndExecute failed to initialize rig logic.")
 	);
@@ -312,8 +312,8 @@ IMPLEMENT_RIGUNIT_AUTOMATION_TEST(FRigUnit_RigLogic)
 	Test.Exec_MapInputCurve(TestCurveContainerNameMismatch.Get());
 	//Assert
 	AddErrorIfFalse(
-		Test.GetData()->InputCurveIndices.Num() == 1 &&
-		Test.GetData()->InputCurveIndices[0] == INDEX_NONE,
+		Test.GetData()->Mappings->InputCurveIndices.Num() == 1 &&
+		Test.GetData()->Mappings->InputCurveIndices[0] == INDEX_NONE,
 		TEXT("MapInputCurve - ValidReader CurveContainerWithNameMismatch")
 	);
 
@@ -327,7 +327,7 @@ IMPLEMENT_RIGUNIT_AUTOMATION_TEST(FRigUnit_RigLogic)
 	Test.Exec_MapInputCurve(TestCurveContainerValid.Get());
 	//Assert
 	AddErrorIfFalse(
-		Test.GetData()->InputCurveIndices.Num() == 0,
+		Test.GetData()->Mappings->InputCurveIndices.Num() == 0,
 		TEXT("MapInputCurve - EmptyReader ValidCurveContainer")
 	);
 
@@ -340,8 +340,8 @@ IMPLEMENT_RIGUNIT_AUTOMATION_TEST(FRigUnit_RigLogic)
 	Test.Exec_MapInputCurve(TestCurveContainerEmpty.Get());
 	//Assert
 	AddErrorIfFalse(
-		Test.GetData()->InputCurveIndices.Num() == 1 &&
-		Test.GetData()->InputCurveIndices[0] == INDEX_NONE,
+		Test.GetData()->Mappings->InputCurveIndices.Num() == 1 &&
+		Test.GetData()->Mappings->InputCurveIndices[0] == INDEX_NONE,
 		TEXT("MapInputCurve - ValidReader EmptyCurveContainer")
 	);
 
@@ -355,7 +355,7 @@ IMPLEMENT_RIGUNIT_AUTOMATION_TEST(FRigUnit_RigLogic)
 	Test.Exec_MapInputCurve(TestCurveContainerValid.Get());
 	//Assert
 	AddErrorIfFalse(
-		Test.GetData()->InputCurveIndices.Num() == 0,
+		Test.GetData()->Mappings->InputCurveIndices.Num() == 0,
 		TEXT("MapInputCurve - InvalidReader ValidCurveContainer")
 	);
 
@@ -368,8 +368,8 @@ IMPLEMENT_RIGUNIT_AUTOMATION_TEST(FRigUnit_RigLogic)
 	Test.Exec_MapInputCurve(TestCurveContainerValid.Get());
 	//Assert
 	AddErrorIfFalse(
-		Test.GetData()->InputCurveIndices.Num() == 1 &&
-		Test.GetData()->InputCurveIndices[0] == 0,
+		Test.GetData()->Mappings->InputCurveIndices.Num() == 1 &&
+		Test.GetData()->Mappings->InputCurveIndices[0] == 0,
 		TEXT("MapInputCurve - Valid Inputs")
 	);
 
@@ -384,7 +384,7 @@ IMPLEMENT_RIGUNIT_AUTOMATION_TEST(FRigUnit_RigLogic)
 	Test.Exec_MapJoints(TestHierarchyEmpty.Get());
 	//Assert
 	AddErrorIfFalse(
-		Test.GetData()->HierarchyBoneIndices.Num() == 0,
+		Test.GetData()->Mappings->HierarchyBoneIndices.Num() == 0,
 		TEXT("MapJoints - Empty Inputs")
 	);
 
@@ -395,7 +395,7 @@ IMPLEMENT_RIGUNIT_AUTOMATION_TEST(FRigUnit_RigLogic)
 	//Test
 	Test.Exec_MapJoints(TestHierarchyTwoBones.Get());
 	AddErrorIfFalse(
-		Test.GetData()->HierarchyBoneIndices.Num() == 0,
+		Test.GetData()->Mappings->HierarchyBoneIndices.Num() == 0,
 		TEXT("MapJoints - EmptyReader TwoBones")
 	);
 
@@ -406,17 +406,17 @@ IMPLEMENT_RIGUNIT_AUTOMATION_TEST(FRigUnit_RigLogic)
 	//Test
 	Test.Exec_MapJoints(TestHierarchyEmpty.Get());
 	AddErrorIfFalse(
-		Test.GetData()->HierarchyBoneIndices.Num() == 2,
+		Test.GetData()->Mappings->HierarchyBoneIndices.Num() == 2,
 		TEXT("MapJoints - TwoJoints NoBones - expected 2 bone indices")
 	);
 	AddErrorIfFalse(
-		Test.GetData()->HierarchyBoneIndices.Num() == 2 && //repeated condition to prevent crash
-		Test.GetData()->HierarchyBoneIndices[0] == INDEX_NONE,
+		Test.GetData()->Mappings->HierarchyBoneIndices.Num() == 2 && //repeated condition to prevent crash
+		Test.GetData()->Mappings->HierarchyBoneIndices[0] == INDEX_NONE,
 		TEXT("MapJoints - TwoJoints NoBones - Expected joint 0 index to be NONE")
 	);
 	AddErrorIfFalse(
-		Test.GetData()->HierarchyBoneIndices.Num() == 2 && //repeated condition to prevent crash
-		Test.GetData()->HierarchyBoneIndices[1] == INDEX_NONE,
+		Test.GetData()->Mappings->HierarchyBoneIndices.Num() == 2 && //repeated condition to prevent crash
+		Test.GetData()->Mappings->HierarchyBoneIndices[1] == INDEX_NONE,
 		TEXT("MapJoints - TwoJoints NoBones - Expected joint 1 index to be NONE")
 	);
 
@@ -428,17 +428,17 @@ IMPLEMENT_RIGUNIT_AUTOMATION_TEST(FRigUnit_RigLogic)
 	Test.Exec_MapJoints(TestHierarchyTwoBones.Get());
 	//Assert
 	AddErrorIfFalse(
-		Test.GetData()->HierarchyBoneIndices.Num() == 2,
+		Test.GetData()->Mappings->HierarchyBoneIndices.Num() == 2,
 		TEXT("MapJoints - TwoJoints TwoBones - Expected 2 bone indices")
 	);
 	AddErrorIfFalse(
-		Test.GetData()->HierarchyBoneIndices.Num() == 2 && //prevent crash
-		Test.GetData()->HierarchyBoneIndices[0] == 0,
+		Test.GetData()->Mappings->HierarchyBoneIndices.Num() == 2 && //prevent crash
+		Test.GetData()->Mappings->HierarchyBoneIndices[0] == 0,
 		TEXT("MapJoints - TwoJoints TwoBones - Expected bone 0 index to be 0")
 	);
 	AddErrorIfFalse(
-		Test.GetData()->HierarchyBoneIndices.Num() == 2 && //prevent crash
-		Test.GetData()->HierarchyBoneIndices[1] == 1,
+		Test.GetData()->Mappings->HierarchyBoneIndices.Num() == 2 && //prevent crash
+		Test.GetData()->Mappings->HierarchyBoneIndices[1] == 1,
 		TEXT("MapJoints - TwoJoints TwoBones - Expected bone index 1 index to be 1")
 	);
 
@@ -455,12 +455,12 @@ IMPLEMENT_RIGUNIT_AUTOMATION_TEST(FRigUnit_RigLogic)
 	Test.Exec_MapMorphTargets(TestMorphTargetNameMismatch.Get());
 	//Assert
 	AddErrorIfFalse(
-		Test.GetData()->BlendShapeIndices.Num() == 1 && //LOD 0
-		Test.GetData()->MorphTargetCurveIndices.Num() == 1 && //LOD 0
-		Test.GetData()->BlendShapeIndices[0].Values.Num() == 1 && //at least one blendshape exists
-		Test.GetData()->BlendShapeIndices[0].Values[0] == 0 && //has index 0
-		Test.GetData()->MorphTargetCurveIndices[0].Values.Num() == 1 && //but morph target corresponding to that blendshape
-		Test.GetData()->MorphTargetCurveIndices[0].Values[0] == INDEX_NONE, //wasn't found
+		Test.GetData()->Mappings->BlendShapeIndices.Num() == 1 && //LOD 0
+		Test.GetData()->Mappings->MorphTargetCurveIndices.Num() == 1 && //LOD 0
+		Test.GetData()->Mappings->BlendShapeIndices[0].Values.Num() == 1 && //at least one blendshape exists
+		Test.GetData()->Mappings->BlendShapeIndices[0].Values[0] == 0 && //has index 0
+		Test.GetData()->Mappings->MorphTargetCurveIndices[0].Values.Num() == 1 && //but morph target corresponding to that blendshape
+		Test.GetData()->Mappings->MorphTargetCurveIndices[0].Values[0] == INDEX_NONE, //wasn't found
 		TEXT("MapMorphTargets - ValidReader MorphTargetWithNameMismatch")
 	);
 
@@ -474,8 +474,8 @@ IMPLEMENT_RIGUNIT_AUTOMATION_TEST(FRigUnit_RigLogic)
 	Test.Exec_MapMorphTargets(TestMorphTargetCurveValid.Get());
 	//Assert
 	AddErrorIfFalse(
-		Test.GetData()->MorphTargetCurveIndices.Num() == 0 &&
-		Test.GetData()->BlendShapeIndices.Num() == 0,
+		Test.GetData()->Mappings->MorphTargetCurveIndices.Num() == 0 &&
+		Test.GetData()->Mappings->BlendShapeIndices.Num() == 0,
 		TEXT("MapMorphTargets - EmptyReader ValidMorphTargetCurve")
 	);
 
@@ -489,9 +489,9 @@ IMPLEMENT_RIGUNIT_AUTOMATION_TEST(FRigUnit_RigLogic)
 	Test.Exec_MapMorphTargets(TestMorphTargetCurveValid.Get());
 	//Assert
 	AddErrorIfFalse(
-		Test.GetData()->BlendShapeIndices.Num() == 1 && //LOD 0 exists
-		Test.GetData()->BlendShapeIndices[0].Values.Num() == 0 && //but no blend shapes mapped
-		Test.GetData()->MorphTargetCurveIndices[0].Values.Num() == 0, //or morph targets
+		Test.GetData()->Mappings->BlendShapeIndices.Num() == 1 && //LOD 0 exists
+		Test.GetData()->Mappings->BlendShapeIndices[0].Values.Num() == 0 && //but no blend shapes mapped
+		Test.GetData()->Mappings->MorphTargetCurveIndices[0].Values.Num() == 0, //or morph targets
 		TEXT("MapMorphTargets - NoBlendShapes ValidMorphTargetCurve")
 	);
 
@@ -504,11 +504,11 @@ IMPLEMENT_RIGUNIT_AUTOMATION_TEST(FRigUnit_RigLogic)
 	Test.Exec_MapMorphTargets(TestCurveContainerEmpty.Get());
 	//Assert
 	AddErrorIfFalse(
-		Test.GetData()->BlendShapeIndices.Num() == 1 && //LOD 0
-		Test.GetData()->BlendShapeIndices[0].Values.Num() == 1 && //one blend shape
-		Test.GetData()->BlendShapeIndices[0].Values[0] == 0 && //of index 0
-		Test.GetData()->MorphTargetCurveIndices[0].Values.Num() == 1 && //we put in a morph target index corresponding to it
-		Test.GetData()->MorphTargetCurveIndices[0].Values[0] == INDEX_NONE, //but just to signal that it is not found
+		Test.GetData()->Mappings->BlendShapeIndices.Num() == 1 && //LOD 0
+		Test.GetData()->Mappings->BlendShapeIndices[0].Values.Num() == 1 && //one blend shape
+		Test.GetData()->Mappings->BlendShapeIndices[0].Values[0] == 0 && //of index 0
+		Test.GetData()->Mappings->MorphTargetCurveIndices[0].Values.Num() == 1 && //we put in a morph target index corresponding to it
+		Test.GetData()->Mappings->MorphTargetCurveIndices[0].Values[0] == INDEX_NONE, //but just to signal that it is not found
 		TEXT("MapMorphTargets - ValidReader EmptyCurveContainer")
 	);
 
@@ -521,11 +521,11 @@ IMPLEMENT_RIGUNIT_AUTOMATION_TEST(FRigUnit_RigLogic)
 	Test.Exec_MapMorphTargets(TestMorphTargetCurveValid.Get());
 	//Assert
 	AddErrorIfFalse(
-		Test.GetData()->BlendShapeIndices.Num() == 1 && //LOD 0
-		Test.GetData()->BlendShapeIndices[0].Values.Num() == 1 && //one blend shape (empty named)
-		Test.GetData()->BlendShapeIndices[0].Values[0] == 0 && //of index 0
-		Test.GetData()->MorphTargetCurveIndices[0].Values.Num() == 1 && //we put in a morph target index corresponding to it
-		Test.GetData()->MorphTargetCurveIndices[0].Values[0] == INDEX_NONE, //but just to signal that it is not found
+		Test.GetData()->Mappings->BlendShapeIndices.Num() == 1 && //LOD 0
+		Test.GetData()->Mappings->BlendShapeIndices[0].Values.Num() == 1 && //one blend shape (empty named)
+		Test.GetData()->Mappings->BlendShapeIndices[0].Values[0] == 0 && //of index 0
+		Test.GetData()->Mappings->MorphTargetCurveIndices[0].Values.Num() == 1 && //we put in a morph target index corresponding to it
+		Test.GetData()->Mappings->MorphTargetCurveIndices[0].Values[0] == INDEX_NONE, //but just to signal that it is not found
 		TEXT("MapMorphTargets - InvalidReader ValidMorphTargetCurve")
 	);
 
@@ -538,11 +538,11 @@ IMPLEMENT_RIGUNIT_AUTOMATION_TEST(FRigUnit_RigLogic)
 	Test.Exec_MapMorphTargets(TestMorphTargetCurvesInvalid.Get());
 	//Assert
 	AddErrorIfFalse(
-		Test.GetData()->BlendShapeIndices.Num() == 1 && //LOD 0
-		Test.GetData()->BlendShapeIndices[0].Values.Num() == 1 && //one blend shape
-		Test.GetData()->BlendShapeIndices[0].Values[0] == 0 && //of index 0
-		Test.GetData()->MorphTargetCurveIndices[0].Values.Num() == 1 && //we put in a morph target index corresponding to it
-		Test.GetData()->MorphTargetCurveIndices[0].Values[0] == INDEX_NONE, //but just to signal that it is not found
+		Test.GetData()->Mappings->BlendShapeIndices.Num() == 1 && //LOD 0
+		Test.GetData()->Mappings->BlendShapeIndices[0].Values.Num() == 1 && //one blend shape
+		Test.GetData()->Mappings->BlendShapeIndices[0].Values[0] == 0 && //of index 0
+		Test.GetData()->Mappings->MorphTargetCurveIndices[0].Values.Num() == 1 && //we put in a morph target index corresponding to it
+		Test.GetData()->Mappings->MorphTargetCurveIndices[0].Values[0] == INDEX_NONE, //but just to signal that it is not found
 		TEXT("MapMorphTargets - ValidReader InvalidMorphTargetCurve")
 	);
 
@@ -555,12 +555,12 @@ IMPLEMENT_RIGUNIT_AUTOMATION_TEST(FRigUnit_RigLogic)
 	Test.Exec_MapMorphTargets(TestMorphTargetCurveValid.Get());
 	//Assert
 	AddErrorIfFalse(
-		Test.GetData()->BlendShapeIndices.Num() == 1 && //LOD 0
-		Test.GetData()->MorphTargetCurveIndices.Num() == 1 && //LOD 0
-		Test.GetData()->BlendShapeIndices[0].Values.Num() == 1 && //at least one blendshape exists
-		Test.GetData()->BlendShapeIndices[0].Values[0] == 0 && //has index 0
-		Test.GetData()->MorphTargetCurveIndices[0].Values.Num() == 1 && //morph target corresponding to that blendshape exists
-		Test.GetData()->MorphTargetCurveIndices[0].Values[0] == 0, //and actually points to the right index
+		Test.GetData()->Mappings->BlendShapeIndices.Num() == 1 && //LOD 0
+		Test.GetData()->Mappings->MorphTargetCurveIndices.Num() == 1 && //LOD 0
+		Test.GetData()->Mappings->BlendShapeIndices[0].Values.Num() == 1 && //at least one blendshape exists
+		Test.GetData()->Mappings->BlendShapeIndices[0].Values[0] == 0 && //has index 0
+		Test.GetData()->Mappings->MorphTargetCurveIndices[0].Values.Num() == 1 && //morph target corresponding to that blendshape exists
+		Test.GetData()->Mappings->MorphTargetCurveIndices[0].Values[0] == 0, //and actually points to the right index
 		TEXT("MapMorphTargets - ValidReader ValidTestMorphTarget")
 	);
 
@@ -583,28 +583,28 @@ IMPLEMENT_RIGUNIT_AUTOMATION_TEST(FRigUnit_RigLogic)
 	Test.Exec_MapMorphTargets(TestMorphTargetTwoCurves.Get());
 	//Assert
 	AddErrorIfFalse(
-		Test.GetData()->BlendShapeIndices.Num() == 2 &&  //2 LODs
-		Test.GetData()->MorphTargetCurveIndices.Num() == 2,
+		Test.GetData()->Mappings->BlendShapeIndices.Num() == 2 &&  //2 LODs
+		Test.GetData()->Mappings->MorphTargetCurveIndices.Num() == 2,
 		TEXT("MapMorphTargets LOD0(AB) LOD1(A) - Expected 2 LODs for both blendshapes and morph targets")
 	);
 
 	AddErrorIfFalse(
-		Test.GetData()->BlendShapeIndices.Num() == 2 &&         //condition repeated for crash prevention
-		Test.GetData()->BlendShapeIndices[0].Values.Num() == 2 && //two blendshapes at LOD 0
-		Test.GetData()->BlendShapeIndices[0].Values[0] == 0 &&  // A
-		Test.GetData()->BlendShapeIndices[0].Values[1] == 1 &&  // B
-		Test.GetData()->BlendShapeIndices[1].Values.Num() == 1 && //one blendshape at LOD 1
-		Test.GetData()->BlendShapeIndices[1].Values[0] == 0,    // A
-			TEXT("MapMorphTargets LOD0(AB) LOD1(A) - resulting blendshape indices not correct")
+		Test.GetData()->Mappings->BlendShapeIndices.Num() == 2 &&         //condition repeated for crash prevention
+		Test.GetData()->Mappings->BlendShapeIndices[0].Values.Num() == 2 && //two blendshapes at LOD 0
+		Test.GetData()->Mappings->BlendShapeIndices[0].Values[0] == 0 &&  // A
+		Test.GetData()->Mappings->BlendShapeIndices[0].Values[1] == 1 &&  // B
+		Test.GetData()->Mappings->BlendShapeIndices[1].Values.Num() == 1 && //one blendshape at LOD 1
+		Test.GetData()->Mappings->BlendShapeIndices[1].Values[0] == 0,    // A
+		TEXT("MapMorphTargets LOD0(AB) LOD1(A) - resulting blendshape indices not correct")
 	);
 
 	AddErrorIfFalse(
-		Test.GetData()->MorphTargetCurveIndices.Num() == 2 &&         //condition repeated for crash prevention
-		Test.GetData()->MorphTargetCurveIndices[0].Values.Num() == 2 && //two morph targets at LOD 0
-		Test.GetData()->MorphTargetCurveIndices[0].Values[0] == 0 &&  // A
-		Test.GetData()->MorphTargetCurveIndices[0].Values[1] == 1 &&  // B
-		Test.GetData()->MorphTargetCurveIndices[1].Values.Num() == 1 && //one morph target at LOD 1
-		Test.GetData()->MorphTargetCurveIndices[1].Values[0] == 0,    // A
+		Test.GetData()->Mappings->MorphTargetCurveIndices.Num() == 2 &&         //condition repeated for crash prevention
+		Test.GetData()->Mappings->MorphTargetCurveIndices[0].Values.Num() == 2 && //two morph targets at LOD 0
+		Test.GetData()->Mappings->MorphTargetCurveIndices[0].Values[0] == 0 &&  // A
+		Test.GetData()->Mappings->MorphTargetCurveIndices[0].Values[1] == 1 &&  // B
+		Test.GetData()->Mappings->MorphTargetCurveIndices[1].Values.Num() == 1 && //one morph target at LOD 1
+		Test.GetData()->Mappings->MorphTargetCurveIndices[1].Values[0] == 0,    // A
 		TEXT("MapMorphTargets LOD0(AB) LOD1(A) - resulting morph target indices not correct")
 	);
 
@@ -621,25 +621,25 @@ IMPLEMENT_RIGUNIT_AUTOMATION_TEST(FRigUnit_RigLogic)
 	Test.Exec_MapMorphTargets(TestMorphTargetTwoCurves.Get());
 	//Assert
 	AddErrorIfFalse(
-		Test.GetData()->BlendShapeIndices.Num() == 1 &&  //1 LOD
-		Test.GetData()->MorphTargetCurveIndices.Num() == 1, //1 LOD
+		Test.GetData()->Mappings->BlendShapeIndices.Num() == 1 &&  //1 LOD
+		Test.GetData()->Mappings->MorphTargetCurveIndices.Num() == 1, //1 LOD
 		TEXT("MapMorphTargets LOD0(AB) LOD1(-) - Expected 1 LOD for both blendshapes and morph targets")
 	);
 
 
 	AddErrorIfFalse(
-		Test.GetData()->BlendShapeIndices.Num() == 1 &&         //condition repeated for crash prevention
-		Test.GetData()->BlendShapeIndices[0].Values.Num() == 2 && //two blendshapes at LOD 0
-		Test.GetData()->BlendShapeIndices[0].Values[0] == 0 &&  // A
-		Test.GetData()->BlendShapeIndices[0].Values[1] == 1,    // B
+		Test.GetData()->Mappings->BlendShapeIndices.Num() == 1 &&         //condition repeated for crash prevention
+		Test.GetData()->Mappings->BlendShapeIndices[0].Values.Num() == 2 && //two blendshapes at LOD 0
+		Test.GetData()->Mappings->BlendShapeIndices[0].Values[0] == 0 &&  // A
+		Test.GetData()->Mappings->BlendShapeIndices[0].Values[1] == 1,    // B
 		TEXT("MapMorphTargets LOD0(AB) LOD1(-) - Resulting blendshapes not correct")
 	);
 
 	AddErrorIfFalse(
-		Test.GetData()->MorphTargetCurveIndices.Num() == 1 &&         //condition repeated for crash prevention
-		Test.GetData()->MorphTargetCurveIndices[0].Values.Num() == 2 && //two morph targets at LOD 0
-		Test.GetData()->MorphTargetCurveIndices[0].Values[0] == 0 &&  // A
-		Test.GetData()->MorphTargetCurveIndices[0].Values[1] == 1,    // B
+		Test.GetData()->Mappings->MorphTargetCurveIndices.Num() == 1 &&         //condition repeated for crash prevention
+		Test.GetData()->Mappings->MorphTargetCurveIndices[0].Values.Num() == 2 && //two morph targets at LOD 0
+		Test.GetData()->Mappings->MorphTargetCurveIndices[0].Values[0] == 0 &&  // A
+		Test.GetData()->Mappings->MorphTargetCurveIndices[0].Values[1] == 1,    // B
 		TEXT("MapMorphTargets LOD0(AB) LOD1(-) - Resulting morph targets not correct")
 	);
 
@@ -655,26 +655,26 @@ IMPLEMENT_RIGUNIT_AUTOMATION_TEST(FRigUnit_RigLogic)
 	Test.Exec_MapMorphTargets(TestMorphTargetTwoCurves.Get());
 	//Assert
 	AddErrorIfFalse(
-		Test.GetData()->BlendShapeIndices.Num() == 2 &&  //2 LODs
-		Test.GetData()->MorphTargetCurveIndices.Num() == 2, //2 LODs
+		Test.GetData()->Mappings->BlendShapeIndices.Num() == 2 &&  //2 LODs
+		Test.GetData()->Mappings->MorphTargetCurveIndices.Num() == 2, //2 LODs
 		TEXT("MapMorphTargets LOD0(A) LOD1(B) - Expected 2 LODs for both blendshapes and morph targets")
 	);
 
 	AddErrorIfFalse(
-		Test.GetData()->BlendShapeIndices.Num() == 2 &&         //condition repeated for crash prevention
-		Test.GetData()->BlendShapeIndices[0].Values.Num() == 1 && //1 blendshape at LOD 0
-		Test.GetData()->BlendShapeIndices[0].Values[0] == 0 &&  // A
-		Test.GetData()->BlendShapeIndices[1].Values.Num() == 1 && //1 blendshape at LOD 1
-		Test.GetData()->BlendShapeIndices[1].Values[0] == 1,    // B
+		Test.GetData()->Mappings->BlendShapeIndices.Num() == 2 &&         //condition repeated for crash prevention
+		Test.GetData()->Mappings->BlendShapeIndices[0].Values.Num() == 1 && //1 blendshape at LOD 0
+		Test.GetData()->Mappings->BlendShapeIndices[0].Values[0] == 0 &&  // A
+		Test.GetData()->Mappings->BlendShapeIndices[1].Values.Num() == 1 && //1 blendshape at LOD 1
+		Test.GetData()->Mappings->BlendShapeIndices[1].Values[0] == 1,    // B
 		TEXT("MapMorphTargets LOD0(A) LOD1(B) - Resulting blendshape indices not correct")
 	);
 
 	AddErrorIfFalse(
-		Test.GetData()->MorphTargetCurveIndices.Num() == 2 &&         //condition repeated for crash prevention
-		Test.GetData()->MorphTargetCurveIndices[0].Values.Num() == 1 && //1 morph target at LOD 0
-		Test.GetData()->MorphTargetCurveIndices[0].Values[0] == 0 &&  // A
-		Test.GetData()->MorphTargetCurveIndices[1].Values.Num() == 1 && //1 morph target at LOD 1
-		Test.GetData()->MorphTargetCurveIndices[1].Values[0] == 1,    // B
+		Test.GetData()->Mappings->MorphTargetCurveIndices.Num() == 2 &&         //condition repeated for crash prevention
+		Test.GetData()->Mappings->MorphTargetCurveIndices[0].Values.Num() == 1 && //1 morph target at LOD 0
+		Test.GetData()->Mappings->MorphTargetCurveIndices[0].Values[0] == 0 &&  // A
+		Test.GetData()->Mappings->MorphTargetCurveIndices[1].Values.Num() == 1 && //1 morph target at LOD 1
+		Test.GetData()->Mappings->MorphTargetCurveIndices[1].Values[0] == 1,    // B
 		TEXT("MapMorphTargets LOD0(A) LOD1(B) - Resulting morph target indices not correct")
 	);
 
@@ -691,12 +691,12 @@ IMPLEMENT_RIGUNIT_AUTOMATION_TEST(FRigUnit_RigLogic)
 	Test.Exec_MapMaskMultipliers(TestCurveContainerForAnimMapsNameMismatch.Get());
 	//Assert
 	AddErrorIfFalse(
-		Test.GetData()->CurveElementIndicesForAnimMaps.Num() == 1 &&
-		Test.GetData()->CurveElementIndicesForAnimMaps[0].Values.Num() == 1 &&
-		Test.GetData()->CurveElementIndicesForAnimMaps[0].Values[0] == INDEX_NONE &&
-		Test.GetData()->RigLogicIndicesForAnimMaps.Num() == 1 &&
-		Test.GetData()->RigLogicIndicesForAnimMaps[0].Values.Num() == 1 &&
-		Test.GetData()->RigLogicIndicesForAnimMaps[0].Values[0] == 0,
+		Test.GetData()->Mappings->CurveElementIndicesForAnimMaps.Num() == 1 &&
+		Test.GetData()->Mappings->CurveElementIndicesForAnimMaps[0].Values.Num() == 1 &&
+		Test.GetData()->Mappings->CurveElementIndicesForAnimMaps[0].Values[0] == INDEX_NONE &&
+		Test.GetData()->Mappings->RigLogicIndicesForAnimMaps.Num() == 1 &&
+		Test.GetData()->Mappings->RigLogicIndicesForAnimMaps[0].Values.Num() == 1 &&
+		Test.GetData()->Mappings->RigLogicIndicesForAnimMaps[0].Values[0] == 0,
 		TEXT("MapMaskMultipliers - ValidReader ValidAnimatedMapNameMismatch")
 	);
 
@@ -709,8 +709,8 @@ IMPLEMENT_RIGUNIT_AUTOMATION_TEST(FRigUnit_RigLogic)
 	Test.Exec_MapMaskMultipliers(TestCurveContainerForAnimMapsValid.Get());
 	//Assert
 	AddErrorIfFalse(
-		Test.GetData()->CurveElementIndicesForAnimMaps.Num() == 0 &&
-		Test.GetData()->RigLogicIndicesForAnimMaps.Num() == 0,
+		Test.GetData()->Mappings->CurveElementIndicesForAnimMaps.Num() == 0 &&
+		Test.GetData()->Mappings->RigLogicIndicesForAnimMaps.Num() == 0,
 		TEXT("MapMaskMultipliers - EmptyReader ValidAnimatedMap")
 	);
 
@@ -721,12 +721,12 @@ IMPLEMENT_RIGUNIT_AUTOMATION_TEST(FRigUnit_RigLogic)
 	Test.Exec_MapMaskMultipliers(TestCurveContainerEmpty.Get());
 	//Assert
 	AddErrorIfFalse(
-		Test.GetData()->CurveElementIndicesForAnimMaps.Num() == 1 &&
-		Test.GetData()->CurveElementIndicesForAnimMaps[0].Values.Num() == 1 &&
-		Test.GetData()->CurveElementIndicesForAnimMaps[0].Values[0] == INDEX_NONE &&
-		Test.GetData()->RigLogicIndicesForAnimMaps.Num() == 1 &&
-		Test.GetData()->RigLogicIndicesForAnimMaps[0].Values.Num() == 1 &&
-		Test.GetData()->RigLogicIndicesForAnimMaps[0].Values[0] == 0,
+		Test.GetData()->Mappings->CurveElementIndicesForAnimMaps.Num() == 1 &&
+		Test.GetData()->Mappings->CurveElementIndicesForAnimMaps[0].Values.Num() == 1 &&
+		Test.GetData()->Mappings->CurveElementIndicesForAnimMaps[0].Values[0] == INDEX_NONE &&
+		Test.GetData()->Mappings->RigLogicIndicesForAnimMaps.Num() == 1 &&
+		Test.GetData()->Mappings->RigLogicIndicesForAnimMaps[0].Values.Num() == 1 &&
+		Test.GetData()->Mappings->RigLogicIndicesForAnimMaps[0].Values[0] == 0,
 		TEXT("MapMaskMultipliers - ValidReader EmptyCurveContainer")
 	);
 
@@ -739,12 +739,12 @@ IMPLEMENT_RIGUNIT_AUTOMATION_TEST(FRigUnit_RigLogic)
 	Test.Exec_MapMaskMultipliers(TestCurveContainerForAnimMapsValid.Get());
 	//Assert
 	AddErrorIfFalse(
-		Test.GetData()->CurveElementIndicesForAnimMaps.Num() == 1 &&
-		Test.GetData()->CurveElementIndicesForAnimMaps[0].Values.Num() == 1 &&
-		Test.GetData()->CurveElementIndicesForAnimMaps[0].Values[0] == 0 &&
-		Test.GetData()->RigLogicIndicesForAnimMaps.Num() == 1 &&
-		Test.GetData()->RigLogicIndicesForAnimMaps[0].Values.Num() == 1 &&
-		Test.GetData()->RigLogicIndicesForAnimMaps[0].Values[0] == 0,
+		Test.GetData()->Mappings->CurveElementIndicesForAnimMaps.Num() == 1 &&
+		Test.GetData()->Mappings->CurveElementIndicesForAnimMaps[0].Values.Num() == 1 &&
+		Test.GetData()->Mappings->CurveElementIndicesForAnimMaps[0].Values[0] == 0 &&
+		Test.GetData()->Mappings->RigLogicIndicesForAnimMaps.Num() == 1 &&
+		Test.GetData()->Mappings->RigLogicIndicesForAnimMaps[0].Values.Num() == 1 &&
+		Test.GetData()->Mappings->RigLogicIndicesForAnimMaps[0].Values[0] == 0,
 		TEXT("MapMaskMultipliers - Valid Inputs")
 	);
 
diff --git a/Engine/Plugins/Animation/RigLogic/Source/RigLogicModule/Public/DNAAsset.h b/Engine/Plugins/Animation/RigLogic/Source/RigLogicModule/Public/DNAAsset.h
--- a/Engine/Plugins/Animation/RigLogic/Source/RigLogicModule/Public/DNAAsset.h
+++ b/Engine/Plugins/Animation/RigLogic/Source/RigLogicModule/Public/DNAAsset.h
@@ -20,7 +20,7 @@ class IBehaviorReader;
 class IGeometryReader;
 class FRigLogicMemoryStream;
 class UAssetUserData;
-enum class EDNADataLayer: uint8;
+enum class EDNADataLayer : uint8;
 
 struct FSharedRigRuntimeContext
 {
@@ -30,6 +30,8 @@ struct FSharedRigRuntimeContext
 		TArray<T> Values;
 	};
 
+	uint32 DNAHash;
+
 	/** Part of the .dna file needed for run-time execution of RigLogic;
 	 **/
 	TSharedPtr<IBehaviorReader> BehaviorReader;
@@ -48,11 +50,11 @@ struct FSharedRigRuntimeContext
 };
 
 
- /** An asset holding the data needed to generate/update/animate a RigLogic character
-  * It is imported from character's DNA file as a bit stream, and separated out it into runtime (behavior) and design-time chunks;
-  * Currently, the design-time part still loads the geometry, as it is needed for the skeletal mesh update; once SkeletalMeshDNAReader is
-  * fully implemented, it will be able to read the geometry directly from the SkeletalMesh and won't load it into this asset 
-  **/
+/** An asset holding the data needed to generate/update/animate a RigLogic character
+ * It is imported from character's DNA file as a bit stream, and separated out it into runtime (behavior) and design-time chunks;
+ * Currently, the design-time part still loads the geometry, as it is needed for the skeletal mesh update; once SkeletalMeshDNAReader is
+ * fully implemented, it will be able to read the geometry directly from the SkeletalMesh and won't load it into this asset
+ **/
 UCLASS(NotBlueprintable, hidecategories = (Object))
 class RIGLOGICMODULE_API UDNAAsset : public UAssetUserData
 {
@@ -63,7 +65,7 @@ public:
 
 #if WITH_EDITORONLY_DATA
 	UPROPERTY(VisibleAnywhere, Instanced, Category = ImportSettings)
-	class UAssetImportData* AssetImportData;
+		class UAssetImportData* AssetImportData;
 #endif
 
 	TSharedPtr<IBehaviorReader> GetBehaviorReader()
@@ -79,7 +81,7 @@ public:
 #endif
 
 	UPROPERTY()
-	FString DNAFileName; 
+		FString DNAFileName;
 
 	bool Init(const FString& Filename);
 	void Serialize(FArchive& Ar) override;
diff --git a/Engine/Plugins/Animation/RigLogic/Source/RigLogicModule/Public/RigUnit_RigLogic.h b/Engine/Plugins/Animation/RigLogic/Source/RigLogicModule/Public/RigUnit_RigLogic.h
--- a/Engine/Plugins/Animation/RigLogic/Source/RigLogicModule/Public/RigUnit_RigLogic.h
+++ b/Engine/Plugins/Animation/RigLogic/Source/RigLogicModule/Public/RigUnit_RigLogic.h
@@ -7,6 +7,7 @@
 
 #include "ControlRig.h"
 #include "DNAAsset.h"
+#include "Containers/DiscardableKeyValueCache.h"
 #include "Materials/Material.h"
 #include "Units/RigUnit.h"
 
@@ -47,28 +48,15 @@ public:
 	FTransformArrayView DeltaTransforms; //the result of rig logic calculations
 };
 
-/* The work data used by the FRigUnit_RigLogic */
+/* A helper struct used inside of the RigUnit_RigLogic to store mappings between skelMesh and DNA. */
 USTRUCT(meta = (DocumentationPolicy = "Strict"))
-struct FRigUnit_RigLogic_Data
+struct FRigDefinitionMapping
 {
+
 	GENERATED_BODY()
 
-	FRigUnit_RigLogic_Data();
- 	~FRigUnit_RigLogic_Data();
- 	FRigUnit_RigLogic_Data(const FRigUnit_RigLogic_Data& Other);
- 	FRigUnit_RigLogic_Data& operator=(const FRigUnit_RigLogic_Data& Other);
-
-	/** Cached Skeletal Mesh Component **/
-	UPROPERTY(transient)
-	TWeakObjectPtr<USkeletalMeshComponent> SkelMeshComponent;
-
-	FSharedRigRuntimeContext* SharedRigRuntimeContext;
-
-	/** RigInstance is a thin class, containing only character instance specific data
-	  * i.e. the output buffers where RigLogic writes the results of computations, specific
-	  * to each instance of a rig.
-	**/
-	TUniquePtr<FRigInstance> RigInstance;
+		FRigDefinitionMapping() = default;
+	FRigDefinitionMapping(const IBehaviorReader* DNABehavior, const URigHierarchy* InHierarchy);
 
 	/** Mapping RL indices to UE indices
 	  * Note: we use int32 instead of uint32 to allow storing INDEX_NONE for missing elements
@@ -77,49 +65,81 @@ struct FRigUnit_RigLogic_Data
 
 	/** RL input index to ControlRig's input curve index for each LOD **/
 	UPROPERTY(transient)
-	TArray<int32> InputCurveIndices;
+		TArray<int32> InputCurveIndices;
 
 	/** RL joint index to ControlRig's hierarchy bone index **/
 	UPROPERTY(transient)
-	TArray<int32> HierarchyBoneIndices;
+		TArray<int32> HierarchyBoneIndices;	/** RL mesh blend shape index to ControlRig's output blendshape curve index for each LOD **/
 
-	/** RL mesh blend shape index to ControlRig's output blendshape curve index for each LOD **/
+		/** Mapping RL indices to UE indices
+		  * Note: we use int32 instead of uint32 to allow storing INDEX_NONE for missing elements
+		  * if value is valid, it is cast to appropriate uint type
+		**/
 	UPROPERTY(transient)
-	TArray<FRigUnit_RigLogic_IntArray> MorphTargetCurveIndices;
+		TArray<FRigUnit_RigLogic_IntArray> MorphTargetCurveIndices;
 
 	/** RL mesh+blend shape array index to RL blend shape index for each LOD **/
 	UPROPERTY(transient)
-	TArray<FRigUnit_RigLogic_IntArray> BlendShapeIndices;
+		TArray<FRigUnit_RigLogic_IntArray> BlendShapeIndices;
 
 	/** RL animated map index to ControlRig's output anim map curve index for each LOD **/
 	UPROPERTY(transient)
-	TArray<FRigUnit_RigLogic_IntArray> CurveElementIndicesForAnimMaps;
+		TArray<FRigUnit_RigLogic_IntArray> CurveElementIndicesForAnimMaps;
 
 	/** RL animated map index to RL anim map curve index for each LOD **/
 	UPROPERTY(transient)
-	TArray<FRigUnit_RigLogic_IntArray> RigLogicIndicesForAnimMaps;
+		TArray<FRigUnit_RigLogic_IntArray> RigLogicIndicesForAnimMaps;
+
+private:
+	/** Makes a map of input curve indices from DNA file to the control rig curves **/
+	void MapInputCurveIndices(const IBehaviorReader* DNABehavior, const URigHierarchy* InHierarchy);
+	/** Uses names to map joint indices from DNA file to the indices of bones in control rig hierarchy **/
+	void MapJoints(const IBehaviorReader* DNABehavior, const URigHierarchy* Hierarchy);
+	/** Uses names of blend shapes and meshes in DNA file, for all LODs, to map their indices to the indices of
+	  * morph target curves in the curve container; curve name format is <mesh>__<blendshape> **/
+	void MapMorphTargets(const IBehaviorReader* DNABehavior, const URigHierarchy* InHierarchy);
+	/** Uses names to map mask multiplier indices from DNA file, for all LODs, to the indices of curves in the
+	  * control rig's curve container **/
+	void MapMaskMultipliers(const IBehaviorReader* DNABehavior, const URigHierarchy* InHierarchy);
+
+};
+
+/* The work data used by the FRigUnit_RigLogic */
+USTRUCT(meta = (DocumentationPolicy = "Strict"))
+struct FRigUnit_RigLogic_Data
+{
+	GENERATED_BODY()
+
+		FRigUnit_RigLogic_Data();
+	~FRigUnit_RigLogic_Data();
+	FRigUnit_RigLogic_Data(const FRigUnit_RigLogic_Data& Other);
+	FRigUnit_RigLogic_Data& operator=(const FRigUnit_RigLogic_Data& Other);
+
+	/** Cached Skeletal Mesh Component **/
+	UPROPERTY(transient)
+		TWeakObjectPtr<USkeletalMeshComponent> SkelMeshComponent;
+
+	FSharedRigRuntimeContext* SharedRigRuntimeContext;
+
+	/** RigInstance is a thin class, containing only character instance specific data
+	  * i.e. the output buffers where RigLogic writes the results of computations, specific
+	  * to each instance of a rig.
+	**/
+	TUniquePtr<FRigInstance> RigInstance;
+
+	TSharedPtr<const FRigDefinitionMapping> Mappings;
 
 	/** LOD for which the model is rendered **/
 	UPROPERTY(transient)
-	uint32 CurrentLOD;
+		uint32 CurrentLOD;
 
 	static const uint8 MAX_ATTRS_PER_JOINT;
 
 	bool IsRigLogicInitialized();
 	void InitializeRigLogic(const URigHierarchy* InHierarchy);
 
-	/** Makes a map of input curve indices from DNA file to the control rig curves **/
-	void MapInputCurveIndices(const URigHierarchy* InHierarchy);
-	/** Uses names to map joint indices from DNA file to the indices of bones in control rig hierarchy **/
-	void MapJoints(const URigHierarchy* Hierarchy);
 	/** Cache the joint indices that change per each LOD **/
 	void CacheVariableJointIndices();
-	/** Uses names of blend shapes and meshes in DNA file, for all LODs, to map their indices to the indices of
-	  * morph target curves in the curve container; curve name format is <mesh>__<blendshape> **/
-	void MapMorphTargets(const URigHierarchy* InHierarchy);
-	/** Uses names to map mask multiplier indices from DNA file, for all LODs, to the indices of curves in the
-	  * control rig's curve container **/
-	void MapMaskMultipliers(const URigHierarchy* InHierarchy);
 
 	/** Calculates joint positions, orientation and scale based on inputs curves of the control rig **/
 	void CalculateRigLogic(const URigHierarchy* InHierarchy);
@@ -129,6 +149,8 @@ struct FRigUnit_RigLogic_Data
 	void UpdateBlendShapeCurves(URigHierarchy* InHierarchy, TArrayView<const float> BlendShapeValues);
 	/** Updates anim map curve values based on values of input curves of the control rig **/
 	void UpdateAnimMapCurves(URigHierarchy* InHierarchy, TArrayView<const float> AnimMapOutputs);
+private:
+	static TDiscardableKeyValueCache<uint32, TSharedPtr<const FRigDefinitionMapping>> RigDefintionMappingsCache;
 };
 
 
@@ -136,7 +158,7 @@ struct FRigUnit_RigLogic_Data
   *  animated map multiplier curves */
 
 USTRUCT(meta = (DisplayName = "RigLogic", Category = "RigLogic", DocumentationPolicy = "Strict", Keywords = "Rig,RigLogic"))
-struct FRigUnit_RigLogic: public FRigUnitMutable
+struct FRigUnit_RigLogic : public FRigUnitMutable
 {
 	GENERATED_BODY()
 
@@ -148,7 +170,7 @@ public:
 #endif
 
 	RIGVM_METHOD()
-	void Execute(const FRigUnitContext& Context) override;
+		void Execute(const FRigUnitContext& Context) override;
 
 private:
 	static FSharedRigRuntimeContext* GetSharedRigRuntimeContext(USkeletalMesh* SkelMesh);
@@ -156,5 +178,5 @@ private:
 private:
 	// internal work data for the unit
 	UPROPERTY(transient)
-	FRigUnit_RigLogic_Data Data;
+		FRigUnit_RigLogic_Data Data;
 };
