Disable SSE support if target is arm64.

#rnx

diff --git a/Engine/Plugins/Animation/RigLogic/Source/RigLogicLib/Private/riglogic/joints/JointsBuilder.cpp b/Engine/Plugins/Animation/RigLogic/Source/RigLogicLib/Private/riglogic/joints/JointsBuilder.cpp
--- a/Engine/Plugins/Animation/RigLogic/Source/RigLogicLib/Private/riglogic/joints/JointsBuilder.cpp
+++ b/Engine/Plugins/Animation/RigLogic/Source/RigLogicLib/Private/riglogic/joints/JointsBuilder.cpp
@@ -1,22 +1,24 @@
 // Copyright Epic Games, Inc. All Rights Reserved.
 
 #include "riglogic/joints/JointsBuilder.h"
-
 #include "riglogic/joints/bpcm/builders/Scalar.h"
-#ifdef RL_BUILD_WITH_AVX
-    #ifdef RL_USE_HALF_FLOATS
-        #include "riglogic/joints/bpcm/builders/AVXHalfFloat.h"
-    #else
-        #include "riglogic/joints/bpcm/builders/AVX.h"
-    #endif  // RL_USE_HALF_FLOATS
-#endif  // RL_BUILD_WITH_AVX
-#ifdef RL_BUILD_WITH_SSE
-    #ifdef RL_USE_HALF_FLOATS
-        #include "riglogic/joints/bpcm/builders/SSEHalfFloat.h"
-    #else
-        #include "riglogic/joints/bpcm/builders/SSE.h"
-    #endif  // RL_USE_HALF_FLOATS
-#endif  // RL_BUILD_WITH_SSE
+
+#if PLATFORM_CPU_X86_FAMILY
+    #ifdef RL_BUILD_WITH_AVX
+        #ifdef RL_USE_HALF_FLOATS
+            #include "riglogic/joints/bpcm/builders/AVXHalfFloat.h"
+        #else
+            #include "riglogic/joints/bpcm/builders/AVX.h"
+        #endif  // RL_USE_HALF_FLOATS
+    #endif  // RL_BUILD_WITH_AVX
+    #ifdef RL_BUILD_WITH_SSE
+        #ifdef RL_USE_HALF_FLOATS
+            #include "riglogic/joints/bpcm/builders/SSEHalfFloat.h"
+        #else
+            #include "riglogic/joints/bpcm/builders/SSE.h"
+        #endif  // RL_USE_HALF_FLOATS
+    #endif  // RL_BUILD_WITH_SSE
+#endif // PLATFORM_CPU_X86_FAMILY
 
 #include <pma/utils/ManagedInstance.h>
 
@@ -25,16 +27,18 @@ namespace rl4 {
 JointsBuilder::~JointsBuilder() = default;
 
 JointsBuilder::JointsBuilderPtr JointsBuilder::create(Configuration config, MemoryResource* memRes) {
-    #ifdef RL_BUILD_WITH_SSE
-        if (config.calculationType == CalculationType::SSE) {
-            return pma::UniqueInstance<bpcm::SSEJointsBuilder, JointsBuilder>::with(memRes).create(memRes);
-        }
-    #endif  // RL_BUILD_WITH_SSE
-    #ifdef RL_BUILD_WITH_AVX
-        if (config.calculationType == CalculationType::AVX) {
-            return pma::UniqueInstance<bpcm::AVXJointsBuilder, JointsBuilder>::with(memRes).create(memRes);
-        }
-    #endif  // RL_BUILD_WITH_AVX
+    #if PLATFORM_CPU_X86_FAMILY
+        #ifdef RL_BUILD_WITH_SSE
+            if (config.calculationType == CalculationType::SSE) {
+                return pma::UniqueInstance<bpcm::SSEJointsBuilder, JointsBuilder>::with(memRes).create(memRes);
+            }
+        #endif  // RL_BUILD_WITH_SSE
+        #ifdef RL_BUILD_WITH_AVX
+            if (config.calculationType == CalculationType::AVX) {
+                return pma::UniqueInstance<bpcm::AVXJointsBuilder, JointsBuilder>::with(memRes).create(memRes);
+            }
+        #endif  // RL_BUILD_WITH_AVX
+    #endif // PLATFORM_CPU_X86_FAMILY
     return pma::UniqueInstance<bpcm::ScalarJointsBuilder, JointsBuilder>::with(memRes).create(memRes);
 }
 
diff --git a/Engine/Plugins/Animation/RigLogic/Source/RigLogicLib/Private/riglogic/joints/JointsBuilder.h b/Engine/Plugins/Animation/RigLogic/Source/RigLogicLib/Private/riglogic/joints/JointsBuilder.h
--- a/Engine/Plugins/Animation/RigLogic/Source/RigLogicLib/Private/riglogic/joints/JointsBuilder.h
+++ b/Engine/Plugins/Animation/RigLogic/Source/RigLogicLib/Private/riglogic/joints/JointsBuilder.h
@@ -2,6 +2,7 @@
 
 #pragma once
 
+#include "HAL/Platform.h"
 #include "riglogic/joints/JointsEvaluator.h"
 #include "riglogic/riglogic/Configuration.h"
 #include "riglogic/types/Aliases.h"
diff --git a/Engine/Plugins/Animation/RigLogic/Source/RigLogicLib/Private/riglogic/joints/bpcm/builders/SSE.cpp b/Engine/Plugins/Animation/RigLogic/Source/RigLogicLib/Private/riglogic/joints/bpcm/builders/SSE.cpp
--- a/Engine/Plugins/Animation/RigLogic/Source/RigLogicLib/Private/riglogic/joints/bpcm/builders/SSE.cpp
+++ b/Engine/Plugins/Animation/RigLogic/Source/RigLogicLib/Private/riglogic/joints/bpcm/builders/SSE.cpp
@@ -1,9 +1,9 @@
 // Copyright Epic Games, Inc. All Rights Reserved.
 
-#if !defined(RL_USE_HALF_FLOATS) && defined(RL_BUILD_WITH_SSE)
-
 #include "riglogic/joints/bpcm/builders/SSE.h"
 
+#if !defined(RL_USE_HALF_FLOATS) && defined(RL_BUILD_WITH_SSE) && PLATFORM_CPU_X86_FAMILY
+
 #include "riglogic/joints/bpcm/Consts.h"
 #include "riglogic/joints/bpcm/strategies/SSE.h"
 
@@ -26,4 +26,4 @@ SSEJointsBuilder::~SSEJointsBuilder() = default;
 
 }  // namespace rl4
 
-#endif  // !defined(RL_USE_HALF_FLOATS) && defined(RL_BUILD_WITH_SSE)
+#endif  // !defined(RL_USE_HALF_FLOATS) && defined(RL_BUILD_WITH_SSE) && PLATFORM_CPU_X86_FAMILY
diff --git a/Engine/Plugins/Animation/RigLogic/Source/RigLogicLib/Private/riglogic/joints/bpcm/builders/SSE.h b/Engine/Plugins/Animation/RigLogic/Source/RigLogicLib/Private/riglogic/joints/bpcm/builders/SSE.h
--- a/Engine/Plugins/Animation/RigLogic/Source/RigLogicLib/Private/riglogic/joints/bpcm/builders/SSE.h
+++ b/Engine/Plugins/Animation/RigLogic/Source/RigLogicLib/Private/riglogic/joints/bpcm/builders/SSE.h
@@ -2,8 +2,10 @@
 
 #pragma once
 
+#include "HAL/Platform.h"
+
 // *INDENT-OFF*
-#if !defined(RL_USE_HALF_FLOATS) && defined(RL_BUILD_WITH_SSE)
+#if !defined(RL_USE_HALF_FLOATS) && defined(RL_BUILD_WITH_SSE) && PLATFORM_CPU_X86_FAMILY
 
 #include "riglogic/joints/bpcm/builders/Float.h"
 #include "riglogic/types/Aliases.h"
@@ -22,5 +24,5 @@ class SSEJointsBuilder : public FloatStorageBuilder {
 
 }  // namespace rl4
 
-#endif  // !defined(RL_USE_HALF_FLOATS) && defined(RL_BUILD_WITH_SSE)
+#endif  // !defined(RL_USE_HALF_FLOATS) && defined(RL_BUILD_WITH_SSE) && PLATFORM_CPU_X86_FAMILY
 // *INDENT-ON*
diff --git a/Engine/Plugins/Animation/RigLogic/Source/RigLogicLib/Private/riglogic/joints/bpcm/builders/SSEHalfFloat.cpp b/Engine/Plugins/Animation/RigLogic/Source/RigLogicLib/Private/riglogic/joints/bpcm/builders/SSEHalfFloat.cpp
--- a/Engine/Plugins/Animation/RigLogic/Source/RigLogicLib/Private/riglogic/joints/bpcm/builders/SSEHalfFloat.cpp
+++ b/Engine/Plugins/Animation/RigLogic/Source/RigLogicLib/Private/riglogic/joints/bpcm/builders/SSEHalfFloat.cpp
@@ -1,9 +1,9 @@
 // Copyright Epic Games, Inc. All Rights Reserved.
 
-#if defined(RL_USE_HALF_FLOATS) && defined(RL_BUILD_WITH_SSE)
-
 #include "riglogic/joints/bpcm/builders/SSEHalfFloat.h"
 
+#if defined(RL_USE_HALF_FLOATS) && defined(RL_BUILD_WITH_SSE) && PLATFORM_CPU_X86_FAMILY
+
 #include "riglogic/joints/bpcm/Consts.h"
 #include "riglogic/joints/bpcm/strategies/SSE.h"
 #include "riglogic/utils/Extd.h"
@@ -97,4 +97,4 @@ void SSEJointsBuilder::setValues(const dna::BehaviorReader* reader) {
 
 }  // namespace rl4
 
-#endif  // defined(RL_USE_HALF_FLOATS) && defined(RL_BUILD_WITH_SSE)
+#endif  // defined(RL_USE_HALF_FLOATS) && defined(RL_BUILD_WITH_SSE) && PLATFORM_CPU_X86_FAMILY
diff --git a/Engine/Plugins/Animation/RigLogic/Source/RigLogicLib/Private/riglogic/joints/bpcm/builders/SSEHalfFloat.h b/Engine/Plugins/Animation/RigLogic/Source/RigLogicLib/Private/riglogic/joints/bpcm/builders/SSEHalfFloat.h
--- a/Engine/Plugins/Animation/RigLogic/Source/RigLogicLib/Private/riglogic/joints/bpcm/builders/SSEHalfFloat.h
+++ b/Engine/Plugins/Animation/RigLogic/Source/RigLogicLib/Private/riglogic/joints/bpcm/builders/SSEHalfFloat.h
@@ -2,8 +2,10 @@
 
 #pragma once
 
+#include "HAL/Platform.h"
+
 // *INDENT-OFF*
-#if defined(RL_USE_HALF_FLOATS) && defined(RL_BUILD_WITH_SSE)
+#if defined(RL_USE_HALF_FLOATS) && defined(RL_BUILD_WITH_SSE) && PLATFORM_CPU_X86_FAMILY
 
 #include "riglogic/joints/bpcm/JointsBuilderCommon.h"
 #include "riglogic/types/Aliases.h"
@@ -28,5 +30,5 @@ class SSEJointsBuilder : public JointsBuilderCommon<std::uint16_t> {
 
 }  // namespace rl4
 
-#endif  // defined(RL_USE_HALF_FLOATS) && defined(RL_BUILD_WITH_SSE)
+#endif  // defined(RL_USE_HALF_FLOATS) && defined(RL_BUILD_WITH_SSE) && PLATFORM_CPU_X86_FAMILY
 // *INDENT-ON*
diff --git a/Engine/Plugins/Animation/RigLogic/Source/RigLogicLib/Private/riglogic/system/simd/SIMD.h b/Engine/Plugins/Animation/RigLogic/Source/RigLogicLib/Private/riglogic/system/simd/SIMD.h
--- a/Engine/Plugins/Animation/RigLogic/Source/RigLogicLib/Private/riglogic/system/simd/SIMD.h
+++ b/Engine/Plugins/Animation/RigLogic/Source/RigLogicLib/Private/riglogic/system/simd/SIMD.h
@@ -2,15 +2,18 @@
 
 #pragma once
 
+#include "HAL/Platform.h"
+
 #if defined(RL_USE_HALF_FLOATS) && !defined(TRIMD_ENABLE_F16C)
     #define TRIMD_ENABLE_F16C
 #endif  // RL_USE_HALF_FLOATS
 
-#if defined(RL_BUILD_WITH_AVX) && !defined(TRIMD_ENABLE_AVX)
-    #define TRIMD_ENABLE_AVX
-#endif  // RL_BUILD_WITH_AVX
-
-#if defined(RL_BUILD_WITH_SSE) && !defined(TRIMD_ENABLE_SSE)
-    #define TRIMD_ENABLE_SSE
-#endif  // RL_BUILD_WITH_SSE
+#if PLATFORM_CPU_X86_FAMILY
+    #if defined(RL_BUILD_WITH_AVX) && !defined(TRIMD_ENABLE_AVX)
+        #define TRIMD_ENABLE_AVX
+    #endif  // RL_BUILD_WITH_AVX
 
+    #if defined(RL_BUILD_WITH_SSE) && !defined(TRIMD_ENABLE_SSE)
+        #define TRIMD_ENABLE_SSE
+    #endif  // RL_BUILD_WITH_SSE
+#endif // PLATFORM_CPU_X86_FAMILY
