Disable unity build of RigLogicLibTest #jira UE-144385 #rb bojan.brankov #preflight 6233423048746817f12b70b2

diff --git a/Engine/Plugins/Animation/RigLogic/Source/RigLogicLib/RigLogicLib.Build.cs b/Engine/Plugins/Animation/RigLogic/Source/RigLogicLib/RigLogicLib.Build.cs
--- a/Engine/Plugins/Animation/RigLogic/Source/RigLogicLib/RigLogicLib.Build.cs
+++ b/Engine/Plugins/Animation/RigLogic/Source/RigLogicLib/RigLogicLib.Build.cs
@@ -13,6 +13,8 @@ public class RigLogicLib : ModuleRules
 
     public RigLogicLib(ReadOnlyTargetRules Target) : base(Target)
     {
+        bUseUnity = false;
+
         PublicDependencyModuleNames.AddRange(
             new string[]
             {
