Updating the RigLogic plugin.

diff --git a/Engine/Plugins/Animation/RigLogic/RigLogic.uplugin b/Engine/Plugins/Animation/RigLogic/RigLogic.uplugin
--- a/Engine/Plugins/Animation/RigLogic/RigLogic.uplugin
+++ b/Engine/Plugins/Animation/RigLogic/RigLogic.uplugin
@@ -32,6 +32,11 @@
 			"Type": "Runtime",
 			"LoadingPhase": "Default"
 		},
+		{
+			"Name" : "RigLogicDeveloper",
+			"Type" : "UncookedOnly",
+			"LoadingPhase" : "PreDefault"
+		},
 		{
 			"Name": "RigLogicEditor",
 			"Type": "Editor",
diff --git a/Engine/Plugins/Animation/RigLogic/Source/RigLogicDeveloper/Private/AnimGraphNode_RigLogic.cpp b/Engine/Plugins/Animation/RigLogic/Source/RigLogicDeveloper/Private/AnimGraphNode_RigLogic.cpp
new file mode 100644
--- /dev/null
+++ b/Engine/Plugins/Animation/RigLogic/Source/RigLogicDeveloper/Private/AnimGraphNode_RigLogic.cpp
@@ -0,0 +1,28 @@
+// Copyright Epic Games, Inc. All Rights Reserved.
+
+#include "AnimGraphNode_RigLogic.h"
+
+#include "AnimGraphNode_Base.h"
+
+#define LOCTEXT_NAMESPACE "AnimGraphNode_RigLogic_DeveloperModule"
+
+UAnimGraphNode_RigLogic::UAnimGraphNode_RigLogic(const FObjectInitializer& ObjectInitializer) : Super(ObjectInitializer)
+{
+}
+
+void UAnimGraphNode_RigLogic::ValidateAnimNodeDuringCompilation(USkeleton* ForSkeleton, FCompilerResultsLog& MessageLog)
+{
+	Super::ValidateAnimNodeDuringCompilation(ForSkeleton, MessageLog);
+}
+
+FText UAnimGraphNode_RigLogic::GetTooltipText() const
+{
+	return LOCTEXT("AnimGraphNode_RigLogic_Tooltip", "Rig Logic Animation Node");
+}
+
+FText UAnimGraphNode_RigLogic::GetNodeTitle(ENodeTitleType::Type TitleType) const
+{
+	return LOCTEXT("AnimGraphNode_RigLogic_Title", "AnimNode_RigLogic");
+}
+
+#undef LOCTEXT_NAMESPACE
diff --git a/Engine/Plugins/Animation/RigLogic/Source/RigLogicDeveloper/Private/RigLogicDeveloper.cpp b/Engine/Plugins/Animation/RigLogic/Source/RigLogicDeveloper/Private/RigLogicDeveloper.cpp
new file mode 100644
--- /dev/null
+++ b/Engine/Plugins/Animation/RigLogic/Source/RigLogicDeveloper/Private/RigLogicDeveloper.cpp
@@ -0,0 +1,29 @@
+// Copyright Epic Games, Inc. All Rights Reserved.
+
+#include "RigLogicDeveloper.h"
+#include "Modules/ModuleManager.h"
+#include "Developer/MessageLog/Public/MessageLogModule.h"
+
+#define LOCTEXT_NAMESPACE "RigLogicDeveloperModule"
+
+DEFINE_LOG_CATEGORY(LogRigLogicDeveloper);
+
+void FRigLogicDeveloperModule::StartupModule()
+{
+	FMessageLogModule& MessageLogModule = FModuleManager::LoadModuleChecked<FMessageLogModule>("MessageLog");
+	FMessageLogInitializationOptions InitOptions;
+	InitOptions.bShowFilters = true;
+	InitOptions.bShowPages = false;
+	InitOptions.bAllowClear = true;
+	MessageLogModule.RegisterLogListing("RigLogicLog", LOCTEXT("RigLogicLog", "RigLogic Log"), InitOptions);
+}
+
+void FRigLogicDeveloperModule::ShutdownModule()
+{
+	FMessageLogModule& MessageLogModule = FModuleManager::LoadModuleChecked<FMessageLogModule>("MessageLog");
+	MessageLogModule.UnregisterLogListing("RigLogicLog");
+}
+
+IMPLEMENT_MODULE(FRigLogicDeveloperModule, RigLogicDeveloper)
+
+#undef LOCTEXT_NAMESPACE
diff --git a/Engine/Plugins/Animation/RigLogic/Source/RigLogicDeveloper/Public/AnimGraphNode_RigLogic.h b/Engine/Plugins/Animation/RigLogic/Source/RigLogicDeveloper/Public/AnimGraphNode_RigLogic.h
new file mode 100644
--- /dev/null
+++ b/Engine/Plugins/Animation/RigLogic/Source/RigLogicDeveloper/Public/AnimGraphNode_RigLogic.h
@@ -0,0 +1,23 @@
+// Copyright Epic Games, Inc. All Rights Reserved.
+
+#pragma once
+
+#include "AnimGraphNode_Base.h"
+#include "AnimNode_RigLogic.h"
+#include "AnimGraphNode_RigLogic.generated.h"
+
+UCLASS(meta = (Keywords = "Rig Logic Animation Node"))
+class RIGLOGICDEVELOPER_API UAnimGraphNode_RigLogic : public UAnimGraphNode_Base
+{
+	GENERATED_UCLASS_BODY()
+
+	UPROPERTY(EditAnywhere, Category = Settings)
+	FAnimNode_RigLogic Node;
+
+public:
+	FText GetTooltipText() const;
+	FText GetNodeTitle(ENodeTitleType::Type TitleType) const;
+
+protected:
+	void ValidateAnimNodeDuringCompilation(USkeleton* ForSkeleton, FCompilerResultsLog& MessageLog);
+};
diff --git a/Engine/Plugins/Animation/RigLogic/Source/RigLogicDeveloper/Public/RigLogicDeveloper.h b/Engine/Plugins/Animation/RigLogic/Source/RigLogicDeveloper/Public/RigLogicDeveloper.h
new file mode 100644
--- /dev/null
+++ b/Engine/Plugins/Animation/RigLogic/Source/RigLogicDeveloper/Public/RigLogicDeveloper.h
@@ -0,0 +1,15 @@
+// Copyright Epic Games, Inc. All Rights Reserved.
+
+#pragma once
+
+#include "CoreMinimal.h"
+#include "Modules/ModuleInterface.h"
+
+DECLARE_LOG_CATEGORY_EXTERN(LogRigLogicDeveloper, Log, All);
+
+class RIGLOGICDEVELOPER_API FRigLogicDeveloperModule : public IModuleInterface 
+{
+public:
+	void StartupModule() override;
+	void ShutdownModule() override;
+};
diff --git a/Engine/Plugins/Animation/RigLogic/Source/RigLogicDeveloper/RigLogicDeveloper.Build.cs b/Engine/Plugins/Animation/RigLogic/Source/RigLogicDeveloper/RigLogicDeveloper.Build.cs
new file mode 100644
--- /dev/null
+++ b/Engine/Plugins/Animation/RigLogic/Source/RigLogicDeveloper/RigLogicDeveloper.Build.cs
@@ -0,0 +1,54 @@
+// Copyright Epic Games, Inc. All Rights Reserved.
+
+namespace UnrealBuildTool.Rules
+{
+	public class RigLogicDeveloper : ModuleRules
+	{
+		public RigLogicDeveloper(ReadOnlyTargetRules Target) : base(Target)
+		{
+			PrivateDependencyModuleNames.AddRange(
+				new string[]
+				{
+					"AnimGraphRuntime",
+					"AnimationCore",
+					"Core",
+					"CoreUObject",
+					"Engine",
+					"PropertyPath",
+					"Slate",
+					"SlateCore",
+					"InputCore",
+					"TimeManagement",
+					"EditorWidgets",
+					"MessageLog",
+					"RigLogicModule"
+				}
+			);
+
+			PublicDependencyModuleNames.AddRange(
+				new string[]
+				{
+					"AnimationCore",
+				}
+			);
+
+			if (Target.bBuildEditor == true)
+			{
+				PrivateDependencyModuleNames.AddRange(
+					new string[]
+					{
+						"EditorFramework",
+						"UnrealEd",
+						"AnimGraph",
+						"BlueprintGraph",
+						"PropertyEditor",
+						"GraphEditor",
+					}
+				);
+
+				PrivateIncludePathModuleNames.Add("RigLogicEditor");
+				DynamicallyLoadedModuleNames.Add("RigLogicEditor");
+			}
+		}
+	}
+}
diff --git a/Engine/Plugins/Animation/RigLogic/Source/RigLogicLib/Private/dna/DNA.h b/Engine/Plugins/Animation/RigLogic/Source/RigLogicLib/Private/dna/DNA.h
--- a/Engine/Plugins/Animation/RigLogic/Source/RigLogicLib/Private/dna/DNA.h
+++ b/Engine/Plugins/Animation/RigLogic/Source/RigLogicLib/Private/dna/DNA.h
@@ -681,6 +681,7 @@ struct RawGeometry {
 };
 
 struct DNA {
+	MemoryResource* memRes;
     Signature<3> signature{{'D', 'N', 'A'}};
     Version version{2, 1};
     SectionLookupTable sections;
@@ -690,7 +691,8 @@ struct DNA {
     RawGeometry geometry;
     Signature<3> eof{{'A', 'N', 'D'}};
 
-    explicit DNA(MemoryResource* memRes) :
+    explicit DNA(MemoryResource* memRes_) :
+        memRes{memRes_},
         sections{},
         descriptor{sections.descriptor, memRes},
         definition{sections.definition, memRes},
@@ -717,6 +719,21 @@ struct DNA {
         archive(signature, version, sections, descriptor, definition, behavior, geometry, eof);
     }
 
+	void unloadDefinition()
+	{
+		definition = RawDefinition{sections.definition, memRes};
+	}
+
+	void unloadBehavior()
+	{
+		behavior = RawBehavior{sections.behavior, sections.controls, sections.joints, sections.blendShapeChannels, sections.animatedMaps, memRes};
+	}
+
+	void unloadGeometry()
+	{
+		geometry = RawGeometry{sections.geometry, memRes};
+	}
+
 };
 
 }  // namespace dna
diff --git a/Engine/Plugins/Animation/RigLogic/Source/RigLogicLib/Private/dna/stream/StreamReaderImpl.cpp b/Engine/Plugins/Animation/RigLogic/Source/RigLogicLib/Private/dna/stream/StreamReaderImpl.cpp
--- a/Engine/Plugins/Animation/RigLogic/Source/RigLogicLib/Private/dna/stream/StreamReaderImpl.cpp
+++ b/Engine/Plugins/Animation/RigLogic/Source/RigLogicLib/Private/dna/stream/StreamReaderImpl.cpp
@@ -21,100 +21,6 @@
 
 namespace dna {
 
-namespace {
-
-constexpr std::size_t operator"" _KB(unsigned long long size) {
-    #if !defined(__clang__) && defined(__GNUC__)
-        #pragma GCC diagnostic push
-        #pragma GCC diagnostic ignored "-Wuseless-cast"
-    #endif
-    return static_cast<std::size_t>(size * 1024ul);
-    #if !defined(__clang__) && defined(__GNUC__)
-        #pragma GCC diagnostic pop
-    #endif
-}
-
-constexpr std::size_t operator"" _MB(unsigned long long size) {
-    #if !defined(__clang__) && defined(__GNUC__)
-        #pragma GCC diagnostic push
-        #pragma GCC diagnostic ignored "-Wuseless-cast"
-    #endif
-    return static_cast<std::size_t>(size * 1024ul * 1024ul);
-    #if !defined(__clang__) && defined(__GNUC__)
-        #pragma GCC diagnostic pop
-    #endif
-}
-
-struct ArenaFactory {
-
-    #ifdef _MSC_VER
-        struct AllocationOverhead {
-            enum Sizes : std::size_t {
-                // All and Geometry denote overheads relative to the stream size itself
-                All = 20_MB,
-                Geometry = 16_MB,
-                // Rest of the overheads are absolute values
-                AllWithoutBlendShapes = 30_MB,
-                GeometryWithoutBlendShapes = 26_MB,
-                Behavior = 5_MB,
-                Definition = 256_KB,
-                Descriptor = 64_KB
-            };
-        };
-    #else
-        struct AllocationOverhead {
-            enum Sizes : std::size_t {
-                // All and Geometry denote overheads relative to the stream size itself
-                All = 16_MB,
-                Geometry = 12_MB,
-                // Rest of the overheads are absolute values
-                AllWithoutBlendShapes = 26_MB,
-                GeometryWithoutBlendShapes = 22_MB,
-                Behavior = 5_MB,
-                Definition = 256_KB,
-                Descriptor = 64_KB
-            };
-        };
-    #endif
-
-    static MemoryResource* create(dna::DataLayer layer, std::size_t streamSize, MemoryResource* upstream) {
-        // upstream may be nullptr, so we rely on PolyAllocator's fallback to a default memory resource in such cases
-        auto createArena = [upstream](std::size_t initialSize, std::size_t regionSize) {
-                // In the unlikely case that the arena runs out of memory, this growth factor will prevent the arena
-                // to be stuck in an infinite loop for allocations of single chunks that are greater than regionSize
-                static constexpr float arenaGrowthFactor = 1.1f;
-                PolyAllocator<ArenaMemoryResource> alloc{upstream};
-                return alloc.newObject(initialSize, regionSize, arenaGrowthFactor, alloc.getMemoryResource());
-            };
-        if (layer == dna::DataLayer::All) {
-            return createArena(streamSize + AllocationOverhead::All, 4_MB);
-        } else if (layer == dna::DataLayer::Geometry) {
-            return createArena(streamSize + AllocationOverhead::Geometry, 4_MB);
-        } else if (layer == dna::DataLayer::AllWithoutBlendShapes) {
-            return createArena(AllocationOverhead::AllWithoutBlendShapes, 2_MB);
-        } else if (layer == dna::DataLayer::GeometryWithoutBlendShapes) {
-            return createArena(AllocationOverhead::GeometryWithoutBlendShapes, 2_MB);
-        } else if (layer == dna::DataLayer::Behavior) {
-            return createArena(AllocationOverhead::Behavior, 2_MB);
-        } else if (layer == dna::DataLayer::Definition) {
-            return createArena(AllocationOverhead::Definition, 64_KB);
-        } else if (layer == dna::DataLayer::Descriptor) {
-            return createArena(AllocationOverhead::Descriptor, 64_KB);
-        }
-        // Unreachable, unless a new layer is added
-        return createArena(streamSize, 4_MB);
-    }
-
-    static void destroy(MemoryResource* instance) {
-        auto arena = static_cast<ArenaMemoryResource*>(instance);
-        PolyAllocator<ArenaMemoryResource> alloc{arena->getUpstreamMemoryResource()};
-        alloc.deleteObject(arena);
-    }
-
-};
-
-}  // namespace
-
 const sc::StatusCode StreamReader::SignatureMismatchError{200, "DNA signature mismatched, expected %.3s, got %.3s"};
 const sc::StatusCode StreamReader::VersionMismatchError{201, "DNA version mismatched, expected %hu.%hu, got %hu.%hu"};
 const sc::StatusCode StreamReader::InvalidDataError{202, "Invalid data in DNA"};
@@ -131,16 +37,6 @@ sc::StatusProvider StreamReaderImpl::status{SignatureMismatchError, VersionMisma
 StreamReader::~StreamReader() = default;
 
 StreamReader* StreamReader::create(BoundedIOStream* stream, DataLayer layer, std::uint16_t maxLOD, MemoryResource* memRes) {
-    if (maxLOD == LODLimits::max()) {
-        #if !defined(__clang__) && defined(__GNUC__)
-            #pragma GCC diagnostic push
-            #pragma GCC diagnostic ignored "-Wuseless-cast"
-        #endif
-        memRes = ArenaFactory::create(layer, static_cast<std::size_t>(stream->size()), memRes);
-        #if !defined(__clang__) && defined(__GNUC__)
-            #pragma GCC diagnostic pop
-        #endif
-    }
     PolyAllocator<StreamReaderImpl> alloc{memRes};
     return alloc.newObject(stream, layer, maxLOD, LODLimits::min(), memRes);
 }
@@ -166,19 +62,9 @@ StreamReader* StreamReader::create(BoundedIOStream* stream,
 void StreamReader::destroy(StreamReader* instance) {
     // NOLINTNEXTLINE(cppcoreguidelines-pro-type-static-cast-downcast)
     auto reader = static_cast<StreamReaderImpl*>(instance);
-    // In the presence of LOD constraints, ArenaMemoryResource is not used,
-    // as the approximations for the memory overhead wouldn't be accurate
-    const bool usesArena = !(reader->isLODConstrained());
     MemoryResource* memRes = reader->getMemoryResource();
     PolyAllocator<StreamReaderImpl> readerAlloc{memRes};
     readerAlloc.deleteObject(reader);
-    // Delete DNA arena if it was used
-    if (usesArena) {
-        auto arena = static_cast<ArenaMemoryResource*>(memRes);
-        auto upstream = arena->getUpstreamMemoryResource();
-        PolyAllocator<ArenaMemoryResource> arenaAlloc{upstream};
-        arenaAlloc.deleteObject(arena);
-    }
 }
 
 StreamReaderImpl::StreamReaderImpl(BoundedIOStream* stream_,
@@ -208,6 +94,27 @@ bool StreamReaderImpl::isLODConstrained() const {
     return lodConstrained;
 }
 
+void StreamReaderImpl::unload(DataLayer layer) {
+	if ((layer == DataLayer::All) || (layer == DataLayer::AllWithoutBlendShapes) || (layer == DataLayer::Descriptor))
+	{
+		dna = DNA{memRes};
+	}
+	else if ((layer == DataLayer::Geometry) || (layer == DataLayer::GeometryWithoutBlendShapes))
+	{
+		dna.unloadGeometry();
+	}
+	else if (layer == DataLayer::Behavior)
+	{
+		dna.unloadBehavior();
+	}
+	else if (layer == DataLayer::Definition)
+	{
+		dna.unloadGeometry();
+		dna.unloadBehavior();
+		dna.unloadDefinition();
+	}
+}
+
 void StreamReaderImpl::read() {
     // Due to possible usage of custom stream implementations, the status actually must be cleared at this point
     // as external streams do not have access to the status reset API
diff --git a/Engine/Plugins/Animation/RigLogic/Source/RigLogicLib/Private/dna/stream/StreamReaderImpl.h b/Engine/Plugins/Animation/RigLogic/Source/RigLogicLib/Private/dna/stream/StreamReaderImpl.h
--- a/Engine/Plugins/Animation/RigLogic/Source/RigLogicLib/Private/dna/stream/StreamReaderImpl.h
+++ b/Engine/Plugins/Animation/RigLogic/Source/RigLogicLib/Private/dna/stream/StreamReaderImpl.h
@@ -21,6 +21,7 @@ class StreamReaderImpl : public ReaderImpl<StreamReader> {
                          MemoryResource* memRes_);
         StreamReaderImpl(BoundedIOStream* stream_, DataLayer layer_, ConstArrayView<std::uint16_t> lods, MemoryResource* memRes_);
 
+        void unload(DataLayer layer) override;
         void read() override;
         bool isLODConstrained() const;
 
diff --git a/Engine/Plugins/Animation/RigLogic/Source/RigLogicLib/Private/riglogic/riglogic/RigInstanceImpl.cpp b/Engine/Plugins/Animation/RigLogic/Source/RigLogicLib/Private/riglogic/riglogic/RigInstanceImpl.cpp
--- a/Engine/Plugins/Animation/RigLogic/Source/RigLogicLib/Private/riglogic/riglogic/RigInstanceImpl.cpp
+++ b/Engine/Plugins/Animation/RigLogic/Source/RigLogicLib/Private/riglogic/riglogic/RigInstanceImpl.cpp
@@ -49,7 +49,7 @@ RigInstance::~RigInstance() = default;
 
 RigInstanceImpl::RigInstanceImpl(const RigMetrics& metrics, MemoryResource* memRes_) :
     memRes{memRes_},
-    lodCount{metrics.lodCount},
+    lodMaxLevel{metrics.lodCount > 0u ? static_cast<std::uint16_t>(metrics.lodCount - 1) : static_cast<std::uint16_t>(0)},
     lodLevel{},
     guiControlCount{metrics.guiControlCount},
     rawControlCount{metrics.rawControlCount},
@@ -94,7 +94,7 @@ std::uint16_t RigInstanceImpl::getLOD() const {
 }
 
 void RigInstanceImpl::setLOD(std::uint16_t level) {
-    lodLevel = extd::clamp(level, static_cast<std::uint16_t>(0), lodCount);
+    lodLevel = extd::clamp(level, static_cast<std::uint16_t>(0), lodMaxLevel);
 }
 
 ArrayView<float> RigInstanceImpl::getGUIControlValues() {
diff --git a/Engine/Plugins/Animation/RigLogic/Source/RigLogicLib/Private/riglogic/riglogic/RigInstanceImpl.h b/Engine/Plugins/Animation/RigLogic/Source/RigLogicLib/Private/riglogic/riglogic/RigInstanceImpl.h
--- a/Engine/Plugins/Animation/RigLogic/Source/RigLogicLib/Private/riglogic/riglogic/RigInstanceImpl.h
+++ b/Engine/Plugins/Animation/RigLogic/Source/RigLogicLib/Private/riglogic/riglogic/RigInstanceImpl.h
@@ -42,7 +42,7 @@ class RigInstanceImpl : public RigInstance {
     private:
         MemoryResource* memRes;
 
-        std::uint16_t lodCount;
+        std::uint16_t lodMaxLevel;
         std::uint16_t lodLevel;
         std::uint16_t guiControlCount;
         std::uint16_t rawControlCount;
diff --git a/Engine/Plugins/Animation/RigLogic/Source/RigLogicLib/Public/dna/Reader.h b/Engine/Plugins/Animation/RigLogic/Source/RigLogicLib/Public/dna/Reader.h
--- a/Engine/Plugins/Animation/RigLogic/Source/RigLogicLib/Public/dna/Reader.h
+++ b/Engine/Plugins/Animation/RigLogic/Source/RigLogicLib/Public/dna/Reader.h
@@ -3,6 +3,7 @@
 #pragma once
 
 #include "dna/Defs.h"
+#include "dna/DataLayer.h"
 #include "dna/layers/BehaviorReader.h"
 #include "dna/layers/GeometryReader.h"
 
@@ -20,7 +21,12 @@ namespace dna {
 class DNAAPI Reader : public BehaviorReader, public GeometryReader {
     public:
         ~Reader() override;
-
+		/**
+			@brief Unload all data of the specified layer and all layers dependent on it.
+			@param layer
+				Layer which data should be unloaded.
+		*/
+		virtual void unload(DataLayer layer) = 0;
 };
 
 }  // namespace dna
diff --git a/Engine/Plugins/Animation/RigLogic/Source/RigLogicLibTest/Private/rltests/dna/FakeReader.h b/Engine/Plugins/Animation/RigLogic/Source/RigLogicLibTest/Private/rltests/dna/FakeReader.h
--- a/Engine/Plugins/Animation/RigLogic/Source/RigLogicLibTest/Private/rltests/dna/FakeReader.h
+++ b/Engine/Plugins/Animation/RigLogic/Source/RigLogicLibTest/Private/rltests/dna/FakeReader.h
@@ -23,6 +23,9 @@ class FakeReader : public Reader {
     public:
         ~FakeReader();
 
+        void unload(DataLayer layer) override {
+        }
+
         // DescriptorReader methods start
         StringView getName() const override {
             return {};
diff --git a/Engine/Plugins/Animation/RigLogic/Source/RigLogicModule/Private/AnimNode_RigLogic.cpp b/Engine/Plugins/Animation/RigLogic/Source/RigLogicModule/Private/AnimNode_RigLogic.cpp
new file mode 100644
--- /dev/null
+++ b/Engine/Plugins/Animation/RigLogic/Source/RigLogicModule/Private/AnimNode_RigLogic.cpp
@@ -0,0 +1,235 @@
+// Copyright Epic Games, Inc. All Rights Reserved.
+
+#include "AnimNode_RigLogic.h"
+
+#include "DNAAsset.h"
+#include "DNAIndexMapping.h"
+#include "DNAReader.h"
+#include "RigLogic.h"
+#include "RigInstance.h"
+#include "SharedRigRuntimeContext.h"
+#include "Animation/AnimInstanceProxy.h"
+#include "HAL/LowLevelMemTracker.h"
+#include "Runtime/Engine/Classes/Animation/MorphTarget.h"
+
+LLM_DEFINE_TAG(Animation_RigLogic);
+DEFINE_LOG_CATEGORY(LogRigLogicAnimNode);
+
+static constexpr uint16 ATTR_COUNT_PER_JOINT = 9;
+
+FAnimNode_RigLogic::FAnimNode_RigLogic() : RigInstance(nullptr)
+{
+}
+
+FAnimNode_RigLogic::~FAnimNode_RigLogic()
+{
+	if (RigInstance != nullptr)
+	{
+		delete RigInstance;
+		RigInstance = nullptr;
+	}
+}
+
+void FAnimNode_RigLogic::Initialize_AnyThread(const FAnimationInitializeContext& Context)
+{
+	DECLARE_SCOPE_HIERARCHICAL_COUNTER_FUNC()
+	QUICK_SCOPE_CYCLE_COUNTER(STAT_AnimNode_RigLogic_Evaluate_AnyThread);
+	LLM_SCOPE_BYNAME(TEXT("Animation/RigLogic"));
+
+	FAnimNode_Base::Initialize_AnyThread(Context);
+	AnimSequence.Initialize(Context);
+	USkeletalMeshComponent* SkeletalMeshComponent = Context.AnimInstanceProxy->GetSkelMeshComponent();
+	USkeletalMesh* SkeletalMesh = SkeletalMeshComponent->SkeletalMesh;
+
+	static FRWLock GlobalCreationLock;
+	FRWScopeLock ScopeLock(GlobalCreationLock, SLT_ReadOnly);
+	UDNAIndexMapping* DNAIndexMappingContainer = Cast<UDNAIndexMapping>(SkeletalMesh->GetAssetUserDataOfClass(UDNAIndexMapping::StaticClass()));
+	if (DNAIndexMappingContainer == nullptr)
+	{
+		ScopeLock.ReleaseReadOnlyLockAndAcquireWriteLock_USE_WITH_CAUTION();
+		DNAIndexMappingContainer = Cast<UDNAIndexMapping>(SkeletalMesh->GetAssetUserDataOfClass(UDNAIndexMapping::StaticClass()));
+		if (DNAIndexMappingContainer == nullptr)
+		{
+			DNAIndexMappingContainer = NewObject<UDNAIndexMapping>();
+			SkeletalMesh->AddAssetUserData(DNAIndexMappingContainer);
+		}
+	}
+}
+
+void FAnimNode_RigLogic::CacheBones_AnyThread(const FAnimationCacheBonesContext& Context)
+{
+	DECLARE_SCOPE_HIERARCHICAL_COUNTER_FUNC()
+	QUICK_SCOPE_CYCLE_COUNTER(STAT_AnimNode_RigLogic_Evaluate_AnyThread);
+	LLM_SCOPE_BYNAME(TEXT("Animation/RigLogic"));
+
+	AnimSequence.CacheBones(Context);
+
+	// Initialize things that depend on the skeleton of Anim BP
+	USkeletalMeshComponent* SkeletalMeshComponent = Context.AnimInstanceProxy->GetSkelMeshComponent();
+	USkeletalMesh* SkeletalMesh = SkeletalMeshComponent->SkeletalMesh;
+	USkeleton* Skeleton = Context.AnimInstanceProxy->GetSkeleton();
+	UDNAAsset* DNAAsset = Cast<UDNAAsset>(SkeletalMesh->GetAssetUserDataOfClass(UDNAAsset::StaticClass()));
+	UDNAIndexMapping* DNAIndexMappingContainer = Cast<UDNAIndexMapping>(SkeletalMesh->GetAssetUserDataOfClass(UDNAIndexMapping::StaticClass()));
+	if (DNAAsset == nullptr)
+	{
+		return;
+	}
+
+	TSharedPtr<FSharedRigRuntimeContext> SharedRigRuntimeContext = DNAAsset->GetRigRuntimeContext(UDNAAsset::EDNARetentionPolicy::Unload);
+	if (!SharedRigRuntimeContext.IsValid())
+	{
+		return;
+	}
+
+	if (LocalRigRuntimeContext != SharedRigRuntimeContext)
+	{
+		LocalRigRuntimeContext = SharedRigRuntimeContext;
+		if (RigInstance != nullptr)
+		{
+			delete RigInstance;
+		}
+		RigInstance = new FRigInstance(LocalRigRuntimeContext->RigLogic.Get());
+	}
+
+	LocalDNAIndexMapping = DNAIndexMappingContainer->GetCachedMapping(LocalRigRuntimeContext->BehaviorReader.Get(), Skeleton, SkeletalMesh, SkeletalMeshComponent);
+	// CacheBones is called on LOD switches as well, in which case compact pose bone indices must be remapped
+	const FBoneContainer& RequiredBones = Context.AnimInstanceProxy->GetRequiredBones();
+	if (RequiredBones.IsValid())
+	{
+		RigInstance->SetLOD(Context.AnimInstanceProxy->GetLODLevel());
+		const uint16 CurrentLOD = RigInstance->GetLOD();
+		const TArray<uint16>& VariableJointIndices = DNAAsset->RigRuntimeContext->VariableJointIndicesPerLOD[CurrentLOD].Values;
+		JointsMapDNAIndicesToCompactPoseBoneIndices.Reset(VariableJointIndices.Num());
+		for (const uint16 JointIndex : VariableJointIndices)
+		{
+			const FMeshPoseBoneIndex MeshPoseBoneIndex = LocalDNAIndexMapping->JointsMapDNAIndicesToMeshPoseBoneIndices[JointIndex];
+			JointsMapDNAIndicesToCompactPoseBoneIndices.Add(RequiredBones.MakeCompactPoseIndex(MeshPoseBoneIndex));
+		}
+	}
+}
+
+void FAnimNode_RigLogic::Update_AnyThread(const FAnimationUpdateContext& Context)
+{
+	DECLARE_SCOPE_HIERARCHICAL_COUNTER_FUNC()
+	QUICK_SCOPE_CYCLE_COUNTER(STAT_AnimNode_RigLogic_Update_AnyThread);
+
+	GetEvaluateGraphExposedInputs().Execute(Context);
+	AnimSequence.Update(Context);
+}
+
+void FAnimNode_RigLogic::Evaluate_AnyThread(FPoseContext& OutputContext)
+{
+	DECLARE_SCOPE_HIERARCHICAL_COUNTER_FUNC()
+	QUICK_SCOPE_CYCLE_COUNTER(STAT_AnimNode_RigLogic_Evaluate_AnyThread);
+
+	if (!LocalRigRuntimeContext.IsValid())
+	{
+		return;
+	}
+	// Initialize things that depend on the skeleton of Anim BP
+	AnimSequence.Evaluate(OutputContext);
+	UpdateControlCurves(OutputContext, LocalDNAIndexMapping.Get());
+	CalculateRigLogic(LocalRigRuntimeContext->RigLogic.Get());
+	const uint16 CurrentLOD = RigInstance->GetLOD();
+	TArrayView<const uint16> VariableJointIndices = LocalRigRuntimeContext->VariableJointIndicesPerLOD[CurrentLOD].Values;
+	UpdateJoints(VariableJointIndices, LocalRigRuntimeContext->RigLogic->GetRawNeutralJointValues(), RigInstance->GetRawJointOutputs(), OutputContext);
+	UpdateBlendShapeCurves(LocalDNAIndexMapping.Get(), RigInstance->GetBlendShapeOutputs(), OutputContext);
+	UpdateAnimMapCurves(LocalDNAIndexMapping.Get(), RigInstance->GetAnimatedMapOutputs(), OutputContext);
+}
+
+void FAnimNode_RigLogic::GatherDebugData(FNodeDebugData& DebugData)
+{
+	AnimSequence.GatherDebugData(DebugData);
+}
+
+void FAnimNode_RigLogic::UpdateControlCurves(const FPoseContext& InputContext, const FDNAIndexMapping* DNAIndexMapping)
+{
+	DECLARE_SCOPE_HIERARCHICAL_COUNTER_FUNC()
+	QUICK_SCOPE_CYCLE_COUNTER(STAT_AnimNode_RigLogic_UpdateControlCurves);
+
+	const int32 RawControlCount = RigInstance->GetRawControlCount();
+	for (int32 ControlIndex = 0; ControlIndex < RawControlCount; ++ControlIndex)
+	{
+		const int32 ControlCurveUID = DNAIndexMapping->ControlAttributesMapDNAIndicesToUEUIDs[ControlIndex];
+		if (ControlCurveUID != INDEX_NONE)
+		{
+			const float Value = InputContext.Curve.Get(static_cast<SmartName::UID_Type>(ControlCurveUID));
+			RigInstance->SetRawControl(ControlIndex, Value);
+		}
+	}
+}
+
+void FAnimNode_RigLogic::CalculateRigLogic(FRigLogic* RigLogic)
+{
+	DECLARE_SCOPE_HIERARCHICAL_COUNTER_FUNC()
+	QUICK_SCOPE_CYCLE_COUNTER(STAT_AnimNode_RigLogic_CalculateRigLogic);
+	RigLogic->Calculate(RigInstance);
+}
+
+void FAnimNode_RigLogic::UpdateJoints(TArrayView<const uint16> VariableJointIndices, TArrayView<const float> NeutralJointValues, TArrayView<const float> DeltaJointValues, FPoseContext& OutputContext)
+{
+	DECLARE_SCOPE_HIERARCHICAL_COUNTER_FUNC()
+	QUICK_SCOPE_CYCLE_COUNTER(STAT_AnimNode_RigLogic_UpdateJoints);
+
+	const float* N = NeutralJointValues.GetData();
+	const float* D = DeltaJointValues.GetData();
+	check(VariableJointIndices.Num() == JointsMapDNAIndicesToCompactPoseBoneIndices.Num());
+	for (int32 MappingIndex = 0; MappingIndex < JointsMapDNAIndicesToCompactPoseBoneIndices.Num(); ++MappingIndex)
+	{
+		const uint16 JointIndex = VariableJointIndices[MappingIndex];
+		const FCompactPoseBoneIndex CompactPoseBoneIndex = JointsMapDNAIndicesToCompactPoseBoneIndices[MappingIndex];
+		if (CompactPoseBoneIndex != INDEX_NONE)
+		{
+			const uint16 AttrIndex = JointIndex * ATTR_COUNT_PER_JOINT;
+			FTransform& CompactPose = OutputContext.Pose[CompactPoseBoneIndex];
+			// Translation: X = X, Y = -Y, Z = Z
+			CompactPose.SetTranslation(FVector((N[AttrIndex + 0] + D[AttrIndex + 0]), -(N[AttrIndex + 1] + D[AttrIndex + 1]), (N[AttrIndex + 2] + D[AttrIndex + 2])));
+			// Rotation: X = -Y, Y = -Z, Z = X
+			CompactPose.SetRotation(FQuat(FRotator(-N[AttrIndex + 4], -N[AttrIndex + 5], N[AttrIndex + 3])) * FQuat(FRotator(-D[AttrIndex + 4], -D[AttrIndex + 5], D[AttrIndex + 3])));
+			// Scale: X = X, Y = Y, Z = Z
+			CompactPose.SetScale3D(FVector((N[AttrIndex + 6] + D[AttrIndex + 6]), (N[AttrIndex + 7] + D[AttrIndex + 7]), (N[AttrIndex + 8] + D[AttrIndex + 8])));
+		}
+	}
+}
+
+void FAnimNode_RigLogic::UpdateBlendShapeCurves(const FDNAIndexMapping* DNAIndexMapping, TArrayView<const float> BlendShapeValues, FPoseContext& OutputContext)
+{
+	DECLARE_SCOPE_HIERARCHICAL_COUNTER_FUNC()
+	QUICK_SCOPE_CYCLE_COUNTER(STAT_AnimNode_RigLogic_UpdateBlendShapeCurves);
+
+	const uint16 LOD = RigInstance->GetLOD();
+	const TArray<int32>& BlendShapeIndices = DNAIndexMapping->BlendShapeIndicesPerLOD[LOD].Values;
+	const TArray<int32>& MorphTargetIndices = DNAIndexMapping->MorphTargetIndicesPerLOD[LOD].Values;
+	check(BlendShapeIndices.Num() == MorphTargetIndices.Num());
+	for (uint32 MappingIndex = 0; MappingIndex < static_cast<uint32>(BlendShapeIndices.Num()); ++MappingIndex)
+	{
+		const int32 BlendShapeIndex = BlendShapeIndices[MappingIndex];
+		const int32 MorphTargetCurveUID = MorphTargetIndices[MappingIndex];
+		if (MorphTargetCurveUID != INDEX_NONE)
+		{
+			const float Value = BlendShapeValues[BlendShapeIndex];
+			OutputContext.Curve.Set(static_cast<SmartName::UID_Type>(MorphTargetCurveUID), Value);
+		}
+	}
+}
+
+void FAnimNode_RigLogic::UpdateAnimMapCurves(const FDNAIndexMapping* DNAIndexMapping, TArrayView<const float> AnimMapOutputs, FPoseContext& OutputContext)
+{
+	DECLARE_SCOPE_HIERARCHICAL_COUNTER_FUNC()
+	QUICK_SCOPE_CYCLE_COUNTER(STAT_AnimNode_RigLogic_UpdateAnimMapCurves);
+
+	const uint16 LOD = RigInstance->GetLOD();
+	const TArray<int32>& AnimatedMapIndices = DNAIndexMapping->AnimatedMapIndicesPerLOD[LOD].Values;
+	const TArray<int32>& MaskMultiplierIndices = DNAIndexMapping->MaskMultiplierIndicesPerLOD[LOD].Values;
+	check(AnimatedMapIndices.Num() == MaskMultiplierIndices.Num());
+	for (uint32 MappingIndex = 0; MappingIndex < static_cast<uint32>(AnimatedMapIndices.Num()); ++MappingIndex)
+	{
+		const int32 AnimatedMapIndex = AnimatedMapIndices[MappingIndex];
+		const int32 MaskMultiplierUID = MaskMultiplierIndices[MappingIndex];
+		if (MaskMultiplierUID != INDEX_NONE)
+		{
+			const float Value = AnimMapOutputs[AnimatedMapIndex];
+			OutputContext.Curve.Set(static_cast<SmartName::UID_Type>(MaskMultiplierUID), Value);
+		}
+	}
+}
diff --git a/Engine/Plugins/Animation/RigLogic/Source/RigLogicModule/Private/DNAAsset.cpp b/Engine/Plugins/Animation/RigLogic/Source/RigLogicModule/Private/DNAAsset.cpp
--- a/Engine/Plugins/Animation/RigLogic/Source/RigLogicModule/Private/DNAAsset.cpp
+++ b/Engine/Plugins/Animation/RigLogic/Source/RigLogicModule/Private/DNAAsset.cpp
@@ -7,11 +7,11 @@
 #include "DNAAssetCustomVersion.h"
 #include "DNAReaderAdapter.h"
 #include "FMemoryResource.h"
-#include "Hasher.h"
 #include "RigLogicMemoryStream.h"
+#include "SharedRigRuntimeContext.h"
 
 #if WITH_EDITORONLY_DATA
-#include "EditorFramework/AssetImportData.h"
+    #include "EditorFramework/AssetImportData.h"
 #endif
 #include "Engine/AssetUserData.h"
 #include "Serialization/BufferArchive.h"
@@ -19,6 +19,7 @@
 #include "Misc/Paths.h"
 #include "Misc/FileHelper.h"
 #include "UObject/UObjectGlobals.h"
+#include "HAL/LowLevelMemTracker.h"
 
 #include "riglogic/RigLogic.h"
 
@@ -26,7 +27,7 @@ DEFINE_LOG_CATEGORY(LogDNAAsset);
 
 static constexpr uint32 AVG_EMPTY_SIZE = 4 * 1024;
 static constexpr uint32 AVG_BEHAVIOR_SIZE = 5 * 1024 * 1024;
-static constexpr uint32 AVG_GEOMETRY_SIZE = 150 * 1024 * 1024;
+static constexpr uint32 AVG_GEOMETRY_SIZE = 50 * 1024 * 1024;
 
 static TSharedPtr<IDNAReader> ReadDNAFromStream(rl4::BoundedIOStream* Stream, EDNADataLayer Layer, uint16 MaxLOD)
 {
@@ -80,74 +81,120 @@ static TSharedPtr<IDNAReader> CreateEmptyDNA(uint32 PredictedSize)
 	return ReadDNAFromBuffer(&MemoryBuffer, EDNADataLayer::All);
 }
 
-static void InvalidateSharedRigRuntimeContext(FSharedRigRuntimeContext& Context)
+UDNAAsset::UDNAAsset() : RigRuntimeContext{nullptr}
 {
-	Context.RigLogic = nullptr;
-	Context.VariableJointIndices.Reset();
-	Context.DNAHash = 0u;
 }
 
-static uint32 HashDescriptor(const IDescriptorReader* Reader)
+UDNAAsset::~UDNAAsset() = default;
+
+TSharedPtr<IBehaviorReader> UDNAAsset::GetBehaviorReader()
+{
+	FScopeLock ScopeLock{&DNAUpdateSection};
+	return BehaviorReader;
+}
+
+#if WITH_EDITORONLY_DATA
+TSharedPtr<IGeometryReader> UDNAAsset::GetGeometryReader()
 {
-	Hasher DescriptorHasher;
-	DescriptorHasher.Update(Reader->GetName());
-	DescriptorHasher.Update(Reader->GetArchetype());
-	DescriptorHasher.Update(Reader->GetGender());
-	DescriptorHasher.Update(Reader->GetAge());
-	DescriptorHasher.Update(Reader->GetMetaDataCount());
-	DescriptorHasher.Update(Reader->GetTranslationUnit());
-	DescriptorHasher.Update(Reader->GetRotationUnit());
-	DescriptorHasher.Update(Reader->GetCoordinateSystem());
-	DescriptorHasher.Update(Reader->GetLODCount());
-	DescriptorHasher.Update(Reader->GetDBMaxLOD());
-	DescriptorHasher.Update(Reader->GetDBComplexity());
-	DescriptorHasher.Update(Reader->GetDBName());
-	return DescriptorHasher.GetHash();
+	FScopeLock ScopeLock{&DNAUpdateSection};
+	return GeometryReader;
 }
+#endif
 
-UDNAAsset::UDNAAsset() :
-	Context{}
+void UDNAAsset::SetBehaviorReader(TSharedPtr<IDNAReader> SourceDNAReader)
 {
+	FScopeLock ScopeLock{&DNAUpdateSection};
+	BehaviorReader = CopyDNALayer(SourceDNAReader.Get(), EDNADataLayer::Behavior, AVG_BEHAVIOR_SIZE);
+	InvalidateRigRuntimeContext();
+}
+
+void UDNAAsset::SetGeometryReader(TSharedPtr<IDNAReader> SourceDNAReader)
+{
+#if WITH_EDITORONLY_DATA
+	FScopeLock ScopeLock{&DNAUpdateSection};
+	GeometryReader = CopyDNALayer(SourceDNAReader.Get(), EDNADataLayer::Geometry, AVG_GEOMETRY_SIZE);
+#endif // #if WITH_EDITORONLY_DATA
+}
+
+void UDNAAsset::InvalidateRigRuntimeContext()
+{
+	FWriteScopeLock ScopeLock{RigRuntimeContextUpdateLock};
+	RigRuntimeContext = nullptr;
+}
+
+TSharedPtr<FSharedRigRuntimeContext> UDNAAsset::GetRigRuntimeContext(EDNARetentionPolicy Policy)
+{
+	LLM_SCOPE_BYNAME(TEXT("Animation/RigLogic"));
+
+	FRWScopeLock ContextScopeLock(RigRuntimeContextUpdateLock, SLT_ReadOnly);
+	if (!RigRuntimeContext.IsValid())
+	{
+		ContextScopeLock.ReleaseReadOnlyLockAndAcquireWriteLock_USE_WITH_CAUTION();
+		if (!RigRuntimeContext.IsValid())
+		{
+			TSharedPtr<FSharedRigRuntimeContext> NewContext = MakeShared<FSharedRigRuntimeContext>();
+			FScopeLock DNAScopeLock{&DNAUpdateSection};
+			if (BehaviorReader.IsValid() && (BehaviorReader->GetJointCount() != 0))
+			{
+				NewContext->BehaviorReader = BehaviorReader;
+				NewContext->RigLogic = MakeShared<FRigLogic>(BehaviorReader.Get());
+				NewContext->CacheVariableJointIndices();
+				RigRuntimeContext = NewContext;
+#if !WITH_EDITOR
+				if (Policy == EDNARetentionPolicy::Unload)
+				{
+					BehaviorReader->Unload(EDNADataLayer::Behavior);
+					BehaviorReader->Unload(EDNADataLayer::Geometry);
+				}
+#endif  // !WITH_EDITOR
+			}
+		}
+	}
+	return RigRuntimeContext;
 }
 
 bool UDNAAsset::Init(const FString& DNAFilename)
 {
+	LLM_SCOPE_BYNAME(TEXT("Animation/RigLogic"));
+
 	if (!rl4::Status::isOk())
 	{
 		UE_LOG(LogDNAAsset, Warning, TEXT("%s"), ANSI_TO_TCHAR(rl4::Status::get().message));
 	}
 
 	this->DNAFileName = DNAFilename; //memorize for re-import
-
+	
 	if (!FPaths::FileExists(DNAFilename))
 	{
 		UE_LOG(LogDNAAsset, Error, TEXT("DNA file %s doesn't exist!"), *DNAFilename);
 		return false;
 	}
-
+	
 	// Temporary buffer for the DNA file
 	TArray<uint8> TempFileBuffer;
-
+	
 	if (!FFileHelper::LoadFileToArray(TempFileBuffer, *DNAFilename)) //load entire DNA file into the array
 	{
 		UE_LOG(LogDNAAsset, Error, TEXT("Couldn't read DNA file %s!"), *DNAFilename);
 		return false;
 	}
+	
+	FScopeLock ScopeLock{&DNAUpdateSection};
 
 	// Load run-time data (behavior) from whole-DNA buffer into BehaviorReader
-	InvalidateSharedRigRuntimeContext(Context);
-	Context.BehaviorReader = ReadDNAFromBuffer(&TempFileBuffer, EDNADataLayer::Behavior, 0u); //0u = MaxLOD
-	if (!Context.BehaviorReader.IsValid())
+	BehaviorReader = ReadDNAFromBuffer(&TempFileBuffer, EDNADataLayer::Behavior, 0u); //0u = MaxLOD
+	if (!BehaviorReader.IsValid())
 	{
 		return false;
 	}
-	Context.DNAHash = HashDescriptor(Context.BehaviorReader.Get());
+
+	InvalidateRigRuntimeContext();
 
 #if WITH_EDITORONLY_DATA
 	//We use geometry part of the data in MHC only (for updating the SkeletalMesh with
 	//result of GeneSplicer), so we can drop geometry part when cooking for runtime
-	Context.GeometryReader = ReadDNAFromBuffer(&TempFileBuffer, EDNADataLayer::Geometry, 0u); //0u = MaxLOD
-	if (!Context.GeometryReader.IsValid())
+	GeometryReader = ReadDNAFromBuffer(&TempFileBuffer, EDNADataLayer::Geometry, 0u); //0u = MaxLOD
+	if (!GeometryReader.IsValid())
 	{
 		return false;
 	}
@@ -162,59 +209,45 @@ bool UDNAAsset::Init(const FString& DNAFilename)
 
 void UDNAAsset::Serialize(FArchive& Ar)
 {
+	LLM_SCOPE_BYNAME(TEXT("Animation/RigLogic"));
+
 	Super::Serialize(Ar);
 
 	Ar.UsingCustomVersion(FDNAAssetCustomVersion::GUID);
 
+	FScopeLock ScopeLock{&DNAUpdateSection};
+
 	if (Ar.CustomVer(FDNAAssetCustomVersion::GUID) >= FDNAAssetCustomVersion::BeforeCustomVersionWasAdded)
 	{
-		LLM_SCOPE(ELLMTag::SkeletalMesh);
-
 		if (Ar.IsLoading())
 		{
-			InvalidateSharedRigRuntimeContext(Context);
-
-			FArchiveMemoryStream BehaviorStream{ &Ar };
-			Context.BehaviorReader = ReadDNAFromStream(&BehaviorStream, EDNADataLayer::Behavior, 0u); //0u = max LOD
-			Context.DNAHash = HashDescriptor(Context.BehaviorReader.Get());
-
+			FArchiveMemoryStream BehaviorStream{&Ar};
+			BehaviorReader = ReadDNAFromStream(&BehaviorStream, EDNADataLayer::Behavior, 0u); //0u = max LOD
 			// Geometry data is always present (even if only as an empty placeholder), just so the uasset
 			// format remains consistent between editor and non-editor builds
-			FArchiveMemoryStream GeometryStream{ &Ar };
+			FArchiveMemoryStream GeometryStream{&Ar};
 			auto Reader = ReadDNAFromStream(&GeometryStream, EDNADataLayer::Geometry, 0u); //0u = max LOD
 #if WITH_EDITORONLY_DATA
 			// Geometry data is discarded unless in Editor
-			Context.GeometryReader = Reader;
+			GeometryReader = Reader;
 #endif // #if WITH_EDITORONLY_DATA
+
+			InvalidateRigRuntimeContext();
 		}
 
 		if (Ar.IsSaving())
 		{
 			TSharedPtr<IDNAReader> EmptyDNA = CreateEmptyDNA(AVG_EMPTY_SIZE);
-			IDNAReader* BehaviorReaderPtr = (Context.BehaviorReader.IsValid() ? static_cast<IDNAReader*>(Context.BehaviorReader.Get()) : EmptyDNA.Get());
-			FArchiveMemoryStream BehaviorStream{ &Ar };
+			IDNAReader* BehaviorReaderPtr = (BehaviorReader.IsValid() ? static_cast<IDNAReader*>(BehaviorReader.Get()) : EmptyDNA.Get());
+			FArchiveMemoryStream BehaviorStream{&Ar};
 			WriteDNAToStream(BehaviorReaderPtr, EDNADataLayer::Behavior, &BehaviorStream);
 
 			// When cooking (or when there was no Geometry data available), an empty DNA structure is written
 			// into the stream, serving as a placeholder just so uasset files can be conveniently loaded
 			// regardless if they were cooked or prepared for in-editor work
-			IDNAReader* GeometryReaderPtr = (Context.GeometryReader.IsValid() && !Ar.IsCooking() ? static_cast<IDNAReader*>(Context.GeometryReader.Get()) : EmptyDNA.Get());
-			FArchiveMemoryStream GeometryStream{ &Ar };
+			IDNAReader* GeometryReaderPtr = (GeometryReader.IsValid() && !Ar.IsCooking() ? static_cast<IDNAReader*>(GeometryReader.Get()) : EmptyDNA.Get());
+			FArchiveMemoryStream GeometryStream{&Ar};
 			WriteDNAToStream(GeometryReaderPtr, EDNADataLayer::Geometry, &GeometryStream);
 		}
 	}
 }
-
-void UDNAAsset::SetBehaviorReader(TSharedPtr<IDNAReader> SourceDNAReader)
-{
-	InvalidateSharedRigRuntimeContext(Context);
-	Context.BehaviorReader = CopyDNALayer(SourceDNAReader.Get(), EDNADataLayer::Behavior, AVG_BEHAVIOR_SIZE);
-	Context.DNAHash = HashDescriptor(Context.BehaviorReader.Get());
-}
-
-void UDNAAsset::SetGeometryReader(TSharedPtr<IDNAReader> SourceDNAReader)
-{
-#if WITH_EDITORONLY_DATA
-	Context.GeometryReader = CopyDNALayer(SourceDNAReader.Get(), EDNADataLayer::Geometry, AVG_GEOMETRY_SIZE);
-#endif // #if WITH_EDITORONLY_DATA
-}
diff --git a/Engine/Plugins/Animation/RigLogic/Source/RigLogicModule/Private/DNAIndexMapping.cpp b/Engine/Plugins/Animation/RigLogic/Source/RigLogicModule/Private/DNAIndexMapping.cpp
new file mode 100644
--- /dev/null
+++ b/Engine/Plugins/Animation/RigLogic/Source/RigLogicModule/Private/DNAIndexMapping.cpp
@@ -0,0 +1,188 @@
+// Copyright Epic Games, Inc. All Rights Reserved.
+
+#include "DNAIndexMapping.h"
+
+#include "CoreMinimal.h"
+#include "HAL/LowLevelMemTracker.h"
+#include "DNAReader.h"
+#include "Hasher.h"
+
+static uint32 HashDNA(const IDefinitionReader* Reader)
+{
+	Hasher DNAHasher;
+	DNAHasher.Update(Reader->GetTranslationUnit());
+	DNAHasher.Update(Reader->GetRotationUnit());
+	DNAHasher.Update(Reader->GetCoordinateSystem());
+	DNAHasher.Update(Reader->GetLODCount());
+	DNAHasher.Update(Reader->GetDBMaxLOD());
+	DNAHasher.Update(Reader->GetDBComplexity());
+	DNAHasher.Update(Reader->GetDBName());
+	DNAHasher.Update(Reader->GetJointCount());
+	DNAHasher.Update(Reader->GetBlendShapeChannelCount());
+	DNAHasher.Update(Reader->GetAnimatedMapCount());
+	return DNAHasher.GetHash();
+}
+
+/** Constructs curve name from nameToSplit using formatString of form x<obj>y<attr>z **/
+static FString CreateCurveName(const FString& NameToSplit, const FString& FormatString)
+{
+	// constructs curve name from NameToSplit (always in form <obj>.<attr>)
+	// using FormatString of form x<obj>y<attr>z
+	// where x, y and z are arbitrary strings
+	// example:
+	// FormatString="mesh_<obj>_<attr>"
+	// 'head.blink_L' becomes 'mesh_head_blink_L'
+	FString ObjectName, AttributeName;
+	if (!NameToSplit.Split(".", &ObjectName, &AttributeName))
+	{
+		return TEXT("");
+	}
+	FString CurveName = FormatString;
+	CurveName = CurveName.Replace(TEXT("<obj>"), *ObjectName);
+	CurveName = CurveName.Replace(TEXT("<attr>"), *AttributeName);
+	return CurveName;
+}
+
+void FDNAIndexMapping::MapControlCurves(const IBehaviorReader* DNABehavior, const USkeleton* Skeleton)
+{
+	LLM_SCOPE_BYNAME(TEXT("Animation/RigLogic"));
+
+	const uint32 ControlCount = DNABehavior->GetRawControlCount();
+	ControlAttributesMapDNAIndicesToUEUIDs.Reset(ControlCount);
+	for (uint32_t ControlIndex = 0; ControlIndex < ControlCount; ++ControlIndex)
+	{
+		const FString DNAControlName = DNABehavior->GetRawControlName(ControlIndex);
+		const FString AnimatedControlName = CreateCurveName(DNAControlName, TEXT("<obj>_<attr>"));
+		if (AnimatedControlName == TEXT(""))
+		{
+			return;
+		}
+		const FName ControlFName(*AnimatedControlName);
+		FSmartName ControlCurveName;
+		const bool IsControlFound = Skeleton->GetSmartNameByName(USkeleton::AnimCurveMappingName, ControlFName, ControlCurveName);
+		// SmartName::UID_Type is uint32; we use int32 to allow for NONE value
+		const int32 UID = (IsControlFound ? static_cast<int32>(ControlCurveName.UID) : INDEX_NONE);
+		ControlAttributesMapDNAIndicesToUEUIDs.Add(UID);
+	}
+}
+
+void FDNAIndexMapping::MapJoints(const IBehaviorReader* DNABehavior, const USkeletalMeshComponent* SkeletalMeshComponent)
+{
+	LLM_SCOPE_BYNAME(TEXT("Animation/RigLogic"));
+
+	const uint16 JointCount = DNABehavior->GetJointCount();
+	JointsMapDNAIndicesToMeshPoseBoneIndices.Reset(JointCount);
+	for (uint16 JointIndex = 0; JointIndex < JointCount; ++JointIndex)
+	{
+		const FString JointName = DNABehavior->GetJointName(JointIndex);
+		const FName BoneName = FName(*JointName);
+		const int32 BoneIndex = SkeletalMeshComponent->GetBoneIndex(BoneName);
+		// BoneIndex may be INDEX_NONE, but it's handled properly by the Evaluate method
+		JointsMapDNAIndicesToMeshPoseBoneIndices.Add(FMeshPoseBoneIndex{BoneIndex});
+	}
+}
+
+void FDNAIndexMapping::MapMorphTargets(const IBehaviorReader* DNABehavior, const USkeleton* Skeleton, const USkeletalMesh* SkeletalMesh)
+{
+	LLM_SCOPE_BYNAME(TEXT("Animation/RigLogic"));
+
+	const uint16 LODCount = DNABehavior->GetLODCount();
+	const TMap<FName, int32>& MorphTargetIndexMap = SkeletalMesh->GetMorphTargetIndexMap();
+	const TArray<UMorphTarget*>& MorphTargets = SkeletalMesh->GetMorphTargets();
+
+	BlendShapeIndicesPerLOD.Reset(LODCount);
+	BlendShapeIndicesPerLOD.AddDefaulted(LODCount);
+	MorphTargetIndicesPerLOD.Reset(LODCount);
+	MorphTargetIndicesPerLOD.AddDefaulted(LODCount);
+
+	for (uint16 LODIndex = 0; LODIndex < LODCount; ++LODIndex)
+	{
+		TArrayView<const uint16> MappingIndicesForLOD = DNABehavior->GetMeshBlendShapeChannelMappingIndicesForLOD(LODIndex);
+		MorphTargetIndicesPerLOD[LODIndex].Values.Reserve(MappingIndicesForLOD.Num());
+		BlendShapeIndicesPerLOD[LODIndex].Values.Reserve(MappingIndicesForLOD.Num());
+		for (uint16 MappingIndex : MappingIndicesForLOD)
+		{
+			const FMeshBlendShapeChannelMapping Mapping = DNABehavior->GetMeshBlendShapeChannelMapping(MappingIndex);
+			const FString MeshName = DNABehavior->GetMeshName(Mapping.MeshIndex);
+			const FString BlendShapeName = DNABehavior->GetBlendShapeChannelName(Mapping.BlendShapeChannelIndex);
+			const FString MorphTargetStr = MeshName + TEXT("__") + BlendShapeName;
+			const FName MorphTargetName(*MorphTargetStr);
+			const int32* MorphTargetIndex = MorphTargetIndexMap.Find(MorphTargetName);
+			int32 UID = INDEX_NONE;
+			if ((MorphTargetIndex != nullptr) && (*MorphTargetIndex != INDEX_NONE))
+			{
+				const UMorphTarget* MorphTarget = MorphTargets[*MorphTargetIndex];
+				UID = Skeleton->GetUIDByName(USkeleton::AnimCurveMappingName, MorphTarget->GetFName());
+			}
+			BlendShapeIndicesPerLOD[LODIndex].Values.Add(Mapping.BlendShapeChannelIndex);
+			MorphTargetIndicesPerLOD[LODIndex].Values.Add(UID);
+		}
+	}
+}
+
+void FDNAIndexMapping::MapMaskMultipliers(const IBehaviorReader* DNABehavior, const USkeleton* Skeleton)
+{
+	LLM_SCOPE_BYNAME(TEXT("Animation/RigLogic"));
+
+	const uint16 LODCount = DNABehavior->GetLODCount();
+
+	AnimatedMapIndicesPerLOD.Reset();
+	AnimatedMapIndicesPerLOD.AddDefaulted(LODCount);
+	MaskMultiplierIndicesPerLOD.Reset();
+	MaskMultiplierIndicesPerLOD.AddDefaulted(LODCount);
+
+	for (uint16 LODIndex = 0; LODIndex < LODCount; ++LODIndex)
+	{
+		TArrayView<const uint16> IndicesPerLOD = DNABehavior->GetAnimatedMapIndicesForLOD(LODIndex);
+		AnimatedMapIndicesPerLOD[LODIndex].Values.Reserve(IndicesPerLOD.Num());
+		MaskMultiplierIndicesPerLOD[LODIndex].Values.Reserve(IndicesPerLOD.Num());
+		for (uint16 AnimMapIndex : IndicesPerLOD)
+		{
+			const FString AnimatedMapName = DNABehavior->GetAnimatedMapName(AnimMapIndex);
+			const FString MaskMultiplierNameStr = CreateCurveName(AnimatedMapName, TEXT("<obj>_<attr>"));
+			if (MaskMultiplierNameStr == "")
+			{
+				return;
+			}
+			const FName MaskMultiplierName(*MaskMultiplierNameStr);
+			FSmartName CurveName;
+			const bool IsControlFound = Skeleton->GetSmartNameByName(USkeleton::AnimCurveMappingName, MaskMultiplierName, CurveName);
+			const int32 UID = (IsControlFound ? CurveName.UID : INDEX_NONE);
+			AnimatedMapIndicesPerLOD[LODIndex].Values.Add(AnimMapIndex);
+			MaskMultiplierIndicesPerLOD[LODIndex].Values.Add(UID);
+		}
+	}
+}
+
+UDNAIndexMapping::UDNAIndexMapping() : Cached{nullptr}
+{
+}
+
+TSharedPtr<FDNAIndexMapping> UDNAIndexMapping::GetCachedMapping(const IBehaviorReader* DNABehavior,
+																const USkeleton* Skeleton,
+																const USkeletalMesh* SkeletalMesh,
+																const USkeletalMeshComponent* SkeletalMeshComponent)
+{
+	LLM_SCOPE_BYNAME(TEXT("Animation/RigLogic"));
+
+	const FGuid SkeletonGuid = Skeleton->GetGuid();
+	const uint32 DNAHash = HashDNA(DNABehavior);
+	// If GUID or DNA hash is different, all skeleton <-> DNA mappings should be recreated
+	FRWScopeLock ScopeLock(Lock, SLT_ReadOnly);
+	if (!Cached.IsValid() || (Cached->SkeletonGuid != SkeletonGuid) || (Cached->DNAHash != DNAHash))
+	{
+		ScopeLock.ReleaseReadOnlyLockAndAcquireWriteLock_USE_WITH_CAUTION();
+		if (!Cached.IsValid() || (Cached->SkeletonGuid != SkeletonGuid) || (Cached->DNAHash != DNAHash))
+		{
+			TSharedPtr<FDNAIndexMapping> NewMapping = MakeShared<FDNAIndexMapping>();
+			NewMapping->SkeletonGuid = SkeletonGuid;
+			NewMapping->DNAHash = DNAHash;
+			NewMapping->MapControlCurves(DNABehavior, Skeleton);
+			NewMapping->MapJoints(DNABehavior, SkeletalMeshComponent);
+			NewMapping->MapMorphTargets(DNABehavior, Skeleton, SkeletalMesh);
+			NewMapping->MapMaskMultipliers(DNABehavior, Skeleton);
+			Cached = NewMapping;
+		}
+	}
+	return Cached;
+}
diff --git a/Engine/Plugins/Animation/RigLogic/Source/RigLogicModule/Private/DNAIndexMapping.h b/Engine/Plugins/Animation/RigLogic/Source/RigLogicModule/Private/DNAIndexMapping.h
new file mode 100644
--- /dev/null
+++ b/Engine/Plugins/Animation/RigLogic/Source/RigLogicModule/Private/DNAIndexMapping.h
@@ -0,0 +1,56 @@
+// Copyright Epic Games, Inc. All Rights Reserved.
+
+#pragma once
+
+#include "DNAReader.h"
+
+#include "CoreMinimal.h"
+#include "UObject/Object.h"
+#include "UObject/ObjectMacros.h"
+#include "Engine/AssetUserData.h"
+#include "Interfaces/Interface_AssetUserData.h"
+
+#include "DNAIndexMapping.generated.h"
+
+struct FDNAIndexMapping
+{
+	template <typename T>
+	struct TArrayWrapper
+	{
+		TArray<T> Values;
+	};
+
+	uint32 DNAHash;
+	FGuid SkeletonGuid;
+	// For maps, we use int32 instead of SmartName::UID_Ttype directly to allow storing INDEX_NONE for missing elements
+	// if value is valid, it is cast to SmartName::UID_Type
+	TArray<int32> ControlAttributesMapDNAIndicesToUEUIDs;
+	TArray<FMeshPoseBoneIndex> JointsMapDNAIndicesToMeshPoseBoneIndices;
+	TArray<TArrayWrapper<int32>> BlendShapeIndicesPerLOD;
+	TArray<TArrayWrapper<int32>> MorphTargetIndicesPerLOD;
+	TArray<TArrayWrapper<int32>> AnimatedMapIndicesPerLOD;
+	TArray<TArrayWrapper<int32>> MaskMultiplierIndicesPerLOD;
+
+	void MapControlCurves(const IBehaviorReader* DNABehavior, const USkeleton* Skeleton);
+	void MapJoints(const IBehaviorReader* DNABehavior, const USkeletalMeshComponent* SkeletalMeshComponent);
+	void MapMorphTargets(const IBehaviorReader* DNABehavior, const USkeleton* Skeleton, const USkeletalMesh* SkeletalMesh);
+	void MapMaskMultipliers(const IBehaviorReader* DNABehavior, const USkeleton* Skeleton);
+
+};
+
+UCLASS(NotBlueprintable, hidecategories = (Object))
+class UDNAIndexMapping : public UAssetUserData
+{
+	GENERATED_BODY()
+
+public:
+	UDNAIndexMapping();
+	TSharedPtr<FDNAIndexMapping> GetCachedMapping(const IBehaviorReader* DNABehavior,
+												  const USkeleton* Skeleton,
+												  const USkeletalMesh* SkeletalMesh,
+												  const USkeletalMeshComponent* SkeletalMeshComponent);
+
+private:
+	FRWLock Lock;
+	TSharedPtr<FDNAIndexMapping> Cached;
+};
diff --git a/Engine/Plugins/Animation/RigLogic/Source/RigLogicModule/Private/DNAUtils.cpp b/Engine/Plugins/Animation/RigLogic/Source/RigLogicModule/Private/DNAUtils.cpp
--- a/Engine/Plugins/Animation/RigLogic/Source/RigLogicModule/Private/DNAUtils.cpp
+++ b/Engine/Plugins/Animation/RigLogic/Source/RigLogicModule/Private/DNAUtils.cpp
@@ -3,6 +3,7 @@
 #include "DNAUtils.h"
 
 #include "CoreMinimal.h"
+#include "HAL/LowLevelMemTracker.h"
 #include "DNAReaderAdapter.h"
 #include "FMemoryResource.h"
 #include "RigLogicMemoryStream.h"
@@ -22,7 +23,7 @@ static TSharedPtr<IDNAReader> ReadDNAStream(rl4::ScopedPtr<dna::StreamReader> DN
 
 TSharedPtr<IDNAReader> ReadDNAFromFile(const FString& Path, EDNADataLayer Layer, uint16_t MaxLOD)
 {
-	LLM_SCOPE(ELLMTag::SkeletalMesh);
+	LLM_SCOPE_BYNAME(TEXT("Animation/RigLogic"));
 	auto DNAFileStream = rl4::makeScoped<rl4::MemoryMappedFileStream>(TCHAR_TO_UTF8(*Path), rl4::MemoryMappedFileStream::AccessMode::Read, FMemoryResource::Instance());
 	auto DNAStreamReader = rl4::makeScoped<dna::StreamReader>(DNAFileStream.get(), static_cast<dna::DataLayer>(Layer), MaxLOD, FMemoryResource::Instance());
 	return ReadDNAStream(MoveTemp(DNAStreamReader));
@@ -30,7 +31,7 @@ TSharedPtr<IDNAReader> ReadDNAFromFile(const FString& Path, EDNADataLayer Layer,
 
 TSharedPtr<IDNAReader> ReadDNAFromFile(const FString& Path, EDNADataLayer Layer, TArrayView<uint16_t> LODs)
 {
-	LLM_SCOPE(ELLMTag::SkeletalMesh);
+	LLM_SCOPE_BYNAME(TEXT("Animation/RigLogic"));
 	auto DNAFileStream = rl4::makeScoped<rl4::MemoryMappedFileStream>(TCHAR_TO_UTF8(*Path), rl4::MemoryMappedFileStream::AccessMode::Read, FMemoryResource::Instance());
 	auto DNAStreamReader = rl4::makeScoped<dna::StreamReader>(DNAFileStream.get(), static_cast<dna::DataLayer>(Layer), LODs.GetData(), static_cast<uint16>(LODs.Num()), FMemoryResource::Instance());
 	return ReadDNAStream(MoveTemp(DNAStreamReader));
@@ -38,7 +39,7 @@ TSharedPtr<IDNAReader> ReadDNAFromFile(const FString& Path, EDNADataLayer Layer,
 
 TSharedPtr<IDNAReader> ReadDNAFromBuffer(TArray<uint8>* DNABuffer, EDNADataLayer Layer, uint16_t MaxLOD)
 {
-	LLM_SCOPE(ELLMTag::SkeletalMesh);
+	LLM_SCOPE_BYNAME(TEXT("Animation/RigLogic"));
 	FRigLogicMemoryStream DNAMemoryStream(DNABuffer);
 	auto DNAStreamReader = rl4::makeScoped<dna::StreamReader>(&DNAMemoryStream, static_cast<dna::DataLayer>(Layer), MaxLOD, FMemoryResource::Instance());
 	return ReadDNAStream(MoveTemp(DNAStreamReader));
@@ -46,7 +47,7 @@ TSharedPtr<IDNAReader> ReadDNAFromBuffer(TArray<uint8>* DNABuffer, EDNADataLayer
 
 TSharedPtr<IDNAReader> ReadDNAFromBuffer(TArray<uint8>* DNABuffer, EDNADataLayer Layer, TArrayView<uint16_t> LODs)
 {
-	LLM_SCOPE(ELLMTag::SkeletalMesh);
+	LLM_SCOPE_BYNAME(TEXT("Animation/RigLogic"));
 	FRigLogicMemoryStream DNAMemoryStream(DNABuffer);
 	auto DNAStreamReader = rl4::makeScoped<dna::StreamReader>(&DNAMemoryStream, static_cast<dna::DataLayer>(Layer), LODs.GetData(), static_cast<uint16>(LODs.Num()), FMemoryResource::Instance());
 	return ReadDNAStream(MoveTemp(DNAStreamReader));
@@ -54,7 +55,7 @@ TSharedPtr<IDNAReader> ReadDNAFromBuffer(TArray<uint8>* DNABuffer, EDNADataLayer
 
 void WriteDNAToFile(const IDNAReader* Reader, EDNADataLayer Layer, const FString& Path)
 {
-	LLM_SCOPE(ELLMTag::SkeletalMesh);
+	LLM_SCOPE_BYNAME(TEXT("Animation/RigLogic"));
 	auto DNAFileStream = rl4::makeScoped<rl4::FileStream>(TCHAR_TO_UTF8(*Path), rl4::FileStream::AccessMode::Write, rl4::FileStream::OpenMode::Binary, FMemoryResource::Instance());
 	auto DNAStreamWriter = rl4::makeScoped<dna::StreamWriter>(DNAFileStream.get(), FMemoryResource::Instance());
 	DNAStreamWriter->setFrom(Reader->Unwrap(), static_cast<dna::DataLayer>(Layer), FMemoryResource::Instance());
diff --git a/Engine/Plugins/Animation/RigLogic/Source/RigLogicModule/Private/RigLogic.cpp b/Engine/Plugins/Animation/RigLogic/Source/RigLogicModule/Private/RigLogic.cpp
--- a/Engine/Plugins/Animation/RigLogic/Source/RigLogicModule/Private/RigLogic.cpp
+++ b/Engine/Plugins/Animation/RigLogic/Source/RigLogicModule/Private/RigLogic.cpp
@@ -27,6 +27,12 @@ uint16 FRigLogic::GetLODCount() const
 	return RigLogic->getLODCount();
 }
 
+TArrayView<const float> FRigLogic::GetRawNeutralJointValues() const
+{
+	rl4::ConstArrayView<float> Values = RigLogic->getRawNeutralJointValues();
+	return TArrayView<const float>{Values.data(), static_cast<int32>(Values.size())};
+}
+
 FTransformArrayView FRigLogic::GetNeutralJointValues() const
 {
 	rl4::ConstArrayView<float> Values = RigLogic->getRawNeutralJointValues();
diff --git a/Engine/Plugins/Animation/RigLogic/Source/RigLogicModule/Private/RigUnit_RigLogic.cpp b/Engine/Plugins/Animation/RigLogic/Source/RigLogicModule/Private/RigUnit_RigLogic.cpp
--- a/Engine/Plugins/Animation/RigLogic/Source/RigLogicModule/Private/RigUnit_RigLogic.cpp
+++ b/Engine/Plugins/Animation/RigLogic/Source/RigLogicModule/Private/RigUnit_RigLogic.cpp
@@ -6,6 +6,7 @@
 #include "DNAReader.h"
 #include "RigInstance.h"
 #include "RigLogic.h"
+#include "SharedRigRuntimeContext.h"
 #include "Components/SkeletalMeshComponent.h"
 #include "Math/TransformNonVectorized.h"
 #include "Units/RigUnitContext.h"
@@ -39,25 +40,17 @@ static FString ConstructCurveName(const FString& NameToSplit, const FString& For
 	return CurveName;
 }
 
-FRigDefinitionMapping::FRigDefinitionMapping(const IBehaviorReader* DNABehavior, const URigHierarchy* InHierarchy) : FRigDefinitionMapping()
-{
-	MapInputCurveIndices(DNABehavior, InHierarchy);
-	MapJoints(DNABehavior, InHierarchy);
-	MapMaskMultipliers(DNABehavior, InHierarchy);
-	MapMorphTargets(DNABehavior, InHierarchy);
-}
-
 FRigUnit_RigLogic_Data::FRigUnit_RigLogic_Data()
-	: SkelMeshComponent(nullptr)
-	, SharedRigRuntimeContext(nullptr)
-	, RigInstance(nullptr)
-	, CurrentLOD(0)
+: SkelMeshComponent(nullptr)
+, LocalRigRuntimeContext(nullptr)
+, RigInstance(nullptr)
+, CurrentLOD(0)
 {
 }
 
 FRigUnit_RigLogic_Data::~FRigUnit_RigLogic_Data()
 {
-	SharedRigRuntimeContext = nullptr;
+	LocalRigRuntimeContext = nullptr;
 }
 
 FRigUnit_RigLogic_Data::FRigUnit_RigLogic_Data(const FRigUnit_RigLogic_Data& Other)
@@ -68,59 +61,58 @@ FRigUnit_RigLogic_Data::FRigUnit_RigLogic_Data(const FRigUnit_RigLogic_Data& Oth
 FRigUnit_RigLogic_Data& FRigUnit_RigLogic_Data::operator=(const FRigUnit_RigLogic_Data& Other)
 {
 	SkelMeshComponent = Other.SkelMeshComponent;
-	SharedRigRuntimeContext = nullptr;
+	LocalRigRuntimeContext = nullptr;
 	RigInstance = nullptr;
-	Mappings = Other.Mappings;
+	InputCurveIndices = Other.InputCurveIndices;
+	HierarchyBoneIndices = Other.HierarchyBoneIndices;
+	MorphTargetCurveIndices = Other.MorphTargetCurveIndices;
+	BlendShapeIndices = Other.BlendShapeIndices;
+	CurveElementIndicesForAnimMaps = Other.CurveElementIndicesForAnimMaps;
+	RigLogicIndicesForAnimMaps = Other.RigLogicIndicesForAnimMaps;
 	CurrentLOD = Other.CurrentLOD;
 	return *this;
 }
 
 bool FRigUnit_RigLogic_Data::IsRigLogicInitialized()
 {
-	return (SharedRigRuntimeContext != nullptr) && SharedRigRuntimeContext->RigLogic.IsValid() && RigInstance.IsValid();
+	return (LocalRigRuntimeContext != nullptr) && LocalRigRuntimeContext->RigLogic.IsValid() && RigInstance.IsValid();
 }
 
-TDiscardableKeyValueCache<uint32, TSharedPtr<const FRigDefinitionMapping>> FRigUnit_RigLogic_Data::RigDefintionMappingsCache{};
-
-void FRigUnit_RigLogic_Data::InitializeRigLogic(const URigHierarchy* InHierarchy)
+void FRigUnit_RigLogic_Data::InitializeRigLogic(const URigHierarchy* InHierarchy, TSharedPtr<FSharedRigRuntimeContext> NewContext)
 {
-	if ((SharedRigRuntimeContext == nullptr) || !SharedRigRuntimeContext->BehaviorReader.IsValid() || SharedRigRuntimeContext->BehaviorReader->GetJointCount() == 0u)
+	DECLARE_SCOPE_HIERARCHICAL_COUNTER_FUNC()
+	QUICK_SCOPE_CYCLE_COUNTER(STAT_RigUnit_RigLogic_InitializeRigLogic);
+	if (!NewContext.IsValid())
 	{
 		UE_LOG(LogRigLogicUnit, Warning, TEXT("No valid DNA file found, abort initialization."));
 		return;
 	}
 
-	if (!SharedRigRuntimeContext->RigLogic.IsValid())
+	if (LocalRigRuntimeContext != NewContext)
 	{
-		SharedRigRuntimeContext->RigLogic = MakeShared<FRigLogic>(SharedRigRuntimeContext->BehaviorReader.Get());
-		CacheVariableJointIndices();
+		LocalRigRuntimeContext = NewContext;
 		RigInstance = nullptr;
 	}
 
 	if (!RigInstance.IsValid())
 	{
-		RigInstance = MakeUnique<FRigInstance>(SharedRigRuntimeContext->RigLogic.Get());
+		RigInstance = MakeUnique<FRigInstance>(LocalRigRuntimeContext->RigLogic.Get());
 		RigInstance->SetLOD(CurrentLOD);
+		CurrentLOD = RigInstance->GetLOD();
 
-		const IBehaviorReader* DNABehavior = SharedRigRuntimeContext->BehaviorReader.Get();
-		const uint16 TopologyVersion = InHierarchy->GetTopologyVersion();
-		const uint32 CacheKey = SharedRigRuntimeContext->DNAHash + TopologyVersion;
-		auto& Cache = FRigUnit_RigLogic_Data::RigDefintionMappingsCache;
-
-		using FLockFLags = TDiscardableKeyValueCache< FString, TSharedPtr<FRigDefinitionMapping>>::LockFlags;
-		uint32 LockFlags = Cache.ApplyLock(0, FLockFLags::ReadLock);
-		if (!Cache.Find(CacheKey, Mappings, LockFlags | FLockFLags::WriteLockOnAddFail, LockFlags))
-		{
-			Mappings = MakeShared<const FRigDefinitionMapping>(DNABehavior, InHierarchy);
-			Cache.Add(CacheKey, Mappings, LockFlags);
-		}
-		Cache.Unlock(LockFlags);
+		MapJoints(InHierarchy);
+		MapInputCurveIndices(InHierarchy);
+		MapMorphTargets(InHierarchy);
+		MapMaskMultipliers(InHierarchy);
 	}
 }
 
 //maps indices of input curves from dna file to control rig curves
-void FRigDefinitionMapping::MapInputCurveIndices(const IBehaviorReader* DNABehavior, const URigHierarchy* InHierarchy)
+void FRigUnit_RigLogic_Data::MapInputCurveIndices(const URigHierarchy* InHierarchy)
 {
+	DECLARE_SCOPE_HIERARCHICAL_COUNTER_FUNC()
+	QUICK_SCOPE_CYCLE_COUNTER(STAT_RigUnit_RigLogic_MapInputCurveIndices);
+	const IBehaviorReader* DNABehavior = LocalRigRuntimeContext->BehaviorReader.Get();
 	const uint32 ControlCount = DNABehavior->GetRawControlCount();
 	InputCurveIndices.Reset(ControlCount);
 	for (uint32_t ControlIndex = 0; ControlIndex < ControlCount; ++ControlIndex)
@@ -137,42 +129,27 @@ void FRigDefinitionMapping::MapInputCurveIndices(const IBehaviorReader* DNABehav
 	}
 }
 
-void FRigDefinitionMapping::MapJoints(const IBehaviorReader* DNABehavior, const URigHierarchy* InHierarchy)
+void FRigUnit_RigLogic_Data::MapJoints(const URigHierarchy* Hierarchy)
 {
+	DECLARE_SCOPE_HIERARCHICAL_COUNTER_FUNC()
+	QUICK_SCOPE_CYCLE_COUNTER(STAT_RigUnit_RigLogic_MapJoints);
+	const IBehaviorReader* DNABehavior = LocalRigRuntimeContext->BehaviorReader.Get();
 	const uint16 JointCount = DNABehavior->GetJointCount();
 	HierarchyBoneIndices.Reset(JointCount);
-	for (uint16 JointIndex = 0; JointIndex < JointCount; ++JointIndex)
+	for (uint16 JointIndex = 0; JointIndex < JointCount ; ++JointIndex)
 	{
 		const FString RLJointName = DNABehavior->GetJointName(JointIndex);
 		const FName JointFName = FName(*RLJointName);
-		const int32 BoneIndex = InHierarchy->GetIndex(FRigElementKey(JointFName, ERigElementType::Bone));
+		const int32 BoneIndex = Hierarchy->GetIndex(FRigElementKey(JointFName, ERigElementType::Bone));
 		HierarchyBoneIndices.Add(BoneIndex);
 	}
 }
 
-void FRigUnit_RigLogic_Data::CacheVariableJointIndices()
-{
-	const IBehaviorReader* DNABehavior = SharedRigRuntimeContext->BehaviorReader.Get();
-	const uint16 LODCount = DNABehavior->GetLODCount();
-	SharedRigRuntimeContext->VariableJointIndices.Reset();
-	SharedRigRuntimeContext->VariableJointIndices.AddDefaulted(LODCount);
-	TSet<uint16> DistinctVariableJointIndices;
-	for (uint16 LodIndex = 0; LodIndex < LODCount; ++LodIndex)
-	{
-		TArrayView<const uint16> VariableAttributeIndices = DNABehavior->GetJointVariableAttributeIndices(LodIndex);
-		DistinctVariableJointIndices.Reset();
-		DistinctVariableJointIndices.Reserve(VariableAttributeIndices.Num());
-		for (const uint16 AttrIndex : VariableAttributeIndices)
-		{
-			const uint16 JointIndex = AttrIndex / MAX_ATTRS_PER_JOINT;
-			DistinctVariableJointIndices.Add(JointIndex);
-		}
-		SharedRigRuntimeContext->VariableJointIndices[LodIndex].Values = DistinctVariableJointIndices.Array();
-	}
-}
-
-void FRigDefinitionMapping::MapMorphTargets(const IBehaviorReader* DNABehavior, const URigHierarchy* InHierarchy)
+void FRigUnit_RigLogic_Data::MapMorphTargets(const URigHierarchy* InHierarchy)
 {
+	DECLARE_SCOPE_HIERARCHICAL_COUNTER_FUNC()
+	QUICK_SCOPE_CYCLE_COUNTER(STAT_RigUnit_RigLogic_MapMorphTargets);
+	const IBehaviorReader* DNABehavior = LocalRigRuntimeContext->BehaviorReader.Get();
 	const uint16 LODCount = DNABehavior->GetLODCount();
 
 	MorphTargetCurveIndices.Reset();
@@ -186,7 +163,7 @@ void FRigDefinitionMapping::MapMorphTargets(const IBehaviorReader* DNABehavior,
 			DNABehavior->GetMeshBlendShapeChannelMappingIndicesForLOD(LodIndex);
 		MorphTargetCurveIndices[LodIndex].Values.Reserve(BlendShapeChannelIndicesForLOD.Num());
 		BlendShapeIndices[LodIndex].Values.Reserve(BlendShapeChannelIndicesForLOD.Num());
-		for (uint16 MappingIndex : BlendShapeChannelIndicesForLOD)
+		for (uint16 MappingIndex: BlendShapeChannelIndicesForLOD)
 		{
 			const FMeshBlendShapeChannelMapping Mapping = DNABehavior->GetMeshBlendShapeChannelMapping(MappingIndex);
 			const uint16 BlendShapeIndex = Mapping.BlendShapeChannelIndex;
@@ -202,8 +179,11 @@ void FRigDefinitionMapping::MapMorphTargets(const IBehaviorReader* DNABehavior,
 	}
 }
 
-void FRigDefinitionMapping::MapMaskMultipliers(const IBehaviorReader* DNABehavior, const URigHierarchy* InHierarchy)
+void FRigUnit_RigLogic_Data::MapMaskMultipliers(const URigHierarchy* InHierarchy)
 {
+	DECLARE_SCOPE_HIERARCHICAL_COUNTER_FUNC()
+	QUICK_SCOPE_CYCLE_COUNTER(STAT_RigUnit_RigLogic_MapMaskMultipliers);
+	const IBehaviorReader* DNABehavior = LocalRigRuntimeContext->BehaviorReader.Get();
 	const uint16 LODCount = DNABehavior->GetLODCount();
 	CurveElementIndicesForAnimMaps.Reset();
 	CurveElementIndicesForAnimMaps.AddDefaulted(LODCount);
@@ -216,7 +196,7 @@ void FRigDefinitionMapping::MapMaskMultipliers(const IBehaviorReader* DNABehavio
 		TArrayView<const uint16> AnimMapIndicesPerLOD = DNABehavior->GetAnimatedMapIndicesForLOD(LodIndex);
 		CurveElementIndicesForAnimMaps[LodIndex].Values.Reserve(AnimMapIndicesPerLOD.Num());
 		RigLogicIndicesForAnimMaps[LodIndex].Values.Reserve(AnimMapIndicesPerLOD.Num());
-		for (uint16 AnimMapIndexPerLOD : AnimMapIndicesPerLOD)
+		for (uint16 AnimMapIndexPerLOD: AnimMapIndicesPerLOD)
 		{
 			const FString AnimMapNameFStr = DNABehavior->GetAnimatedMapName(AnimMapIndexPerLOD);
 			const FString MaskMultiplierNameStr = ConstructCurveName(AnimMapNameFStr, TEXT("<obj>_<attr>"));
@@ -233,36 +213,45 @@ void FRigDefinitionMapping::MapMaskMultipliers(const IBehaviorReader* DNABehavio
 
 void FRigUnit_RigLogic_Data::CalculateRigLogic(const URigHierarchy* InHierarchy)
 {
+	DECLARE_SCOPE_HIERARCHICAL_COUNTER_FUNC()
 	QUICK_SCOPE_CYCLE_COUNTER(STAT_RigUnit_RigLogic_Calculate);
 
 	// LOD change is inexpensive
 	RigInstance->SetLOD(CurrentLOD);
+	CurrentLOD = RigInstance->GetLOD();
 
 	const int32 RawControlCount = RigInstance->GetRawControlCount();
 	for (int32 ControlIndex = 0; ControlIndex < RawControlCount; ++ControlIndex)
 	{
-		const uint32 CurveIndex = Mappings->InputCurveIndices[ControlIndex];
+		const uint32 CurveIndex = InputCurveIndices[ControlIndex];
 		const float Value = InHierarchy->GetCurveValue(CurveIndex);
 		RigInstance->SetRawControl(ControlIndex, Value);
 	}
-	SharedRigRuntimeContext->RigLogic->Calculate(RigInstance.Get());
+	LocalRigRuntimeContext->RigLogic->Calculate(RigInstance.Get());
 }
 
-void FRigUnit_RigLogic_Data::UpdateJoints(URigHierarchy* Hierarchy, const FRigUnit_RigLogic_JointUpdateParams& JointUpdateParams)
+void FRigUnit_RigLogic_Data::UpdateJoints(URigHierarchy* Hierarchy, TArrayView<const float> NeutralJointValues, TArrayView<const float> DeltaJointValues)
 {
-	FRigHierarchyValidityBracket ValidityBracket(Hierarchy);
-	for (const uint16 JointIndex : SharedRigRuntimeContext->VariableJointIndices[CurrentLOD].Values)
+	DECLARE_SCOPE_HIERARCHICAL_COUNTER_FUNC()
+	QUICK_SCOPE_CYCLE_COUNTER(STAT_RigUnit_RigLogic_UpdateJoints);
+
+	const float* N = NeutralJointValues.GetData();
+	const float* D = DeltaJointValues.GetData();
+
+	for (const uint16 JointIndex : LocalRigRuntimeContext->VariableJointIndicesPerLOD[CurrentLOD].Values)
 	{
-		const int32 BoneIndex = Mappings->HierarchyBoneIndices[JointIndex];
+		const int32 BoneIndex = HierarchyBoneIndices[JointIndex];
 		if (BoneIndex != INDEX_NONE)
 		{
-			const FTransform& Neutral = JointUpdateParams.NeutralJointTransforms[JointIndex];
-			const FTransform& Delta = JointUpdateParams.DeltaTransforms[JointIndex];
+			const uint16 AttrIndex = JointIndex * MAX_ATTRS_PER_JOINT;
 			const FTransform Transform
 			{
-				Neutral.GetRotation() * Delta.GetRotation(),
-				Neutral.GetTranslation() + Delta.GetTranslation(),
-				Neutral.GetScale3D() + Delta.GetScale3D()  // Neutral scale is always 1.0
+				// Rotation: X = -Y, Y = -Z, Z = X
+				FQuat(FRotator(-N[AttrIndex + 4], -N[AttrIndex + 5], N[AttrIndex + 3])) * FQuat(FRotator(-D[AttrIndex + 4], -D[AttrIndex + 5], D[AttrIndex + 3])),
+				// Translation: X = X, Y = -Y, Z = Z
+				FVector((N[AttrIndex + 0] + D[AttrIndex + 0]), -(N[AttrIndex + 1] + D[AttrIndex + 1]), (N[AttrIndex + 2] + D[AttrIndex + 2])),
+				// Scale: X = X, Y = Y, Z = Z
+				FVector((N[AttrIndex + 6] + D[AttrIndex + 6]), (N[AttrIndex + 7] + D[AttrIndex + 7]), (N[AttrIndex + 8] + D[AttrIndex + 8]))
 			};
 			Hierarchy->SetLocalTransform(BoneIndex, Transform);
 		}
@@ -271,17 +260,19 @@ void FRigUnit_RigLogic_Data::UpdateJoints(URigHierarchy* Hierarchy, const FRigUn
 
 void FRigUnit_RigLogic_Data::UpdateBlendShapeCurves(URigHierarchy* InHierarchy, TArrayView<const float> BlendShapeValues)
 {
+	DECLARE_SCOPE_HIERARCHICAL_COUNTER_FUNC()
+	QUICK_SCOPE_CYCLE_COUNTER(STAT_RigUnit_RigLogic_UpdateBlendShapeCurves);
 	// set output blend shapes
-	if (Mappings->BlendShapeIndices.IsValidIndex(CurrentLOD) && Mappings->MorphTargetCurveIndices.IsValidIndex(CurrentLOD))
+	if (BlendShapeIndices.IsValidIndex(CurrentLOD) && MorphTargetCurveIndices.IsValidIndex(CurrentLOD))
 	{
-		const uint32 BlendShapePerLODCount = static_cast<uint32>(Mappings->BlendShapeIndices[CurrentLOD].Values.Num());
+		const uint32 BlendShapePerLODCount = static_cast<uint32>(BlendShapeIndices[CurrentLOD].Values.Num());
 
-		if (ensure(BlendShapePerLODCount == Mappings->MorphTargetCurveIndices[CurrentLOD].Values.Num()))
+		if (ensure(BlendShapePerLODCount == MorphTargetCurveIndices[CurrentLOD].Values.Num()))
 		{
 			for (uint32 MeshBlendIndex = 0; MeshBlendIndex < BlendShapePerLODCount; MeshBlendIndex++)
 			{
-				const int32 BlendShapeIndex = Mappings->BlendShapeIndices[CurrentLOD].Values[MeshBlendIndex];
-				const int32 MorphTargetCurveIndex = Mappings->MorphTargetCurveIndices[CurrentLOD].Values[MeshBlendIndex];
+				const int32 BlendShapeIndex = BlendShapeIndices[CurrentLOD].Values[MeshBlendIndex];
+				const int32 MorphTargetCurveIndex = MorphTargetCurveIndices[CurrentLOD].Values[MeshBlendIndex];
 				if (MorphTargetCurveIndex != INDEX_NONE)
 				{
 					const float Value = BlendShapeValues[BlendShapeIndex];
@@ -298,16 +289,18 @@ void FRigUnit_RigLogic_Data::UpdateBlendShapeCurves(URigHierarchy* InHierarchy,
 
 void FRigUnit_RigLogic_Data::UpdateAnimMapCurves(URigHierarchy* InHierarchy, TArrayView<const float> AnimMapOutputs)
 {
+	DECLARE_SCOPE_HIERARCHICAL_COUNTER_FUNC()
+	QUICK_SCOPE_CYCLE_COUNTER(STAT_RigUnit_RigLogic_UpdateAnimMapCurves);
 	// set output mask multipliers
 	// In case curves are not imported yet into CL, AnimatedMapsCurveIndices will be empty, so we need to check
 	// array bounds before trying to access it:
-	if (Mappings->RigLogicIndicesForAnimMaps.IsValidIndex(CurrentLOD) && Mappings->CurveElementIndicesForAnimMaps.IsValidIndex(CurrentLOD))
+	if (RigLogicIndicesForAnimMaps.IsValidIndex(CurrentLOD) && CurveElementIndicesForAnimMaps.IsValidIndex(CurrentLOD))
 	{
-		const uint32 AnimMapPerLODCount = Mappings->RigLogicIndicesForAnimMaps[CurrentLOD].Values.Num();
+		const uint32 AnimMapPerLODCount = RigLogicIndicesForAnimMaps[CurrentLOD].Values.Num();
 		for (uint32 AnimMapIndexForLOD = 0; AnimMapIndexForLOD < AnimMapPerLODCount; ++AnimMapIndexForLOD)
 		{
-			const int32 RigLogicAnimMapIndex = Mappings->RigLogicIndicesForAnimMaps[CurrentLOD].Values[AnimMapIndexForLOD];
-			const int32 InHierarchyAnimMapIndex = Mappings->CurveElementIndicesForAnimMaps[CurrentLOD].Values[AnimMapIndexForLOD];
+			const int32 RigLogicAnimMapIndex = RigLogicIndicesForAnimMaps[CurrentLOD].Values[AnimMapIndexForLOD];
+			const int32 InHierarchyAnimMapIndex = CurveElementIndicesForAnimMaps[CurrentLOD].Values[AnimMapIndexForLOD];
 			if (InHierarchyAnimMapIndex != INDEX_NONE)
 			{
 				const float Value = AnimMapOutputs[RigLogicAnimMapIndex];
@@ -318,10 +311,10 @@ void FRigUnit_RigLogic_Data::UpdateAnimMapCurves(URigHierarchy* InHierarchy, TAr
 	else
 	{
 		UE_LOG(LogRigLogicUnit, Warning, TEXT("Invalid LOD Index for the AnimationMaps. Ensure your curve is set up correctly!"));
-	}
+	}	
 }
 
-FSharedRigRuntimeContext* FRigUnit_RigLogic::GetSharedRigRuntimeContext(USkeletalMesh* SkelMesh)
+TSharedPtr<FSharedRigRuntimeContext> FRigUnit_RigLogic::GetSharedRigRuntimeContext(USkeletalMesh* SkelMesh)
 {
 	UAssetUserData* UserData = SkelMesh->GetAssetUserDataOfClass(UDNAAsset::StaticClass());
 	if (UserData == nullptr)
@@ -329,69 +322,74 @@ FSharedRigRuntimeContext* FRigUnit_RigLogic::GetSharedRigRuntimeContext(USkeleta
 		return nullptr;
 	}
 	UDNAAsset* DNAAsset = Cast<UDNAAsset>(UserData);
-	return &(DNAAsset->Context);
+	return DNAAsset->GetRigRuntimeContext(UDNAAsset::EDNARetentionPolicy::Keep);
 }
 
 FRigUnit_RigLogic_Execute()
 {
+	DECLARE_SCOPE_HIERARCHICAL_COUNTER_FUNC()
 	QUICK_SCOPE_CYCLE_COUNTER(STAT_RigUnit_RigLogic_Execute);
 
-	URigHierarchy* Hierarchy = ExecuteContext.Hierarchy;
+ 	URigHierarchy* Hierarchy = ExecuteContext.Hierarchy;
 	if (Hierarchy)
 	{
 		switch (Context.State)
 		{
-		case EControlRigState::Init:
-		{
-			if (!Data.SkelMeshComponent.IsValid())
+			case EControlRigState::Init:
 			{
-				Data.SkelMeshComponent = Context.DataSourceRegistry->RequestSource<USkeletalMeshComponent>(UControlRig::OwnerComponent);
-				// In normal execution, Data.SkelMeshComponent will be nullptr at the beginning
-				// however, during unit testing we cannot fetch it from DataSourceRegistry 
-				// in that case, a mock version will be inserted into Data by unit test beforehand
-			}
-			if (!Data.SkelMeshComponent.IsValid() || Data.SkelMeshComponent->SkeletalMesh == nullptr)
-			{
-				return;
+				//const double startTime = FPlatformTime::Seconds();
+				if (!Data.SkelMeshComponent.IsValid())
+				{
+					Data.SkelMeshComponent = Context.DataSourceRegistry->RequestSource<USkeletalMeshComponent>(UControlRig::OwnerComponent);
+					// In normal execution, Data.SkelMeshComponent will be nullptr at the beginning
+					// however, during unit testing we cannot fetch it from DataSourceRegistry 
+					// in that case, a mock version will be inserted into Data by unit test beforehand
+				}
+				if (!Data.SkelMeshComponent.IsValid() || Data.SkelMeshComponent->SkeletalMesh == nullptr)
+				{
+					return;
+				}
+				Data.CurrentLOD = Data.SkelMeshComponent->GetPredictedLODLevel();
+
+				// Fetch shared runtime context of rig from DNAAsset
+				TSharedPtr<FSharedRigRuntimeContext> RigRuntimeContext = GetSharedRigRuntimeContext(Data.SkelMeshComponent->SkeletalMesh);
+				// Context is initialized with a BehaviorReader, which can be imported into SkeletalMesh from DNA file
+				// or overwritten by GeneSplicer when making a new character
+				Data.InitializeRigLogic(Hierarchy, RigRuntimeContext);
+				//const double delta = FPlatformTime::Seconds() - startTime;
+				//UE_LOG(LogRigLogicUnit, Warning, TEXT("RigLogic::Init execution time: %f"), delta);
+
+				break; 
 			}
-			Data.CurrentLOD = Data.SkelMeshComponent->GetPredictedLODLevel();
-
-			// Fetch shared runtime context of rig from DNAAsset
-			Data.SharedRigRuntimeContext = GetSharedRigRuntimeContext(Data.SkelMeshComponent->SkeletalMesh);
-			// Context is initialized with a BehaviorReader, which can be imported into SkeletalMesh from DNA file
-			// or overwritten by GeneSplicer when making a new character
-			Data.InitializeRigLogic(Hierarchy);
-			break;
-		}
-		case EControlRigState::Update:
-		{
-			// Fetch shared runtime context of rig from DNAAsset
-			if (!Data.SkelMeshComponent.IsValid() || !Data.IsRigLogicInitialized() || Hierarchy == nullptr)
+			case EControlRigState::Update:
 			{
-				return;
-			}
-			Data.CurrentLOD = Data.SkelMeshComponent->GetPredictedLODLevel();
-			Data.CalculateRigLogic(Hierarchy);
+				//const double startTime = FPlatformTime::Seconds();
+				// Fetch shared runtime context of rig from DNAAsset
+				if (!Data.SkelMeshComponent.IsValid() || !Data.IsRigLogicInitialized() || Hierarchy == nullptr)
+				{
+					return;
+				}
+				Data.CurrentLOD = Data.SkelMeshComponent->GetPredictedLODLevel();
+				Data.CalculateRigLogic(Hierarchy);
 
-			//Filing a struct so we can call the same method for updating joints from tests
-			FRigUnit_RigLogic_JointUpdateParams JointUpdateParamsTemp
-			(
-				Data.SharedRigRuntimeContext->RigLogic->GetNeutralJointValues(),
-				Data.RigInstance->GetJointOutputs()
-			);
+				//Filing a struct so we can call the same method for updating joints from tests
+				TArrayView<const float> NeutralJointValues = Data.LocalRigRuntimeContext->RigLogic->GetRawNeutralJointValues();
+				TArrayView<const float> DeltaJointValues = Data.RigInstance->GetRawJointOutputs();
+				Data.UpdateJoints(Hierarchy, NeutralJointValues, DeltaJointValues);
 
-			Data.UpdateJoints(Hierarchy, JointUpdateParamsTemp);
+				TArrayView<const float> BlendShapeValues = Data.RigInstance->GetBlendShapeOutputs();
+				Data.UpdateBlendShapeCurves(Hierarchy, BlendShapeValues);
 
-			TArrayView<const float> BlendShapeValues = Data.RigInstance->GetBlendShapeOutputs();
-			Data.UpdateBlendShapeCurves(Hierarchy, BlendShapeValues);
+				TArrayView<const float> AnimMapOutputs = Data.RigInstance->GetAnimatedMapOutputs();
+				Data.UpdateAnimMapCurves(Hierarchy, AnimMapOutputs);
 
-			TArrayView<const float> AnimMapOutputs = Data.RigInstance->GetAnimatedMapOutputs();
-			Data.UpdateAnimMapCurves(Hierarchy, AnimMapOutputs);
-		}
-		default:
-		{
-			break;
-		}
+				//const double delta = FPlatformTime::Seconds() - startTime;
+				//UE_LOG(LogRigLogicUnit, Warning, TEXT("RigLogic::Update execution time: %f"), delta);
+			}
+			default:
+			{
+				break;
+			}
 		}
 	}
 }
diff --git a/Engine/Plugins/Animation/RigLogic/Source/RigLogicModule/Private/RigUnit_RigLogic_Test.cpp b/Engine/Plugins/Animation/RigLogic/Source/RigLogicModule/Private/RigUnit_RigLogic_Test.cpp
--- a/Engine/Plugins/Animation/RigLogic/Source/RigLogicModule/Private/RigUnit_RigLogic_Test.cpp
+++ b/Engine/Plugins/Animation/RigLogic/Source/RigLogicModule/Private/RigUnit_RigLogic_Test.cpp
@@ -1,10 +1,11 @@
 // Copyright Epic Games, Inc. All Rights Reserved.
 
 #if WITH_DEV_AUTOMATION_TESTS
-#include "RigUnit_RigLogic_Test.h"
+	#include "RigUnit_RigLogic_Test.h"
 #endif
 
 #include "RigUnit_RigLogic.h"
+#include "SharedRigRuntimeContext.h"
 #include "Components/SkeletalMeshComponent.h"
 #include "Units/RigUnitContext.h"
 #include "ControlRig.h"
@@ -55,7 +56,7 @@ TStrongObjectPtr<URigHierarchy> FRigUnit_RigLogic::TestAccessor::CreateCurveCont
 
 void FRigUnit_RigLogic::TestAccessor::Exec_MapInputCurve(URigHierarchy* TestHierarchy)
 {
-	Unit->Data.InitializeRigLogic(TestHierarchy); //inside the method so we can access data, which is a private member
+	Unit->Data.MapInputCurveIndices(TestHierarchy); //inside the method so we can access data, which is a private member
 }
 
 /** ====== Map Joints ===== **/
@@ -86,7 +87,7 @@ TSharedPtr<TestBehaviorReader> FRigUnit_RigLogic::TestAccessor::CreateBehaviorRe
 
 void FRigUnit_RigLogic::TestAccessor::Exec_MapJoints(URigHierarchy* TestHierarchy)
 {
-	Unit->Data.InitializeRigLogic(TestHierarchy); //inside the method so we can access data, which is a private member
+	Unit->Data.MapJoints(TestHierarchy);
 }
 
 /** ====== Map Morph Targets ===== **/
@@ -145,7 +146,7 @@ TStrongObjectPtr<URigHierarchy> FRigUnit_RigLogic::TestAccessor::CreateCurveCont
 void FRigUnit_RigLogic::TestAccessor::Exec_MapMorphTargets(URigHierarchy* TestCurveContainer)
 {
 	//put into a separate method so we can access the private Data member
-	Unit->Data.InitializeRigLogic(TestCurveContainer);
+	Unit->Data.MapMorphTargets(TestCurveContainer); 
 }
 
 /** ====== Map Mask Multipliers ===== **/
@@ -161,13 +162,13 @@ TSharedPtr<TestBehaviorReader> FRigUnit_RigLogic::TestAccessor::CreateBehaviorRe
 
 void FRigUnit_RigLogic::TestAccessor::Exec_MapMaskMultipliers(URigHierarchy* TestHierarchy)
 {
-	Unit->Data.InitializeRigLogic(TestHierarchy); //inside the method so we can access data, which is a private member
+	Unit->Data.MapMaskMultipliers(TestHierarchy); //inside the method so we can access data, which is a private member
 }
 
 void FRigUnit_RigLogic::TestAccessor::AddToTransformArray(float* InArray, FTransform& Transform)
 {
 	uint32 FirstAttributeIndex = 0;
-
+	
 	FVector Rotation = Transform.GetRotation().Euler();
 	InArray[FirstAttributeIndex + 0] = Rotation.X;
 	InArray[FirstAttributeIndex + 1] = Rotation.Y;
@@ -185,7 +186,7 @@ void FRigUnit_RigLogic::TestAccessor::AddToTransformArray(float* InArray, FTrans
 }
 
 
-FTransformArrayView FRigUnit_RigLogic::TestAccessor::CreateTwoJointNeutralTransforms(float* InValueArray)
+FTransformArrayView FRigUnit_RigLogic::TestAccessor::CreateTwoJointNeutralTransforms(float *InValueArray)
 {
 	float* Transform1Ptr = InValueArray;
 
@@ -229,19 +230,19 @@ TArrayView<const uint16> FRigUnit_RigLogic::TestAccessor::CreateTwoJointVariable
 		InVariableAttributeIndices[9 + 6] = 9 + 6;
 		InVariableAttributeIndices[9 + 7] = 9 + 7;
 		InVariableAttributeIndices[9 + 8] = 9 + 8;
-
+	
 		return TArrayView<const uint16>(InVariableAttributeIndices, 2 * FRigUnit_RigLogic::TestAccessor::MAX_ATTRS_PER_JOINT);
 	}
 
 	return TArrayView<const uint16>(InVariableAttributeIndices, FRigUnit_RigLogic::TestAccessor::MAX_ATTRS_PER_JOINT);
 }
 
-void FRigUnit_RigLogic::TestAccessor::Exec_UpdateJoints(URigHierarchy* TestHierarchy, FRigUnit_RigLogic_JointUpdateParams& JointUpdateParams)
+void FRigUnit_RigLogic::TestAccessor::Exec_UpdateJoints(URigHierarchy* TestHierarchy, TArrayView<const float> NeutralJointValues, TArrayView<const float> DeltaJointValues)
 {
-	Unit->Data.UpdateJoints(TestHierarchy, JointUpdateParams);
+	Unit->Data.UpdateJoints(TestHierarchy, NeutralJointValues, DeltaJointValues);
 }
 
-FSharedRigRuntimeContext* FRigUnit_RigLogic::TestAccessor::GetSharedRigRuntimeContext(USkeletalMesh* SkelMesh)
+TSharedPtr<FSharedRigRuntimeContext> FRigUnit_RigLogic::TestAccessor::GetSharedRigRuntimeContext(USkeletalMesh* SkelMesh)
 {
 	return FRigUnit_RigLogic::GetSharedRigRuntimeContext(SkelMesh);
 }
@@ -276,7 +277,7 @@ IMPLEMENT_RIGUNIT_AUTOMATION_TEST(FRigUnit_RigLogic)
 
 	//Test
 	InitAndExecute();
-	FSharedRigRuntimeContext* SharedRigRuntimeContext = Test.GetSharedRigRuntimeContext(Test.GetData()->SkelMeshComponent->SkeletalMesh);
+	TSharedPtr<FSharedRigRuntimeContext> SharedRigRuntimeContext = Test.GetSharedRigRuntimeContext(Test.GetData()->SkelMeshComponent->SkeletalMesh);
 	//Assert
 	AddErrorIfFalse(
 		// Check rig logic initialized
@@ -284,18 +285,18 @@ IMPLEMENT_RIGUNIT_AUTOMATION_TEST(FRigUnit_RigLogic)
 		Test.GetData()->RigInstance != nullptr &&
 
 		// Check joints
-		Test.GetData()->Mappings->HierarchyBoneIndices.Num() > 0 &&
+		Test.GetData()->HierarchyBoneIndices.Num() > 0 &&
 
 		// Check input curves
-		Test.GetData()->Mappings->InputCurveIndices.Num() > 0 &&
+		Test.GetData()->InputCurveIndices.Num() > 0 &&
 
 		// Check morph targets
-		Test.GetData()->Mappings->MorphTargetCurveIndices.Num() > 0 &&
-		Test.GetData()->Mappings->BlendShapeIndices.Num() > 0 &&
+		Test.GetData()->MorphTargetCurveIndices.Num() > 0 &&
+		Test.GetData()->BlendShapeIndices.Num() > 0 &&
 
 		// Check mask multipliers
-		Test.GetData()->Mappings->RigLogicIndicesForAnimMaps.Num() > 0 &&
-		Test.GetData()->Mappings->CurveElementIndicesForAnimMaps.Num() > 0,
+		Test.GetData()->RigLogicIndicesForAnimMaps.Num() > 0 &&
+		Test.GetData()->CurveElementIndicesForAnimMaps.Num() > 0,
 
 		TEXT("InitAndExecute failed to initialize rig logic.")
 	);
@@ -312,8 +313,8 @@ IMPLEMENT_RIGUNIT_AUTOMATION_TEST(FRigUnit_RigLogic)
 	Test.Exec_MapInputCurve(TestCurveContainerNameMismatch.Get());
 	//Assert
 	AddErrorIfFalse(
-		Test.GetData()->Mappings->InputCurveIndices.Num() == 1 &&
-		Test.GetData()->Mappings->InputCurveIndices[0] == INDEX_NONE,
+		Test.GetData()->InputCurveIndices.Num() == 1 &&
+		Test.GetData()->InputCurveIndices[0] == INDEX_NONE,
 		TEXT("MapInputCurve - ValidReader CurveContainerWithNameMismatch")
 	);
 
@@ -327,7 +328,7 @@ IMPLEMENT_RIGUNIT_AUTOMATION_TEST(FRigUnit_RigLogic)
 	Test.Exec_MapInputCurve(TestCurveContainerValid.Get());
 	//Assert
 	AddErrorIfFalse(
-		Test.GetData()->Mappings->InputCurveIndices.Num() == 0,
+		Test.GetData()->InputCurveIndices.Num() == 0,
 		TEXT("MapInputCurve - EmptyReader ValidCurveContainer")
 	);
 
@@ -340,8 +341,8 @@ IMPLEMENT_RIGUNIT_AUTOMATION_TEST(FRigUnit_RigLogic)
 	Test.Exec_MapInputCurve(TestCurveContainerEmpty.Get());
 	//Assert
 	AddErrorIfFalse(
-		Test.GetData()->Mappings->InputCurveIndices.Num() == 1 &&
-		Test.GetData()->Mappings->InputCurveIndices[0] == INDEX_NONE,
+		Test.GetData()->InputCurveIndices.Num() == 1 &&
+		Test.GetData()->InputCurveIndices[0] == INDEX_NONE,
 		TEXT("MapInputCurve - ValidReader EmptyCurveContainer")
 	);
 
@@ -355,7 +356,7 @@ IMPLEMENT_RIGUNIT_AUTOMATION_TEST(FRigUnit_RigLogic)
 	Test.Exec_MapInputCurve(TestCurveContainerValid.Get());
 	//Assert
 	AddErrorIfFalse(
-		Test.GetData()->Mappings->InputCurveIndices.Num() == 0,
+		Test.GetData()->InputCurveIndices.Num() == 0,
 		TEXT("MapInputCurve - InvalidReader ValidCurveContainer")
 	);
 
@@ -368,8 +369,8 @@ IMPLEMENT_RIGUNIT_AUTOMATION_TEST(FRigUnit_RigLogic)
 	Test.Exec_MapInputCurve(TestCurveContainerValid.Get());
 	//Assert
 	AddErrorIfFalse(
-		Test.GetData()->Mappings->InputCurveIndices.Num() == 1 &&
-		Test.GetData()->Mappings->InputCurveIndices[0] == 0,
+		Test.GetData()->InputCurveIndices.Num() == 1 &&
+		Test.GetData()->InputCurveIndices[0] == 0,
 		TEXT("MapInputCurve - Valid Inputs")
 	);
 
@@ -384,7 +385,7 @@ IMPLEMENT_RIGUNIT_AUTOMATION_TEST(FRigUnit_RigLogic)
 	Test.Exec_MapJoints(TestHierarchyEmpty.Get());
 	//Assert
 	AddErrorIfFalse(
-		Test.GetData()->Mappings->HierarchyBoneIndices.Num() == 0,
+		Test.GetData()->HierarchyBoneIndices.Num() == 0,
 		TEXT("MapJoints - Empty Inputs")
 	);
 
@@ -395,7 +396,7 @@ IMPLEMENT_RIGUNIT_AUTOMATION_TEST(FRigUnit_RigLogic)
 	//Test
 	Test.Exec_MapJoints(TestHierarchyTwoBones.Get());
 	AddErrorIfFalse(
-		Test.GetData()->Mappings->HierarchyBoneIndices.Num() == 0,
+		Test.GetData()->HierarchyBoneIndices.Num() == 0,
 		TEXT("MapJoints - EmptyReader TwoBones")
 	);
 
@@ -406,17 +407,17 @@ IMPLEMENT_RIGUNIT_AUTOMATION_TEST(FRigUnit_RigLogic)
 	//Test
 	Test.Exec_MapJoints(TestHierarchyEmpty.Get());
 	AddErrorIfFalse(
-		Test.GetData()->Mappings->HierarchyBoneIndices.Num() == 2,
+		Test.GetData()->HierarchyBoneIndices.Num() == 2,
 		TEXT("MapJoints - TwoJoints NoBones - expected 2 bone indices")
 	);
 	AddErrorIfFalse(
-		Test.GetData()->Mappings->HierarchyBoneIndices.Num() == 2 && //repeated condition to prevent crash
-		Test.GetData()->Mappings->HierarchyBoneIndices[0] == INDEX_NONE,
+		Test.GetData()->HierarchyBoneIndices.Num() == 2 && //repeated condition to prevent crash
+		Test.GetData()->HierarchyBoneIndices[0] == INDEX_NONE,
 		TEXT("MapJoints - TwoJoints NoBones - Expected joint 0 index to be NONE")
 	);
 	AddErrorIfFalse(
-		Test.GetData()->Mappings->HierarchyBoneIndices.Num() == 2 && //repeated condition to prevent crash
-		Test.GetData()->Mappings->HierarchyBoneIndices[1] == INDEX_NONE,
+		Test.GetData()->HierarchyBoneIndices.Num() == 2 && //repeated condition to prevent crash
+		Test.GetData()->HierarchyBoneIndices[1] == INDEX_NONE,
 		TEXT("MapJoints - TwoJoints NoBones - Expected joint 1 index to be NONE")
 	);
 
@@ -428,17 +429,17 @@ IMPLEMENT_RIGUNIT_AUTOMATION_TEST(FRigUnit_RigLogic)
 	Test.Exec_MapJoints(TestHierarchyTwoBones.Get());
 	//Assert
 	AddErrorIfFalse(
-		Test.GetData()->Mappings->HierarchyBoneIndices.Num() == 2,
+		Test.GetData()->HierarchyBoneIndices.Num() == 2,
 		TEXT("MapJoints - TwoJoints TwoBones - Expected 2 bone indices")
 	);
 	AddErrorIfFalse(
-		Test.GetData()->Mappings->HierarchyBoneIndices.Num() == 2 && //prevent crash
-		Test.GetData()->Mappings->HierarchyBoneIndices[0] == 0,
+		Test.GetData()->HierarchyBoneIndices.Num() == 2 && //prevent crash
+		Test.GetData()->HierarchyBoneIndices[0] == 0,
 		TEXT("MapJoints - TwoJoints TwoBones - Expected bone 0 index to be 0")
 	);
 	AddErrorIfFalse(
-		Test.GetData()->Mappings->HierarchyBoneIndices.Num() == 2 && //prevent crash
-		Test.GetData()->Mappings->HierarchyBoneIndices[1] == 1,
+		Test.GetData()->HierarchyBoneIndices.Num() == 2 && //prevent crash
+		Test.GetData()->HierarchyBoneIndices[1] == 1,
 		TEXT("MapJoints - TwoJoints TwoBones - Expected bone index 1 index to be 1")
 	);
 
@@ -455,12 +456,12 @@ IMPLEMENT_RIGUNIT_AUTOMATION_TEST(FRigUnit_RigLogic)
 	Test.Exec_MapMorphTargets(TestMorphTargetNameMismatch.Get());
 	//Assert
 	AddErrorIfFalse(
-		Test.GetData()->Mappings->BlendShapeIndices.Num() == 1 && //LOD 0
-		Test.GetData()->Mappings->MorphTargetCurveIndices.Num() == 1 && //LOD 0
-		Test.GetData()->Mappings->BlendShapeIndices[0].Values.Num() == 1 && //at least one blendshape exists
-		Test.GetData()->Mappings->BlendShapeIndices[0].Values[0] == 0 && //has index 0
-		Test.GetData()->Mappings->MorphTargetCurveIndices[0].Values.Num() == 1 && //but morph target corresponding to that blendshape
-		Test.GetData()->Mappings->MorphTargetCurveIndices[0].Values[0] == INDEX_NONE, //wasn't found
+		Test.GetData()->BlendShapeIndices.Num() == 1 && //LOD 0
+		Test.GetData()->MorphTargetCurveIndices.Num() == 1 && //LOD 0
+		Test.GetData()->BlendShapeIndices[0].Values.Num() == 1 && //at least one blendshape exists
+		Test.GetData()->BlendShapeIndices[0].Values[0] == 0 && //has index 0
+		Test.GetData()->MorphTargetCurveIndices[0].Values.Num() == 1 && //but morph target corresponding to that blendshape
+		Test.GetData()->MorphTargetCurveIndices[0].Values[0] == INDEX_NONE, //wasn't found
 		TEXT("MapMorphTargets - ValidReader MorphTargetWithNameMismatch")
 	);
 
@@ -474,8 +475,8 @@ IMPLEMENT_RIGUNIT_AUTOMATION_TEST(FRigUnit_RigLogic)
 	Test.Exec_MapMorphTargets(TestMorphTargetCurveValid.Get());
 	//Assert
 	AddErrorIfFalse(
-		Test.GetData()->Mappings->MorphTargetCurveIndices.Num() == 0 &&
-		Test.GetData()->Mappings->BlendShapeIndices.Num() == 0,
+		Test.GetData()->MorphTargetCurveIndices.Num() == 0 &&
+		Test.GetData()->BlendShapeIndices.Num() == 0,
 		TEXT("MapMorphTargets - EmptyReader ValidMorphTargetCurve")
 	);
 
@@ -489,9 +490,9 @@ IMPLEMENT_RIGUNIT_AUTOMATION_TEST(FRigUnit_RigLogic)
 	Test.Exec_MapMorphTargets(TestMorphTargetCurveValid.Get());
 	//Assert
 	AddErrorIfFalse(
-		Test.GetData()->Mappings->BlendShapeIndices.Num() == 1 && //LOD 0 exists
-		Test.GetData()->Mappings->BlendShapeIndices[0].Values.Num() == 0 && //but no blend shapes mapped
-		Test.GetData()->Mappings->MorphTargetCurveIndices[0].Values.Num() == 0, //or morph targets
+		Test.GetData()->BlendShapeIndices.Num() == 1 && //LOD 0 exists
+		Test.GetData()->BlendShapeIndices[0].Values.Num() == 0 && //but no blend shapes mapped
+		Test.GetData()->MorphTargetCurveIndices[0].Values.Num() == 0, //or morph targets
 		TEXT("MapMorphTargets - NoBlendShapes ValidMorphTargetCurve")
 	);
 
@@ -504,11 +505,11 @@ IMPLEMENT_RIGUNIT_AUTOMATION_TEST(FRigUnit_RigLogic)
 	Test.Exec_MapMorphTargets(TestCurveContainerEmpty.Get());
 	//Assert
 	AddErrorIfFalse(
-		Test.GetData()->Mappings->BlendShapeIndices.Num() == 1 && //LOD 0
-		Test.GetData()->Mappings->BlendShapeIndices[0].Values.Num() == 1 && //one blend shape
-		Test.GetData()->Mappings->BlendShapeIndices[0].Values[0] == 0 && //of index 0
-		Test.GetData()->Mappings->MorphTargetCurveIndices[0].Values.Num() == 1 && //we put in a morph target index corresponding to it
-		Test.GetData()->Mappings->MorphTargetCurveIndices[0].Values[0] == INDEX_NONE, //but just to signal that it is not found
+		Test.GetData()->BlendShapeIndices.Num() == 1 && //LOD 0
+		Test.GetData()->BlendShapeIndices[0].Values.Num() == 1 && //one blend shape
+		Test.GetData()->BlendShapeIndices[0].Values[0] == 0 && //of index 0
+		Test.GetData()->MorphTargetCurveIndices[0].Values.Num() == 1 && //we put in a morph target index corresponding to it
+		Test.GetData()->MorphTargetCurveIndices[0].Values[0] == INDEX_NONE, //but just to signal that it is not found
 		TEXT("MapMorphTargets - ValidReader EmptyCurveContainer")
 	);
 
@@ -521,11 +522,11 @@ IMPLEMENT_RIGUNIT_AUTOMATION_TEST(FRigUnit_RigLogic)
 	Test.Exec_MapMorphTargets(TestMorphTargetCurveValid.Get());
 	//Assert
 	AddErrorIfFalse(
-		Test.GetData()->Mappings->BlendShapeIndices.Num() == 1 && //LOD 0
-		Test.GetData()->Mappings->BlendShapeIndices[0].Values.Num() == 1 && //one blend shape (empty named)
-		Test.GetData()->Mappings->BlendShapeIndices[0].Values[0] == 0 && //of index 0
-		Test.GetData()->Mappings->MorphTargetCurveIndices[0].Values.Num() == 1 && //we put in a morph target index corresponding to it
-		Test.GetData()->Mappings->MorphTargetCurveIndices[0].Values[0] == INDEX_NONE, //but just to signal that it is not found
+		Test.GetData()->BlendShapeIndices.Num() == 1 && //LOD 0
+		Test.GetData()->BlendShapeIndices[0].Values.Num() == 1 && //one blend shape (empty named)
+		Test.GetData()->BlendShapeIndices[0].Values[0] == 0 && //of index 0
+		Test.GetData()->MorphTargetCurveIndices[0].Values.Num() == 1 && //we put in a morph target index corresponding to it
+		Test.GetData()->MorphTargetCurveIndices[0].Values[0] == INDEX_NONE, //but just to signal that it is not found
 		TEXT("MapMorphTargets - InvalidReader ValidMorphTargetCurve")
 	);
 
@@ -538,11 +539,11 @@ IMPLEMENT_RIGUNIT_AUTOMATION_TEST(FRigUnit_RigLogic)
 	Test.Exec_MapMorphTargets(TestMorphTargetCurvesInvalid.Get());
 	//Assert
 	AddErrorIfFalse(
-		Test.GetData()->Mappings->BlendShapeIndices.Num() == 1 && //LOD 0
-		Test.GetData()->Mappings->BlendShapeIndices[0].Values.Num() == 1 && //one blend shape
-		Test.GetData()->Mappings->BlendShapeIndices[0].Values[0] == 0 && //of index 0
-		Test.GetData()->Mappings->MorphTargetCurveIndices[0].Values.Num() == 1 && //we put in a morph target index corresponding to it
-		Test.GetData()->Mappings->MorphTargetCurveIndices[0].Values[0] == INDEX_NONE, //but just to signal that it is not found
+		Test.GetData()->BlendShapeIndices.Num() == 1 && //LOD 0
+		Test.GetData()->BlendShapeIndices[0].Values.Num() == 1 && //one blend shape
+		Test.GetData()->BlendShapeIndices[0].Values[0] == 0 && //of index 0
+		Test.GetData()->MorphTargetCurveIndices[0].Values.Num() == 1 && //we put in a morph target index corresponding to it
+		Test.GetData()->MorphTargetCurveIndices[0].Values[0] == INDEX_NONE, //but just to signal that it is not found
 		TEXT("MapMorphTargets - ValidReader InvalidMorphTargetCurve")
 	);
 
@@ -555,12 +556,12 @@ IMPLEMENT_RIGUNIT_AUTOMATION_TEST(FRigUnit_RigLogic)
 	Test.Exec_MapMorphTargets(TestMorphTargetCurveValid.Get());
 	//Assert
 	AddErrorIfFalse(
-		Test.GetData()->Mappings->BlendShapeIndices.Num() == 1 && //LOD 0
-		Test.GetData()->Mappings->MorphTargetCurveIndices.Num() == 1 && //LOD 0
-		Test.GetData()->Mappings->BlendShapeIndices[0].Values.Num() == 1 && //at least one blendshape exists
-		Test.GetData()->Mappings->BlendShapeIndices[0].Values[0] == 0 && //has index 0
-		Test.GetData()->Mappings->MorphTargetCurveIndices[0].Values.Num() == 1 && //morph target corresponding to that blendshape exists
-		Test.GetData()->Mappings->MorphTargetCurveIndices[0].Values[0] == 0, //and actually points to the right index
+		Test.GetData()->BlendShapeIndices.Num() == 1 && //LOD 0
+		Test.GetData()->MorphTargetCurveIndices.Num() == 1 && //LOD 0
+		Test.GetData()->BlendShapeIndices[0].Values.Num() == 1 && //at least one blendshape exists
+		Test.GetData()->BlendShapeIndices[0].Values[0] == 0 && //has index 0
+		Test.GetData()->MorphTargetCurveIndices[0].Values.Num() == 1 && //morph target corresponding to that blendshape exists
+		Test.GetData()->MorphTargetCurveIndices[0].Values[0] == 0, //and actually points to the right index
 		TEXT("MapMorphTargets - ValidReader ValidTestMorphTarget")
 	);
 
@@ -583,28 +584,28 @@ IMPLEMENT_RIGUNIT_AUTOMATION_TEST(FRigUnit_RigLogic)
 	Test.Exec_MapMorphTargets(TestMorphTargetTwoCurves.Get());
 	//Assert
 	AddErrorIfFalse(
-		Test.GetData()->Mappings->BlendShapeIndices.Num() == 2 &&  //2 LODs
-		Test.GetData()->Mappings->MorphTargetCurveIndices.Num() == 2,
+		Test.GetData()->BlendShapeIndices.Num() == 2 &&  //2 LODs
+		Test.GetData()->MorphTargetCurveIndices.Num() == 2,
 		TEXT("MapMorphTargets LOD0(AB) LOD1(A) - Expected 2 LODs for both blendshapes and morph targets")
 	);
 
 	AddErrorIfFalse(
-		Test.GetData()->Mappings->BlendShapeIndices.Num() == 2 &&         //condition repeated for crash prevention
-		Test.GetData()->Mappings->BlendShapeIndices[0].Values.Num() == 2 && //two blendshapes at LOD 0
-		Test.GetData()->Mappings->BlendShapeIndices[0].Values[0] == 0 &&  // A
-		Test.GetData()->Mappings->BlendShapeIndices[0].Values[1] == 1 &&  // B
-		Test.GetData()->Mappings->BlendShapeIndices[1].Values.Num() == 1 && //one blendshape at LOD 1
-		Test.GetData()->Mappings->BlendShapeIndices[1].Values[0] == 0,    // A
-		TEXT("MapMorphTargets LOD0(AB) LOD1(A) - resulting blendshape indices not correct")
+		Test.GetData()->BlendShapeIndices.Num() == 2 &&         //condition repeated for crash prevention
+		Test.GetData()->BlendShapeIndices[0].Values.Num() == 2 && //two blendshapes at LOD 0
+		Test.GetData()->BlendShapeIndices[0].Values[0] == 0 &&  // A
+		Test.GetData()->BlendShapeIndices[0].Values[1] == 1 &&  // B
+		Test.GetData()->BlendShapeIndices[1].Values.Num() == 1 && //one blendshape at LOD 1
+		Test.GetData()->BlendShapeIndices[1].Values[0] == 0,    // A
+			TEXT("MapMorphTargets LOD0(AB) LOD1(A) - resulting blendshape indices not correct")
 	);
 
 	AddErrorIfFalse(
-		Test.GetData()->Mappings->MorphTargetCurveIndices.Num() == 2 &&         //condition repeated for crash prevention
-		Test.GetData()->Mappings->MorphTargetCurveIndices[0].Values.Num() == 2 && //two morph targets at LOD 0
-		Test.GetData()->Mappings->MorphTargetCurveIndices[0].Values[0] == 0 &&  // A
-		Test.GetData()->Mappings->MorphTargetCurveIndices[0].Values[1] == 1 &&  // B
-		Test.GetData()->Mappings->MorphTargetCurveIndices[1].Values.Num() == 1 && //one morph target at LOD 1
-		Test.GetData()->Mappings->MorphTargetCurveIndices[1].Values[0] == 0,    // A
+		Test.GetData()->MorphTargetCurveIndices.Num() == 2 &&         //condition repeated for crash prevention
+		Test.GetData()->MorphTargetCurveIndices[0].Values.Num() == 2 && //two morph targets at LOD 0
+		Test.GetData()->MorphTargetCurveIndices[0].Values[0] == 0 &&  // A
+		Test.GetData()->MorphTargetCurveIndices[0].Values[1] == 1 &&  // B
+		Test.GetData()->MorphTargetCurveIndices[1].Values.Num() == 1 && //one morph target at LOD 1
+		Test.GetData()->MorphTargetCurveIndices[1].Values[0] == 0,    // A
 		TEXT("MapMorphTargets LOD0(AB) LOD1(A) - resulting morph target indices not correct")
 	);
 
@@ -621,25 +622,25 @@ IMPLEMENT_RIGUNIT_AUTOMATION_TEST(FRigUnit_RigLogic)
 	Test.Exec_MapMorphTargets(TestMorphTargetTwoCurves.Get());
 	//Assert
 	AddErrorIfFalse(
-		Test.GetData()->Mappings->BlendShapeIndices.Num() == 1 &&  //1 LOD
-		Test.GetData()->Mappings->MorphTargetCurveIndices.Num() == 1, //1 LOD
+		Test.GetData()->BlendShapeIndices.Num() == 1 &&  //1 LOD
+		Test.GetData()->MorphTargetCurveIndices.Num() == 1, //1 LOD
 		TEXT("MapMorphTargets LOD0(AB) LOD1(-) - Expected 1 LOD for both blendshapes and morph targets")
 	);
 
 
 	AddErrorIfFalse(
-		Test.GetData()->Mappings->BlendShapeIndices.Num() == 1 &&         //condition repeated for crash prevention
-		Test.GetData()->Mappings->BlendShapeIndices[0].Values.Num() == 2 && //two blendshapes at LOD 0
-		Test.GetData()->Mappings->BlendShapeIndices[0].Values[0] == 0 &&  // A
-		Test.GetData()->Mappings->BlendShapeIndices[0].Values[1] == 1,    // B
+		Test.GetData()->BlendShapeIndices.Num() == 1 &&         //condition repeated for crash prevention
+		Test.GetData()->BlendShapeIndices[0].Values.Num() == 2 && //two blendshapes at LOD 0
+		Test.GetData()->BlendShapeIndices[0].Values[0] == 0 &&  // A
+		Test.GetData()->BlendShapeIndices[0].Values[1] == 1,    // B
 		TEXT("MapMorphTargets LOD0(AB) LOD1(-) - Resulting blendshapes not correct")
 	);
 
 	AddErrorIfFalse(
-		Test.GetData()->Mappings->MorphTargetCurveIndices.Num() == 1 &&         //condition repeated for crash prevention
-		Test.GetData()->Mappings->MorphTargetCurveIndices[0].Values.Num() == 2 && //two morph targets at LOD 0
-		Test.GetData()->Mappings->MorphTargetCurveIndices[0].Values[0] == 0 &&  // A
-		Test.GetData()->Mappings->MorphTargetCurveIndices[0].Values[1] == 1,    // B
+		Test.GetData()->MorphTargetCurveIndices.Num() == 1 &&         //condition repeated for crash prevention
+		Test.GetData()->MorphTargetCurveIndices[0].Values.Num() == 2 && //two morph targets at LOD 0
+		Test.GetData()->MorphTargetCurveIndices[0].Values[0] == 0 &&  // A
+		Test.GetData()->MorphTargetCurveIndices[0].Values[1] == 1,    // B
 		TEXT("MapMorphTargets LOD0(AB) LOD1(-) - Resulting morph targets not correct")
 	);
 
@@ -655,26 +656,26 @@ IMPLEMENT_RIGUNIT_AUTOMATION_TEST(FRigUnit_RigLogic)
 	Test.Exec_MapMorphTargets(TestMorphTargetTwoCurves.Get());
 	//Assert
 	AddErrorIfFalse(
-		Test.GetData()->Mappings->BlendShapeIndices.Num() == 2 &&  //2 LODs
-		Test.GetData()->Mappings->MorphTargetCurveIndices.Num() == 2, //2 LODs
+		Test.GetData()->BlendShapeIndices.Num() == 2 &&  //2 LODs
+		Test.GetData()->MorphTargetCurveIndices.Num() == 2, //2 LODs
 		TEXT("MapMorphTargets LOD0(A) LOD1(B) - Expected 2 LODs for both blendshapes and morph targets")
 	);
 
 	AddErrorIfFalse(
-		Test.GetData()->Mappings->BlendShapeIndices.Num() == 2 &&         //condition repeated for crash prevention
-		Test.GetData()->Mappings->BlendShapeIndices[0].Values.Num() == 1 && //1 blendshape at LOD 0
-		Test.GetData()->Mappings->BlendShapeIndices[0].Values[0] == 0 &&  // A
-		Test.GetData()->Mappings->BlendShapeIndices[1].Values.Num() == 1 && //1 blendshape at LOD 1
-		Test.GetData()->Mappings->BlendShapeIndices[1].Values[0] == 1,    // B
+		Test.GetData()->BlendShapeIndices.Num() == 2 &&         //condition repeated for crash prevention
+		Test.GetData()->BlendShapeIndices[0].Values.Num() == 1 && //1 blendshape at LOD 0
+		Test.GetData()->BlendShapeIndices[0].Values[0] == 0 &&  // A
+		Test.GetData()->BlendShapeIndices[1].Values.Num() == 1 && //1 blendshape at LOD 1
+		Test.GetData()->BlendShapeIndices[1].Values[0] == 1,    // B
 		TEXT("MapMorphTargets LOD0(A) LOD1(B) - Resulting blendshape indices not correct")
 	);
 
 	AddErrorIfFalse(
-		Test.GetData()->Mappings->MorphTargetCurveIndices.Num() == 2 &&         //condition repeated for crash prevention
-		Test.GetData()->Mappings->MorphTargetCurveIndices[0].Values.Num() == 1 && //1 morph target at LOD 0
-		Test.GetData()->Mappings->MorphTargetCurveIndices[0].Values[0] == 0 &&  // A
-		Test.GetData()->Mappings->MorphTargetCurveIndices[1].Values.Num() == 1 && //1 morph target at LOD 1
-		Test.GetData()->Mappings->MorphTargetCurveIndices[1].Values[0] == 1,    // B
+		Test.GetData()->MorphTargetCurveIndices.Num() == 2 &&         //condition repeated for crash prevention
+		Test.GetData()->MorphTargetCurveIndices[0].Values.Num() == 1 && //1 morph target at LOD 0
+		Test.GetData()->MorphTargetCurveIndices[0].Values[0] == 0 &&  // A
+		Test.GetData()->MorphTargetCurveIndices[1].Values.Num() == 1 && //1 morph target at LOD 1
+		Test.GetData()->MorphTargetCurveIndices[1].Values[0] == 1,    // B
 		TEXT("MapMorphTargets LOD0(A) LOD1(B) - Resulting morph target indices not correct")
 	);
 
@@ -691,12 +692,12 @@ IMPLEMENT_RIGUNIT_AUTOMATION_TEST(FRigUnit_RigLogic)
 	Test.Exec_MapMaskMultipliers(TestCurveContainerForAnimMapsNameMismatch.Get());
 	//Assert
 	AddErrorIfFalse(
-		Test.GetData()->Mappings->CurveElementIndicesForAnimMaps.Num() == 1 &&
-		Test.GetData()->Mappings->CurveElementIndicesForAnimMaps[0].Values.Num() == 1 &&
-		Test.GetData()->Mappings->CurveElementIndicesForAnimMaps[0].Values[0] == INDEX_NONE &&
-		Test.GetData()->Mappings->RigLogicIndicesForAnimMaps.Num() == 1 &&
-		Test.GetData()->Mappings->RigLogicIndicesForAnimMaps[0].Values.Num() == 1 &&
-		Test.GetData()->Mappings->RigLogicIndicesForAnimMaps[0].Values[0] == 0,
+		Test.GetData()->CurveElementIndicesForAnimMaps.Num() == 1 &&
+		Test.GetData()->CurveElementIndicesForAnimMaps[0].Values.Num() == 1 &&
+		Test.GetData()->CurveElementIndicesForAnimMaps[0].Values[0] == INDEX_NONE &&
+		Test.GetData()->RigLogicIndicesForAnimMaps.Num() == 1 &&
+		Test.GetData()->RigLogicIndicesForAnimMaps[0].Values.Num() == 1 &&
+		Test.GetData()->RigLogicIndicesForAnimMaps[0].Values[0] == 0,
 		TEXT("MapMaskMultipliers - ValidReader ValidAnimatedMapNameMismatch")
 	);
 
@@ -709,8 +710,8 @@ IMPLEMENT_RIGUNIT_AUTOMATION_TEST(FRigUnit_RigLogic)
 	Test.Exec_MapMaskMultipliers(TestCurveContainerForAnimMapsValid.Get());
 	//Assert
 	AddErrorIfFalse(
-		Test.GetData()->Mappings->CurveElementIndicesForAnimMaps.Num() == 0 &&
-		Test.GetData()->Mappings->RigLogicIndicesForAnimMaps.Num() == 0,
+		Test.GetData()->CurveElementIndicesForAnimMaps.Num() == 0 &&
+		Test.GetData()->RigLogicIndicesForAnimMaps.Num() == 0,
 		TEXT("MapMaskMultipliers - EmptyReader ValidAnimatedMap")
 	);
 
@@ -721,12 +722,12 @@ IMPLEMENT_RIGUNIT_AUTOMATION_TEST(FRigUnit_RigLogic)
 	Test.Exec_MapMaskMultipliers(TestCurveContainerEmpty.Get());
 	//Assert
 	AddErrorIfFalse(
-		Test.GetData()->Mappings->CurveElementIndicesForAnimMaps.Num() == 1 &&
-		Test.GetData()->Mappings->CurveElementIndicesForAnimMaps[0].Values.Num() == 1 &&
-		Test.GetData()->Mappings->CurveElementIndicesForAnimMaps[0].Values[0] == INDEX_NONE &&
-		Test.GetData()->Mappings->RigLogicIndicesForAnimMaps.Num() == 1 &&
-		Test.GetData()->Mappings->RigLogicIndicesForAnimMaps[0].Values.Num() == 1 &&
-		Test.GetData()->Mappings->RigLogicIndicesForAnimMaps[0].Values[0] == 0,
+		Test.GetData()->CurveElementIndicesForAnimMaps.Num() == 1 &&
+		Test.GetData()->CurveElementIndicesForAnimMaps[0].Values.Num() == 1 &&
+		Test.GetData()->CurveElementIndicesForAnimMaps[0].Values[0] == INDEX_NONE &&
+		Test.GetData()->RigLogicIndicesForAnimMaps.Num() == 1 &&
+		Test.GetData()->RigLogicIndicesForAnimMaps[0].Values.Num() == 1 &&
+		Test.GetData()->RigLogicIndicesForAnimMaps[0].Values[0] == 0,
 		TEXT("MapMaskMultipliers - ValidReader EmptyCurveContainer")
 	);
 
@@ -739,12 +740,12 @@ IMPLEMENT_RIGUNIT_AUTOMATION_TEST(FRigUnit_RigLogic)
 	Test.Exec_MapMaskMultipliers(TestCurveContainerForAnimMapsValid.Get());
 	//Assert
 	AddErrorIfFalse(
-		Test.GetData()->Mappings->CurveElementIndicesForAnimMaps.Num() == 1 &&
-		Test.GetData()->Mappings->CurveElementIndicesForAnimMaps[0].Values.Num() == 1 &&
-		Test.GetData()->Mappings->CurveElementIndicesForAnimMaps[0].Values[0] == 0 &&
-		Test.GetData()->Mappings->RigLogicIndicesForAnimMaps.Num() == 1 &&
-		Test.GetData()->Mappings->RigLogicIndicesForAnimMaps[0].Values.Num() == 1 &&
-		Test.GetData()->Mappings->RigLogicIndicesForAnimMaps[0].Values[0] == 0,
+		Test.GetData()->CurveElementIndicesForAnimMaps.Num() == 1 &&
+		Test.GetData()->CurveElementIndicesForAnimMaps[0].Values.Num() == 1 &&
+		Test.GetData()->CurveElementIndicesForAnimMaps[0].Values[0] == 0 &&
+		Test.GetData()->RigLogicIndicesForAnimMaps.Num() == 1 &&
+		Test.GetData()->RigLogicIndicesForAnimMaps[0].Values.Num() == 1 &&
+		Test.GetData()->RigLogicIndicesForAnimMaps[0].Values[0] == 0,
 		TEXT("MapMaskMultipliers - Valid Inputs")
 	);
 
@@ -755,8 +756,8 @@ IMPLEMENT_RIGUNIT_AUTOMATION_TEST(FRigUnit_RigLogic)
 	//-----
 	//create neutral transforms for two bones
 	const uint8 TransformArraySize = 2 * FRigUnit_RigLogic::TestAccessor::MAX_ATTRS_PER_JOINT;
-	float Values[TransformArraySize];  //two bones, nine attributes
-	FTransformArrayView TwoJointNeutralTransforms = Test.CreateTwoJointNeutralTransforms(Values);
+	float NeutralValues[TransformArraySize];  //two bones, nine attributes
+	FTransformArrayView TwoJointNeutralTransforms = Test.CreateTwoJointNeutralTransforms(NeutralValues);
 	//create delta transforms
 	float DeltaTransformData[TransformArraySize] = { 0.f };
 	//first bone translation
@@ -767,17 +768,14 @@ IMPLEMENT_RIGUNIT_AUTOMATION_TEST(FRigUnit_RigLogic)
 	DeltaTransformData[9] = 1.f;
 	DeltaTransformData[10] = 2.f;
 	DeltaTransformData[11] = 7.f;
-	FTransformArrayView DeltaTransforms = FTransformArrayView(DeltaTransformData, sizeof(FTransform));
+	FTransformArrayView DeltaTransforms = FTransformArrayView(DeltaTransformData, TransformArraySize);
 	//create variable joint index arrays for two bones
-	SharedRigRuntimeContext->VariableJointIndices.Reset();
-	SharedRigRuntimeContext->VariableJointIndices.AddDefaulted(1);
-	SharedRigRuntimeContext->VariableJointIndices[0].Values.Add(0);
-	SharedRigRuntimeContext->VariableJointIndices[0].Values.Add(1);
-	FRigUnit_RigLogic_JointUpdateParams TestJointUpdateParamsTwoJoints_LOD0(
-		TwoJointNeutralTransforms,
-		DeltaTransforms);
+	SharedRigRuntimeContext->VariableJointIndicesPerLOD.Reset();
+	SharedRigRuntimeContext->VariableJointIndicesPerLOD.AddDefaulted(1);
+	SharedRigRuntimeContext->VariableJointIndicesPerLOD[0].Values.Add(0);
+	SharedRigRuntimeContext->VariableJointIndicesPerLOD[0].Values.Add(1);
 	//Test
-	Test.Exec_UpdateJoints(TestHierarchyTwoBones.Get(), TestJointUpdateParamsTwoJoints_LOD0);
+	Test.Exec_UpdateJoints(TestHierarchyTwoBones.Get(), TArrayView<const float>{NeutralValues}, TArrayView<const float>{DeltaTransformData});
 	//Assert
 	//Note that BoneB.GlobalTransform.Z should be zero since the scale.Z is zero. Also, translation Y becomes -Y 
 	AddErrorIfFalse(TestHierarchyTwoBones->GetGlobalTransform(0).GetTranslation().Equals(FVector(1.f, 0.f, 0.f)), TEXT("UpdateJoints LOD0 Bone 01 - unexpected transform"));
diff --git a/Engine/Plugins/Animation/RigLogic/Source/RigLogicModule/Private/SharedRigRuntimeContext.cpp b/Engine/Plugins/Animation/RigLogic/Source/RigLogicModule/Private/SharedRigRuntimeContext.cpp
new file mode 100644
--- /dev/null
+++ b/Engine/Plugins/Animation/RigLogic/Source/RigLogicModule/Private/SharedRigRuntimeContext.cpp
@@ -0,0 +1,32 @@
+// Copyright Epic Games, Inc. All Rights Reserved.
+
+#include "SharedRigRuntimeContext.h"
+
+#include "CoreMinimal.h"
+#include "HAL/LowLevelMemTracker.h"
+#include "DNAReader.h"
+#include "RigLogic.h"
+
+static constexpr uint16 NUM_ATTRS_PER_JOINT = 9;
+
+void FSharedRigRuntimeContext::CacheVariableJointIndices()
+{
+	LLM_SCOPE_BYNAME(TEXT("Animation/RigLogic"));
+
+	const uint16 LODCount = BehaviorReader->GetLODCount();
+	VariableJointIndicesPerLOD.Reset();
+	VariableJointIndicesPerLOD.AddDefaulted(LODCount);
+	TSet<uint16> DistinctVariableJointIndices;
+	for (uint16 LODIndex = 0; LODIndex < LODCount; ++LODIndex)
+	{
+		TArrayView<const uint16> VariableAttributeIndices = BehaviorReader->GetJointVariableAttributeIndices(LODIndex);
+		DistinctVariableJointIndices.Reset();
+		DistinctVariableJointIndices.Reserve(VariableAttributeIndices.Num());
+		for (const uint16 AttrIndex : VariableAttributeIndices)
+		{
+			const uint16 JointIndex = AttrIndex / NUM_ATTRS_PER_JOINT;
+			DistinctVariableJointIndices.Add(JointIndex);
+		}
+		VariableJointIndicesPerLOD[LODIndex].Values = DistinctVariableJointIndices.Array();
+	}
+}
diff --git a/Engine/Plugins/Animation/RigLogic/Source/RigLogicModule/Private/SharedRigRuntimeContext.h b/Engine/Plugins/Animation/RigLogic/Source/RigLogicModule/Private/SharedRigRuntimeContext.h
new file mode 100644
--- /dev/null
+++ b/Engine/Plugins/Animation/RigLogic/Source/RigLogicModule/Private/SharedRigRuntimeContext.h
@@ -0,0 +1,29 @@
+// Copyright Epic Games, Inc. All Rights Reserved.
+
+#pragma once
+
+#include "CoreMinimal.h"
+#include "DNAReader.h"
+#include "RigLogic.h"
+
+struct FSharedRigRuntimeContext
+{
+	template <typename T>
+	struct TNestedArray
+	{
+		TArray<T> Values;
+	};
+
+	// Redundant but stored here as well so the runtime context can be updated / queried atomically
+	TSharedPtr<IBehaviorReader> BehaviorReader;
+
+	/** RigLogic itself is stateless, and is designed to be shared between
+	  * multiple rig instances based on the same DNA.
+	**/
+	TSharedPtr<FRigLogic> RigLogic;
+
+	/** Cached joint indices that need to be updated for each LOD **/
+	TArray<TNestedArray<uint16>> VariableJointIndicesPerLOD;
+
+	void CacheVariableJointIndices();
+};
diff --git a/Engine/Plugins/Animation/RigLogic/Source/RigLogicModule/Public/AnimNode_RigLogic.h b/Engine/Plugins/Animation/RigLogic/Source/RigLogicModule/Public/AnimNode_RigLogic.h
new file mode 100644
--- /dev/null
+++ b/Engine/Plugins/Animation/RigLogic/Source/RigLogicModule/Public/AnimNode_RigLogic.h
@@ -0,0 +1,50 @@
+// Copyright Epic Games, Inc. All Rights Reserved.
+
+#pragma once
+
+#include "DNAReader.h"
+#include "RigInstance.h"
+#include "RigLogic.h"
+
+#include "Animation/AnimNodeBase.h"
+#include "Animation/SmartName.h"
+
+#include "AnimNode_RigLogic.generated.h"
+
+DECLARE_LOG_CATEGORY_EXTERN(LogRigLogicAnimNode, Log, All);
+
+struct FSharedRigRuntimeContext;
+struct FDNAIndexMapping;
+
+USTRUCT(BlueprintInternalUseOnly)
+struct RIGLOGICMODULE_API FAnimNode_RigLogic : public FAnimNode_Base
+{
+public:
+	GENERATED_USTRUCT_BODY()
+
+	FAnimNode_RigLogic();
+	~FAnimNode_RigLogic();
+
+	void Initialize_AnyThread(const FAnimationInitializeContext& Context) override;
+	void CacheBones_AnyThread(const FAnimationCacheBonesContext& Context) override;
+	void Update_AnyThread(const FAnimationUpdateContext& Context) override;
+	void Evaluate_AnyThread(FPoseContext& Output) override;
+	void GatherDebugData(FNodeDebugData& DebugData) override;
+
+public:
+	UPROPERTY(EditAnywhere, BlueprintReadWrite, Category = Links)
+	FPoseLink AnimSequence;
+
+private:
+	void UpdateControlCurves(const FPoseContext& OutputContext, const FDNAIndexMapping* DNAIndexMapping);
+	void CalculateRigLogic(FRigLogic* RigLogic);
+	void UpdateJoints(TArrayView<const uint16> VariableJointIndices, TArrayView<const float> NeutralJointValues, TArrayView<const float> DeltaJointValues, FPoseContext& OutputContext);
+	void UpdateBlendShapeCurves(const FDNAIndexMapping* DNAIndexMapping, TArrayView<const float> BlendShapeValues, FPoseContext& OutputContext);
+	void UpdateAnimMapCurves(const FDNAIndexMapping* DNAIndexMapping, TArrayView<const float> AnimMapOutputs, FPoseContext& OutputContext);
+
+private:
+	TSharedPtr<FSharedRigRuntimeContext> LocalRigRuntimeContext;
+	TSharedPtr<FDNAIndexMapping> LocalDNAIndexMapping;
+	FRigInstance* RigInstance;
+	TArray<FCompactPoseBoneIndex> JointsMapDNAIndicesToCompactPoseBoneIndices;
+};
diff --git a/Engine/Plugins/Animation/RigLogic/Source/RigLogicModule/Public/DNAAsset.h b/Engine/Plugins/Animation/RigLogic/Source/RigLogicModule/Public/DNAAsset.h
--- a/Engine/Plugins/Animation/RigLogic/Source/RigLogicModule/Public/DNAAsset.h
+++ b/Engine/Plugins/Animation/RigLogic/Source/RigLogicModule/Public/DNAAsset.h
@@ -20,41 +20,14 @@ class IBehaviorReader;
 class IGeometryReader;
 class FRigLogicMemoryStream;
 class UAssetUserData;
-enum class EDNADataLayer : uint8;
-
-struct FSharedRigRuntimeContext
-{
-	template <typename T>
-	struct TNestedArray
-	{
-		TArray<T> Values;
-	};
-
-	uint32 DNAHash;
-
-	/** Part of the .dna file needed for run-time execution of RigLogic;
-	 **/
-	TSharedPtr<IBehaviorReader> BehaviorReader;
-
-	/** Part of the .dna file used design-time for updating SkeletalMesh geometry
-	 **/
-	TSharedPtr<IGeometryReader> GeometryReader;
-
-	/** RigLogic itself is stateless, and is designed to be shared between
-	  * multiple rig instances based on the same DNA.
-	**/
-	TSharedPtr<FRigLogic> RigLogic;
-
-	/** Cached joint indices that need to be updated for each LOD **/
-	TArray<TNestedArray<uint16>> VariableJointIndices;
-};
-
-
-/** An asset holding the data needed to generate/update/animate a RigLogic character
- * It is imported from character's DNA file as a bit stream, and separated out it into runtime (behavior) and design-time chunks;
- * Currently, the design-time part still loads the geometry, as it is needed for the skeletal mesh update; once SkeletalMeshDNAReader is
- * fully implemented, it will be able to read the geometry directly from the SkeletalMesh and won't load it into this asset
- **/
+struct FSharedRigRuntimeContext;
+enum class EDNADataLayer: uint8;
+
+ /** An asset holding the data needed to generate/update/animate a RigLogic character
+  * It is imported from character's DNA file as a bit stream, and separated out it into runtime (behavior) and design-time chunks;
+  * Currently, the design-time part still loads the geometry, as it is needed for the skeletal mesh update; once SkeletalMeshDNAReader is
+  * fully implemented, it will be able to read the geometry directly from the SkeletalMesh and won't load it into this asset 
+  **/
 UCLASS(NotBlueprintable, hidecategories = (Object))
 class RIGLOGICMODULE_API UDNAAsset : public UAssetUserData
 {
@@ -62,26 +35,20 @@ class RIGLOGICMODULE_API UDNAAsset : public UAssetUserData
 
 public:
 	UDNAAsset();
+	~UDNAAsset();
 
 #if WITH_EDITORONLY_DATA
 	UPROPERTY(VisibleAnywhere, Instanced, Category = ImportSettings)
-		class UAssetImportData* AssetImportData;
+	class UAssetImportData* AssetImportData;
 #endif
 
-	TSharedPtr<IBehaviorReader> GetBehaviorReader()
-	{
-		return Context.BehaviorReader;
-	}
-
+	TSharedPtr<IBehaviorReader> GetBehaviorReader();
 #if WITH_EDITORONLY_DATA
-	TSharedPtr<IGeometryReader> GetGeometryReader()
-	{
-		return Context.GeometryReader;
-	}
+	TSharedPtr<IGeometryReader> GetGeometryReader();
 #endif
 
 	UPROPERTY()
-		FString DNAFileName;
+	FString DNAFileName; 
 
 	bool Init(const FString& Filename);
 	void Serialize(FArchive& Ar) override;
@@ -93,10 +60,36 @@ public:
 	void SetGeometryReader(TSharedPtr<IDNAReader> SourceDNAReader);
 
 private:
+	friend struct FAnimNode_RigLogic;
+	friend struct FRigUnit_RigLogic;
+
+	enum class EDNARetentionPolicy
+	{
+		Keep,
+		Unload
+	};
+
+	TSharedPtr<FSharedRigRuntimeContext> GetRigRuntimeContext(EDNARetentionPolicy Policy);
+	void InvalidateRigRuntimeContext();
+
+private:
+	// Synchronize DNA updates
+	FCriticalSection DNAUpdateSection;
+
+	// Synchronize Rig Runtime Context updates
+	FRWLock RigRuntimeContextUpdateLock;
+
+	/** Part of the .dna file needed for run-time execution of RigLogic;
+	 **/
+	TSharedPtr<IBehaviorReader> BehaviorReader;
+
+	/** Part of the .dna file used design-time for updating SkeletalMesh geometry
+	 **/
+	TSharedPtr<IGeometryReader> GeometryReader;
+
 	/** Runtime data necessary for rig computations that is shared between
 	  * multiple rig instances based on the same DNA.
 	**/
-	friend struct FRigUnit_RigLogic;
-	FSharedRigRuntimeContext Context;
+	TSharedPtr<FSharedRigRuntimeContext> RigRuntimeContext;
 
 };
diff --git a/Engine/Plugins/Animation/RigLogic/Source/RigLogicModule/Public/DNAReader.h b/Engine/Plugins/Animation/RigLogic/Source/RigLogicModule/Public/DNAReader.h
--- a/Engine/Plugins/Animation/RigLogic/Source/RigLogicModule/Public/DNAReader.h
+++ b/Engine/Plugins/Animation/RigLogic/Source/RigLogicModule/Public/DNAReader.h
@@ -41,6 +41,7 @@ public:
 	virtual FString GetDBName() const = 0;
 
 	virtual dna::Reader* Unwrap() const = 0;
+	virtual void Unload(EDNADataLayer Layer) = 0;
 
 };
 
diff --git a/Engine/Plugins/Animation/RigLogic/Source/RigLogicModule/Public/DNAReaderAdapter.h b/Engine/Plugins/Animation/RigLogic/Source/RigLogicModule/Public/DNAReaderAdapter.h
--- a/Engine/Plugins/Animation/RigLogic/Source/RigLogicModule/Public/DNAReaderAdapter.h
+++ b/Engine/Plugins/Animation/RigLogic/Source/RigLogicModule/Public/DNAReaderAdapter.h
@@ -117,6 +117,8 @@ public:
 	TArrayView<const float> GetBlendShapeTargetDeltaZs(uint16 MeshIndex, uint16 BlendShapeTargetIndex) const override;
 	TArrayView<const uint32> GetBlendShapeTargetVertexIndices(uint16 MeshIndex, uint16 BlendShapeTargetIndex) const override;
 
+	void Unload(EDNADataLayer Layer) override;
+
 private:
 	dna::Reader* Unwrap() const override;
 
@@ -143,6 +145,12 @@ dna::Reader* FDNAReader<TWrappedReader>::Unwrap() const
 	return ReaderPtr.Get();
 }
 
+template <class TWrappedReader>
+void FDNAReader<TWrappedReader>::Unload(EDNADataLayer Layer)
+{
+	ReaderPtr->unload(static_cast<dna::DataLayer>(Layer));
+}
+
 template <class TWrappedReader>
 void FDNAReader<TWrappedReader>::FWrappedReaderDeleter::operator()(TWrappedReader* Pointer)
 {
diff --git a/Engine/Plugins/Animation/RigLogic/Source/RigLogicModule/Public/RigLogic.h b/Engine/Plugins/Animation/RigLogic/Source/RigLogicModule/Public/RigLogic.h
--- a/Engine/Plugins/Animation/RigLogic/Source/RigLogicModule/Public/RigLogic.h
+++ b/Engine/Plugins/Animation/RigLogic/Source/RigLogicModule/Public/RigLogic.h
@@ -42,6 +42,7 @@ public:
 	FRigLogic& operator=(FRigLogic&&) = default;
 
 	uint16 GetLODCount() const;
+	TArrayView<const float> GetRawNeutralJointValues() const;
 	FTransformArrayView GetNeutralJointValues() const;
 	uint16 GetJointGroupCount() const;
 	TArrayView<const uint16> GetJointVariableAttributeIndices(uint16 LOD) const;
diff --git a/Engine/Plugins/Animation/RigLogic/Source/RigLogicModule/Public/RigUnit_RigLogic.h b/Engine/Plugins/Animation/RigLogic/Source/RigLogicModule/Public/RigUnit_RigLogic.h
--- a/Engine/Plugins/Animation/RigLogic/Source/RigLogicModule/Public/RigUnit_RigLogic.h
+++ b/Engine/Plugins/Animation/RigLogic/Source/RigLogicModule/Public/RigUnit_RigLogic.h
@@ -7,7 +7,6 @@
 
 #include "ControlRig.h"
 #include "DNAAsset.h"
-#include "Containers/DiscardableKeyValueCache.h"
 #include "Materials/Material.h"
 #include "Units/RigUnit.h"
 
@@ -15,6 +14,7 @@
 
 class IBehaviorReader;
 class FTransformArrayView;
+struct FSharedRigRuntimeContext;
 
 DECLARE_LOG_CATEGORY_EXTERN(LogRigLogicUnit, Log, All);
 
@@ -29,34 +29,28 @@ struct FRigUnit_RigLogic_IntArray
 		TArray<int32> Values;
 };
 
-/**
- * A helper struct used inside of the RigUnit_RigLogic holding params for updating joints.
- * Note that these params are array views - they don't own the memory they're pointing to.
- */
-struct FRigUnit_RigLogic_JointUpdateParams
+/* The work data used by the FRigUnit_RigLogic */
+USTRUCT(meta = (DocumentationPolicy = "Strict"))
+struct FRigUnit_RigLogic_Data
 {
-	FRigUnit_RigLogic_JointUpdateParams(
-		FTransformArrayView InNeutralJointTransforms,
-		FTransformArrayView InDeltaTransforms
-	)
-		: NeutralJointTransforms(InNeutralJointTransforms)
-		, DeltaTransforms(InDeltaTransforms)
-	{ }
+	GENERATED_BODY()
 
-public:
-	FTransformArrayView NeutralJointTransforms;
-	FTransformArrayView DeltaTransforms; //the result of rig logic calculations
-};
+	FRigUnit_RigLogic_Data();
+ 	~FRigUnit_RigLogic_Data();
+ 	FRigUnit_RigLogic_Data(const FRigUnit_RigLogic_Data& Other);
+ 	FRigUnit_RigLogic_Data& operator=(const FRigUnit_RigLogic_Data& Other);
 
-/* A helper struct used inside of the RigUnit_RigLogic to store mappings between skelMesh and DNA. */
-USTRUCT(meta = (DocumentationPolicy = "Strict"))
-struct FRigDefinitionMapping
-{
+	/** Cached Skeletal Mesh Component **/
+	UPROPERTY(transient)
+	TWeakObjectPtr<USkeletalMeshComponent> SkelMeshComponent;
 
-	GENERATED_BODY()
+	TSharedPtr<FSharedRigRuntimeContext> LocalRigRuntimeContext;
 
-		FRigDefinitionMapping() = default;
-	FRigDefinitionMapping(const IBehaviorReader* DNABehavior, const URigHierarchy* InHierarchy);
+	/** RigInstance is a thin class, containing only character instance specific data
+	  * i.e. the output buffers where RigLogic writes the results of computations, specific
+	  * to each instance of a rig.
+	**/
+	TUniquePtr<FRigInstance> RigInstance;
 
 	/** Mapping RL indices to UE indices
 	  * Note: we use int32 instead of uint32 to allow storing INDEX_NONE for missing elements
@@ -65,92 +59,56 @@ struct FRigDefinitionMapping
 
 	/** RL input index to ControlRig's input curve index for each LOD **/
 	UPROPERTY(transient)
-		TArray<int32> InputCurveIndices;
+	TArray<int32> InputCurveIndices;
 
 	/** RL joint index to ControlRig's hierarchy bone index **/
 	UPROPERTY(transient)
-		TArray<int32> HierarchyBoneIndices;	/** RL mesh blend shape index to ControlRig's output blendshape curve index for each LOD **/
+	TArray<int32> HierarchyBoneIndices;
 
-		/** Mapping RL indices to UE indices
-		  * Note: we use int32 instead of uint32 to allow storing INDEX_NONE for missing elements
-		  * if value is valid, it is cast to appropriate uint type
-		**/
+	/** RL mesh blend shape index to ControlRig's output blendshape curve index for each LOD **/
 	UPROPERTY(transient)
-		TArray<FRigUnit_RigLogic_IntArray> MorphTargetCurveIndices;
+	TArray<FRigUnit_RigLogic_IntArray> MorphTargetCurveIndices;
 
 	/** RL mesh+blend shape array index to RL blend shape index for each LOD **/
 	UPROPERTY(transient)
-		TArray<FRigUnit_RigLogic_IntArray> BlendShapeIndices;
+	TArray<FRigUnit_RigLogic_IntArray> BlendShapeIndices;
 
 	/** RL animated map index to ControlRig's output anim map curve index for each LOD **/
 	UPROPERTY(transient)
-		TArray<FRigUnit_RigLogic_IntArray> CurveElementIndicesForAnimMaps;
+	TArray<FRigUnit_RigLogic_IntArray> CurveElementIndicesForAnimMaps;
 
 	/** RL animated map index to RL anim map curve index for each LOD **/
 	UPROPERTY(transient)
-		TArray<FRigUnit_RigLogic_IntArray> RigLogicIndicesForAnimMaps;
-
-private:
-	/** Makes a map of input curve indices from DNA file to the control rig curves **/
-	void MapInputCurveIndices(const IBehaviorReader* DNABehavior, const URigHierarchy* InHierarchy);
-	/** Uses names to map joint indices from DNA file to the indices of bones in control rig hierarchy **/
-	void MapJoints(const IBehaviorReader* DNABehavior, const URigHierarchy* Hierarchy);
-	/** Uses names of blend shapes and meshes in DNA file, for all LODs, to map their indices to the indices of
-	  * morph target curves in the curve container; curve name format is <mesh>__<blendshape> **/
-	void MapMorphTargets(const IBehaviorReader* DNABehavior, const URigHierarchy* InHierarchy);
-	/** Uses names to map mask multiplier indices from DNA file, for all LODs, to the indices of curves in the
-	  * control rig's curve container **/
-	void MapMaskMultipliers(const IBehaviorReader* DNABehavior, const URigHierarchy* InHierarchy);
-
-};
-
-/* The work data used by the FRigUnit_RigLogic */
-USTRUCT(meta = (DocumentationPolicy = "Strict"))
-struct FRigUnit_RigLogic_Data
-{
-	GENERATED_BODY()
-
-		FRigUnit_RigLogic_Data();
-	~FRigUnit_RigLogic_Data();
-	FRigUnit_RigLogic_Data(const FRigUnit_RigLogic_Data& Other);
-	FRigUnit_RigLogic_Data& operator=(const FRigUnit_RigLogic_Data& Other);
-
-	/** Cached Skeletal Mesh Component **/
-	UPROPERTY(transient)
-		TWeakObjectPtr<USkeletalMeshComponent> SkelMeshComponent;
-
-	FSharedRigRuntimeContext* SharedRigRuntimeContext;
-
-	/** RigInstance is a thin class, containing only character instance specific data
-	  * i.e. the output buffers where RigLogic writes the results of computations, specific
-	  * to each instance of a rig.
-	**/
-	TUniquePtr<FRigInstance> RigInstance;
-
-	TSharedPtr<const FRigDefinitionMapping> Mappings;
+	TArray<FRigUnit_RigLogic_IntArray> RigLogicIndicesForAnimMaps;
 
 	/** LOD for which the model is rendered **/
 	UPROPERTY(transient)
-		uint32 CurrentLOD;
+	uint32 CurrentLOD;
 
 	static const uint8 MAX_ATTRS_PER_JOINT;
 
 	bool IsRigLogicInitialized();
-	void InitializeRigLogic(const URigHierarchy* InHierarchy);
+	void InitializeRigLogic(const URigHierarchy* InHierarchy, TSharedPtr<FSharedRigRuntimeContext> NewContext);
 
-	/** Cache the joint indices that change per each LOD **/
-	void CacheVariableJointIndices();
+	/** Makes a map of input curve indices from DNA file to the control rig curves **/
+	void MapInputCurveIndices(const URigHierarchy* InHierarchy);
+	/** Uses names to map joint indices from DNA file to the indices of bones in control rig hierarchy **/
+	void MapJoints(const URigHierarchy* Hierarchy);
+	/** Uses names of blend shapes and meshes in DNA file, for all LODs, to map their indices to the indices of
+	  * morph target curves in the curve container; curve name format is <mesh>__<blendshape> **/
+	void MapMorphTargets(const URigHierarchy* InHierarchy);
+	/** Uses names to map mask multiplier indices from DNA file, for all LODs, to the indices of curves in the
+	  * control rig's curve container **/
+	void MapMaskMultipliers(const URigHierarchy* InHierarchy);
 
 	/** Calculates joint positions, orientation and scale based on inputs curves of the control rig **/
 	void CalculateRigLogic(const URigHierarchy* InHierarchy);
 	/** Updates joint positions in the hierarchy based on inputs curves of the control rig **/
-	void UpdateJoints(URigHierarchy* Hierarchy, const FRigUnit_RigLogic_JointUpdateParams& JointUpdateParams);
+	void UpdateJoints(URigHierarchy* Hierarchy, TArrayView<const float> NeutralJointValues, TArrayView<const float> DeltaJointValues);
 	/** Updates morph target curve values based on values of input curves of the control rig **/
 	void UpdateBlendShapeCurves(URigHierarchy* InHierarchy, TArrayView<const float> BlendShapeValues);
 	/** Updates anim map curve values based on values of input curves of the control rig **/
 	void UpdateAnimMapCurves(URigHierarchy* InHierarchy, TArrayView<const float> AnimMapOutputs);
-private:
-	static TDiscardableKeyValueCache<uint32, TSharedPtr<const FRigDefinitionMapping>> RigDefintionMappingsCache;
 };
 
 
@@ -158,7 +116,7 @@ private:
   *  animated map multiplier curves */
 
 USTRUCT(meta = (DisplayName = "RigLogic", Category = "RigLogic", DocumentationPolicy = "Strict", Keywords = "Rig,RigLogic"))
-struct FRigUnit_RigLogic : public FRigUnitMutable
+struct FRigUnit_RigLogic: public FRigUnitMutable
 {
 	GENERATED_BODY()
 
@@ -170,13 +128,13 @@ public:
 #endif
 
 	RIGVM_METHOD()
-		void Execute(const FRigUnitContext& Context) override;
+	void Execute(const FRigUnitContext& Context) override;
 
 private:
-	static FSharedRigRuntimeContext* GetSharedRigRuntimeContext(USkeletalMesh* SkelMesh);
+	static TSharedPtr<FSharedRigRuntimeContext> GetSharedRigRuntimeContext(USkeletalMesh* SkelMesh);
 
 private:
 	// internal work data for the unit
 	UPROPERTY(transient)
-		FRigUnit_RigLogic_Data Data;
+	FRigUnit_RigLogic_Data Data;
 };
diff --git a/Engine/Plugins/Animation/RigLogic/Source/RigLogicModule/Public/RigUnit_RigLogic_Test.h b/Engine/Plugins/Animation/RigLogic/Source/RigLogicModule/Public/RigUnit_RigLogic_Test.h
--- a/Engine/Plugins/Animation/RigLogic/Source/RigLogicModule/Public/RigUnit_RigLogic_Test.h
+++ b/Engine/Plugins/Animation/RigLogic/Source/RigLogicModule/Public/RigUnit_RigLogic_Test.h
@@ -20,6 +20,11 @@ struct FRigUnit_RigLogic_JointUpdateParams;
 class TestBehaviorReader: public IBehaviorReader
 {
 public:
+	void Unload(EDNADataLayer Layer) override
+	{
+
+	}
+
 	// DescriptorReader methods start
 	FString GetName() const override
 	{
@@ -528,9 +533,9 @@ struct FRigUnit_RigLogic::TestAccessor
 	FTransformArrayView CreateTwoJointNeutralTransforms(float *InValueArray);
 	TArrayView<const uint16> CreateTwoJointVariableAttributes(uint16* InVariableAttributeIndices, uint8 LOD);
 
-	void Exec_UpdateJoints(URigHierarchy* TestHierarchy, FRigUnit_RigLogic_JointUpdateParams& JointUpdateParams);
+	void Exec_UpdateJoints(URigHierarchy* TestHierarchy, TArrayView<const float> NeutralJointValues, TArrayView<const float> DeltaJointValues);
 
-	FSharedRigRuntimeContext* GetSharedRigRuntimeContext(USkeletalMesh* SkelMesh);
+	TSharedPtr<FSharedRigRuntimeContext> GetSharedRigRuntimeContext(USkeletalMesh* SkelMesh);
 
 private:
 	FName RawCtrlName;
