RigLogic: Non-unity build fix.

#rnx

diff --git a/Engine/Plugins/Animation/RigLogic/Source/RigLogicModule/Private/DNAIndexMapping.cpp b/Engine/Plugins/Animation/RigLogic/Source/RigLogicModule/Private/DNAIndexMapping.cpp
--- a/Engine/Plugins/Animation/RigLogic/Source/RigLogicModule/Private/DNAIndexMapping.cpp
+++ b/Engine/Plugins/Animation/RigLogic/Source/RigLogicModule/Private/DNAIndexMapping.cpp
@@ -2,10 +2,13 @@
 
 #include "DNAIndexMapping.h"
 
-#include "CoreMinimal.h"
 #include "HAL/LowLevelMemTracker.h"
 #include "DNAReader.h"
 #include "Hasher.h"
+#include "Animation/Skeleton.h"
+#include "Animation/SmartName.h"
+#include "Components/SkeletalMeshComponent.h"
+
 
 static uint32 HashDNA(const IDefinitionReader* Reader)
 {
diff --git a/Engine/Plugins/Animation/RigLogic/Source/RigLogicModule/Private/DNAIndexMapping.h b/Engine/Plugins/Animation/RigLogic/Source/RigLogicModule/Private/DNAIndexMapping.h
--- a/Engine/Plugins/Animation/RigLogic/Source/RigLogicModule/Private/DNAIndexMapping.h
+++ b/Engine/Plugins/Animation/RigLogic/Source/RigLogicModule/Private/DNAIndexMapping.h
@@ -4,14 +4,17 @@
 
 #include "DNAReader.h"
 
-#include "CoreMinimal.h"
-#include "UObject/Object.h"
-#include "UObject/ObjectMacros.h"
 #include "Engine/AssetUserData.h"
-#include "Interfaces/Interface_AssetUserData.h"
+#include "BoneIndices.h"
 
 #include "DNAIndexMapping.generated.h"
 
+
+class IBehaviorReader;
+class USkeleton;
+class USkeletalMesh;
+
+
 struct FDNAIndexMapping
 {
 	template <typename T>
