RigLogic: Remove #pragma once from a cpp file.

#rnx

diff --git a/Engine/Plugins/Animation/RigLogic/Source/RigLogicEditor/Private/RigLogicEditor.cpp b/Engine/Plugins/Animation/RigLogic/Source/RigLogicEditor/Private/RigLogicEditor.cpp
--- a/Engine/Plugins/Animation/RigLogic/Source/RigLogicEditor/Private/RigLogicEditor.cpp
+++ b/Engine/Plugins/Animation/RigLogic/Source/RigLogicEditor/Private/RigLogicEditor.cpp
@@ -1,11 +1,6 @@
 // Copyright Epic Games, Inc. All Rights Reserved.
 
-#pragma once
-
 #include "RigLogicEditor.h"
-#include "Features/IModularFeatures.h"
-#include "Engine/Engine.h"
-#include "Misc/Paths.h"
 #include "Modules/ModuleManager.h"
 #include "Logging/LogMacros.h"
 
