Redid github PR 8944

diff --git a/Engine/Plugins/Animation/RigLogic/Source/RigLogicLib/RigLogicLib.Build.cs b/Engine/Plugins/Animation/RigLogic/Source/RigLogicLib/RigLogicLib.Build.cs
--- a/Engine/Plugins/Animation/RigLogic/Source/RigLogicLib/RigLogicLib.Build.cs
+++ b/Engine/Plugins/Animation/RigLogic/Source/RigLogicLib/RigLogicLib.Build.cs
@@ -13,8 +13,6 @@ public class RigLogicLib : ModuleRules
 
     public RigLogicLib(ReadOnlyTargetRules Target) : base(Target)
     {
-        bUseUnity = false;
-
         PublicDependencyModuleNames.AddRange(
             new string[]
             {
diff --git a/Engine/Plugins/Animation/RigLogic/Source/RigLogicLibTest/RigLogicLibTest.Build.cs b/Engine/Plugins/Animation/RigLogic/Source/RigLogicLibTest/RigLogicLibTest.Build.cs
--- a/Engine/Plugins/Animation/RigLogic/Source/RigLogicLibTest/RigLogicLibTest.Build.cs
+++ b/Engine/Plugins/Animation/RigLogic/Source/RigLogicLibTest/RigLogicLibTest.Build.cs
@@ -9,6 +9,8 @@ public class RigLogicLibTest : ModuleRules
 {
     public RigLogicLibTest(ReadOnlyTargetRules Target) : base(Target)
     {
+		bUseUnity = false; // A windows include is preprocessing some method names causing compile failures.
+
         if (Target.Platform == UnrealTargetPlatform.Win64)
         {
             PrivateDefinitions.Add("RL_BUILD_WITH_SSE=1");
