Updated trimd to latest version

diff --git a/Engine/Plugins/Animation/RigLogic/Source/RigLogicLib/Private/riglogic/joints/JointsBuilder.cpp b/Engine/Plugins/Animation/RigLogic/Source/RigLogicLib/Private/riglogic/joints/JointsBuilder.cpp
--- a/Engine/Plugins/Animation/RigLogic/Source/RigLogicLib/Private/riglogic/joints/JointsBuilder.cpp
+++ b/Engine/Plugins/Animation/RigLogic/Source/RigLogicLib/Private/riglogic/joints/JointsBuilder.cpp
@@ -1,24 +1,22 @@
 // Copyright Epic Games, Inc. All Rights Reserved.
 
 #include "riglogic/joints/JointsBuilder.h"
-#include "riglogic/joints/bpcm/builders/Scalar.h"
 
-#if PLATFORM_CPU_X86_FAMILY
-    #ifdef RL_BUILD_WITH_AVX
-        #ifdef RL_USE_HALF_FLOATS
-            #include "riglogic/joints/bpcm/builders/AVXHalfFloat.h"
-        #else
-            #include "riglogic/joints/bpcm/builders/AVX.h"
-        #endif  // RL_USE_HALF_FLOATS
-    #endif  // RL_BUILD_WITH_AVX
-    #ifdef RL_BUILD_WITH_SSE
-        #ifdef RL_USE_HALF_FLOATS
-            #include "riglogic/joints/bpcm/builders/SSEHalfFloat.h"
-        #else
-            #include "riglogic/joints/bpcm/builders/SSE.h"
-        #endif  // RL_USE_HALF_FLOATS
-    #endif  // RL_BUILD_WITH_SSE
-#endif // PLATFORM_CPU_X86_FAMILY
+#include "riglogic/joints/bpcm/builders/Scalar.h"
+#ifdef RL_BUILD_WITH_AVX
+    #ifdef RL_USE_HALF_FLOATS
+        #include "riglogic/joints/bpcm/builders/AVXHalfFloat.h"
+    #else
+        #include "riglogic/joints/bpcm/builders/AVX.h"
+    #endif  // RL_USE_HALF_FLOATS
+#endif  // RL_BUILD_WITH_AVX
+#ifdef RL_BUILD_WITH_SSE
+    #ifdef RL_USE_HALF_FLOATS
+        #include "riglogic/joints/bpcm/builders/SSEHalfFloat.h"
+    #else
+        #include "riglogic/joints/bpcm/builders/SSE.h"
+    #endif  // RL_USE_HALF_FLOATS
+#endif  // RL_BUILD_WITH_SSE
 
 #include <pma/utils/ManagedInstance.h>
 
@@ -27,18 +25,18 @@ namespace rl4 {
 JointsBuilder::~JointsBuilder() = default;
 
 JointsBuilder::JointsBuilderPtr JointsBuilder::create(Configuration config, MemoryResource* memRes) {
-    #if PLATFORM_CPU_X86_FAMILY
-        #ifdef RL_BUILD_WITH_SSE
-            if (config.calculationType == CalculationType::SSE) {
-                return pma::UniqueInstance<bpcm::SSEJointsBuilder, JointsBuilder>::with(memRes).create(memRes);
-            }
-        #endif  // RL_BUILD_WITH_SSE
-        #ifdef RL_BUILD_WITH_AVX
-            if (config.calculationType == CalculationType::AVX) {
-                return pma::UniqueInstance<bpcm::AVXJointsBuilder, JointsBuilder>::with(memRes).create(memRes);
-            }
-        #endif  // RL_BUILD_WITH_AVX
-    #endif // PLATFORM_CPU_X86_FAMILY
+    // Work around unused parameter warning when building without SSE and AVX
+    static_cast<void>(config);
+    #ifdef RL_BUILD_WITH_SSE
+        if (config.calculationType == CalculationType::SSE) {
+            return pma::UniqueInstance<bpcm::SSEJointsBuilder, JointsBuilder>::with(memRes).create(memRes);
+        }
+    #endif  // RL_BUILD_WITH_SSE
+    #ifdef RL_BUILD_WITH_AVX
+        if (config.calculationType == CalculationType::AVX) {
+            return pma::UniqueInstance<bpcm::AVXJointsBuilder, JointsBuilder>::with(memRes).create(memRes);
+        }
+    #endif  // RL_BUILD_WITH_AVX
     return pma::UniqueInstance<bpcm::ScalarJointsBuilder, JointsBuilder>::with(memRes).create(memRes);
 }
 
diff --git a/Engine/Plugins/Animation/RigLogic/Source/RigLogicLib/Private/riglogic/joints/bpcm/builders/SSE.cpp b/Engine/Plugins/Animation/RigLogic/Source/RigLogicLib/Private/riglogic/joints/bpcm/builders/SSE.cpp
--- a/Engine/Plugins/Animation/RigLogic/Source/RigLogicLib/Private/riglogic/joints/bpcm/builders/SSE.cpp
+++ b/Engine/Plugins/Animation/RigLogic/Source/RigLogicLib/Private/riglogic/joints/bpcm/builders/SSE.cpp
@@ -1,8 +1,8 @@
 // Copyright Epic Games, Inc. All Rights Reserved.
 
-#include "riglogic/joints/bpcm/builders/SSE.h"
+#if !defined(RL_USE_HALF_FLOATS) && defined(RL_BUILD_WITH_SSE)
 
-#if !defined(RL_USE_HALF_FLOATS) && defined(RL_BUILD_WITH_SSE) && PLATFORM_CPU_X86_FAMILY
+#include "riglogic/joints/bpcm/builders/SSE.h"
 
 #include "riglogic/joints/bpcm/Consts.h"
 #include "riglogic/joints/bpcm/strategies/SSE.h"
@@ -26,4 +26,4 @@ SSEJointsBuilder::~SSEJointsBuilder() = default;
 
 }  // namespace rl4
 
-#endif  // !defined(RL_USE_HALF_FLOATS) && defined(RL_BUILD_WITH_SSE) && PLATFORM_CPU_X86_FAMILY
+#endif  // !defined(RL_USE_HALF_FLOATS) && defined(RL_BUILD_WITH_SSE)
diff --git a/Engine/Plugins/Animation/RigLogic/Source/RigLogicLib/Private/riglogic/joints/bpcm/builders/SSE.h b/Engine/Plugins/Animation/RigLogic/Source/RigLogicLib/Private/riglogic/joints/bpcm/builders/SSE.h
--- a/Engine/Plugins/Animation/RigLogic/Source/RigLogicLib/Private/riglogic/joints/bpcm/builders/SSE.h
+++ b/Engine/Plugins/Animation/RigLogic/Source/RigLogicLib/Private/riglogic/joints/bpcm/builders/SSE.h
@@ -2,10 +2,8 @@
 
 #pragma once
 
-#include "HAL/Platform.h"
-
 // *INDENT-OFF*
-#if !defined(RL_USE_HALF_FLOATS) && defined(RL_BUILD_WITH_SSE) && PLATFORM_CPU_X86_FAMILY
+#if !defined(RL_USE_HALF_FLOATS) && defined(RL_BUILD_WITH_SSE)
 
 #include "riglogic/joints/bpcm/builders/Float.h"
 #include "riglogic/types/Aliases.h"
@@ -24,5 +22,5 @@ class SSEJointsBuilder : public FloatStorageBuilder {
 
 }  // namespace rl4
 
-#endif  // !defined(RL_USE_HALF_FLOATS) && defined(RL_BUILD_WITH_SSE) && PLATFORM_CPU_X86_FAMILY
+#endif  // !defined(RL_USE_HALF_FLOATS) && defined(RL_BUILD_WITH_SSE)
 // *INDENT-ON*
diff --git a/Engine/Plugins/Animation/RigLogic/Source/RigLogicLib/Private/riglogic/joints/bpcm/builders/SSEHalfFloat.cpp b/Engine/Plugins/Animation/RigLogic/Source/RigLogicLib/Private/riglogic/joints/bpcm/builders/SSEHalfFloat.cpp
--- a/Engine/Plugins/Animation/RigLogic/Source/RigLogicLib/Private/riglogic/joints/bpcm/builders/SSEHalfFloat.cpp
+++ b/Engine/Plugins/Animation/RigLogic/Source/RigLogicLib/Private/riglogic/joints/bpcm/builders/SSEHalfFloat.cpp
@@ -1,8 +1,8 @@
 // Copyright Epic Games, Inc. All Rights Reserved.
 
-#include "riglogic/joints/bpcm/builders/SSEHalfFloat.h"
+#if defined(RL_USE_HALF_FLOATS) && defined(RL_BUILD_WITH_SSE)
 
-#if defined(RL_USE_HALF_FLOATS) && defined(RL_BUILD_WITH_SSE) && PLATFORM_CPU_X86_FAMILY
+#include "riglogic/joints/bpcm/builders/SSEHalfFloat.h"
 
 #include "riglogic/joints/bpcm/Consts.h"
 #include "riglogic/joints/bpcm/strategies/SSE.h"
@@ -97,4 +97,4 @@ void SSEJointsBuilder::setValues(const dna::BehaviorReader* reader) {
 
 }  // namespace rl4
 
-#endif  // defined(RL_USE_HALF_FLOATS) && defined(RL_BUILD_WITH_SSE) && PLATFORM_CPU_X86_FAMILY
+#endif  // defined(RL_USE_HALF_FLOATS) && defined(RL_BUILD_WITH_SSE)
diff --git a/Engine/Plugins/Animation/RigLogic/Source/RigLogicLib/Private/riglogic/joints/bpcm/builders/SSEHalfFloat.h b/Engine/Plugins/Animation/RigLogic/Source/RigLogicLib/Private/riglogic/joints/bpcm/builders/SSEHalfFloat.h
--- a/Engine/Plugins/Animation/RigLogic/Source/RigLogicLib/Private/riglogic/joints/bpcm/builders/SSEHalfFloat.h
+++ b/Engine/Plugins/Animation/RigLogic/Source/RigLogicLib/Private/riglogic/joints/bpcm/builders/SSEHalfFloat.h
@@ -2,10 +2,8 @@
 
 #pragma once
 
-#include "HAL/Platform.h"
-
 // *INDENT-OFF*
-#if defined(RL_USE_HALF_FLOATS) && defined(RL_BUILD_WITH_SSE) && PLATFORM_CPU_X86_FAMILY
+#if defined(RL_USE_HALF_FLOATS) && defined(RL_BUILD_WITH_SSE)
 
 #include "riglogic/joints/bpcm/JointsBuilderCommon.h"
 #include "riglogic/types/Aliases.h"
@@ -30,5 +28,5 @@ class SSEJointsBuilder : public JointsBuilderCommon<std::uint16_t> {
 
 }  // namespace rl4
 
-#endif  // defined(RL_USE_HALF_FLOATS) && defined(RL_BUILD_WITH_SSE) && PLATFORM_CPU_X86_FAMILY
+#endif  // defined(RL_USE_HALF_FLOATS) && defined(RL_BUILD_WITH_SSE)
 // *INDENT-ON*
diff --git a/Engine/Plugins/Animation/RigLogic/Source/RigLogicLib/Private/riglogic/joints/bpcm/strategies/AVX.h b/Engine/Plugins/Animation/RigLogic/Source/RigLogicLib/Private/riglogic/joints/bpcm/strategies/AVX.h
--- a/Engine/Plugins/Animation/RigLogic/Source/RigLogicLib/Private/riglogic/joints/bpcm/strategies/AVX.h
+++ b/Engine/Plugins/Animation/RigLogic/Source/RigLogicLib/Private/riglogic/joints/bpcm/strategies/AVX.h
@@ -3,7 +3,7 @@
 #pragma once
 
 #include "riglogic/joints/bpcm/strategies/Block8.h"
-#include "trimd/TRiMD.h"
+#include "riglogic/system/simd/SIMD.h"
 
 namespace rl4 {
 
diff --git a/Engine/Plugins/Animation/RigLogic/Source/RigLogicLib/Private/riglogic/joints/bpcm/strategies/SSE.h b/Engine/Plugins/Animation/RigLogic/Source/RigLogicLib/Private/riglogic/joints/bpcm/strategies/SSE.h
--- a/Engine/Plugins/Animation/RigLogic/Source/RigLogicLib/Private/riglogic/joints/bpcm/strategies/SSE.h
+++ b/Engine/Plugins/Animation/RigLogic/Source/RigLogicLib/Private/riglogic/joints/bpcm/strategies/SSE.h
@@ -3,7 +3,7 @@
 #pragma once
 
 #include "riglogic/joints/bpcm/strategies/Block4.h"
-#include "trimd/TRiMD.h"
+#include "riglogic/system/simd/SIMD.h"
 
 namespace rl4 {
 
diff --git a/Engine/Plugins/Animation/RigLogic/Source/RigLogicLib/Private/riglogic/joints/bpcm/strategies/Scalar.h b/Engine/Plugins/Animation/RigLogic/Source/RigLogicLib/Private/riglogic/joints/bpcm/strategies/Scalar.h
--- a/Engine/Plugins/Animation/RigLogic/Source/RigLogicLib/Private/riglogic/joints/bpcm/strategies/Scalar.h
+++ b/Engine/Plugins/Animation/RigLogic/Source/RigLogicLib/Private/riglogic/joints/bpcm/strategies/Scalar.h
@@ -3,7 +3,7 @@
 #pragma once
 
 #include "riglogic/joints/bpcm/strategies/Block4.h"
-#include "trimd/TRiMD.h"
+#include "riglogic/system/simd/SIMD.h"
 
 namespace rl4 {
 
diff --git a/Engine/Plugins/Animation/RigLogic/Source/RigLogicLib/Private/riglogic/system/simd/SIMD.h b/Engine/Plugins/Animation/RigLogic/Source/RigLogicLib/Private/riglogic/system/simd/SIMD.h
--- a/Engine/Plugins/Animation/RigLogic/Source/RigLogicLib/Private/riglogic/system/simd/SIMD.h
+++ b/Engine/Plugins/Animation/RigLogic/Source/RigLogicLib/Private/riglogic/system/simd/SIMD.h
@@ -2,18 +2,16 @@
 
 #pragma once
 
-#include "HAL/Platform.h"
-
 #if defined(RL_USE_HALF_FLOATS) && !defined(TRIMD_ENABLE_F16C)
     #define TRIMD_ENABLE_F16C
 #endif  // RL_USE_HALF_FLOATS
 
-#if PLATFORM_CPU_X86_FAMILY
-    #if defined(RL_BUILD_WITH_AVX) && !defined(TRIMD_ENABLE_AVX)
-        #define TRIMD_ENABLE_AVX
-    #endif  // RL_BUILD_WITH_AVX
+#if defined(RL_BUILD_WITH_AVX) && !defined(TRIMD_ENABLE_AVX)
+    #define TRIMD_ENABLE_AVX
+#endif  // RL_BUILD_WITH_AVX
+
+#if defined(RL_BUILD_WITH_SSE) && !defined(TRIMD_ENABLE_SSE)
+    #define TRIMD_ENABLE_SSE
+#endif  // RL_BUILD_WITH_SSE
 
-    #if defined(RL_BUILD_WITH_SSE) && !defined(TRIMD_ENABLE_SSE)
-        #define TRIMD_ENABLE_SSE
-    #endif  // RL_BUILD_WITH_SSE
-#endif // PLATFORM_CPU_X86_FAMILY
+#include <trimd/TRiMD.h>
diff --git a/Engine/Plugins/Animation/RigLogic/Source/RigLogicLib/Public/trimd/AVX.h b/Engine/Plugins/Animation/RigLogic/Source/RigLogicLib/Public/trimd/AVX.h
--- a/Engine/Plugins/Animation/RigLogic/Source/RigLogicLib/Public/trimd/AVX.h
+++ b/Engine/Plugins/Animation/RigLogic/Source/RigLogicLib/Public/trimd/AVX.h
@@ -13,6 +13,8 @@ namespace trimd {
 namespace avx {
 
 struct F256 {
+    using value_type = float;
+
     __m256 data;
 
     F256() : data{_mm256_setzero_ps()} {
@@ -43,18 +45,15 @@ struct F256 {
 
     #ifdef TRIMD_ENABLE_F16C
         static F256 fromAlignedSource(const std::uint16_t* source) {
-            __m128i halfFloats = _mm_load_si128(reinterpret_cast<const __m128i*>(source));
-            return F256{_mm256_cvtph_ps(halfFloats)};
+            return F256{_mm256_cvtph_ps(_mm_load_si128(reinterpret_cast<const __m128i*>(source)))};
         }
 
         static F256 fromUnalignedSource(const std::uint16_t* source) {
-            __m128i halfFloats = _mm_loadu_si128(reinterpret_cast<const __m128i*>(source));
-            return F256{_mm256_cvtph_ps(halfFloats)};
+            return F256{_mm256_cvtph_ps(_mm_loadu_si128(reinterpret_cast<const __m128i*>(source)))};
         }
 
         static F256 loadSingleValue(const std::uint16_t* source) {
-            __m128i halfFloats = _mm_loadu_si16(source);
-            return F256{_mm256_cvtph_ps(halfFloats)};
+            return F256{_mm256_cvtph_ps(_mm_loadu_si16(source))};
         }
 
     #endif  // TRIMD_ENABLE_F16C
@@ -123,6 +122,44 @@ struct F256 {
         _mm256_storeu_ps(dest, data);
     }
 
+    #ifdef TRIMD_ENABLE_F16C
+    void alignedLoad(const std::uint16_t* source) {
+        data = _mm256_cvtph_ps(_mm_load_si128(reinterpret_cast<const __m128i*>(source)));
+    }
+
+    void unalignedLoad(const std::uint16_t* source) {
+        data = _mm256_cvtph_ps(_mm_loadu_si128(reinterpret_cast<const __m128i*>(source)));
+    }
+
+    void alignedStore(std::uint16_t* dest) const {
+        #if defined(__clang__) || defined(__GNUC__)
+            #pragma GCC diagnostic push
+            #pragma GCC diagnostic ignored "-Wold-style-cast"
+            #if !defined(__clang__)
+                #pragma GCC diagnostic ignored "-Wuseless-cast"
+            #endif
+        #endif
+        _mm_store_si128(reinterpret_cast<__m128i*>(dest), _mm256_cvtps_ph(data, _MM_FROUND_CUR_DIRECTION));
+        #if defined(__clang__) || defined(__GNUC__)
+            #pragma GCC diagnostic pop
+        #endif
+    }
+
+    void unalignedStore(std::uint16_t* dest) const {
+        #if defined(__clang__) || defined(__GNUC__)
+            #pragma GCC diagnostic push
+            #pragma GCC diagnostic ignored "-Wold-style-cast"
+            #if !defined(__clang__)
+                #pragma GCC diagnostic ignored "-Wuseless-cast"
+            #endif
+        #endif
+        _mm_storeu_si128(reinterpret_cast<__m128i*>(dest), _mm256_cvtps_ph(data, _MM_FROUND_CUR_DIRECTION));
+        #if defined(__clang__) || defined(__GNUC__)
+            #pragma GCC diagnostic pop
+        #endif
+    }
+    #endif  // TRIMD_ENABLE_F16C
+
     float sum() const {
         // (data[3] + data[7], data[2] + data[6], data[1] + data[5], data[0] + data[4])
         const __m128 x128 = _mm_add_ps(_mm256_extractf128_ps(data, 1), _mm256_castps256_ps128(data));
@@ -153,6 +190,21 @@ struct F256 {
         return *this;
     }
 
+    F256& operator&=(const F256& rhs) {
+        data = _mm256_and_ps(data, rhs.data);
+        return *this;
+    }
+
+    F256& operator|=(const F256& rhs) {
+        data = _mm256_or_ps(data, rhs.data);
+        return *this;
+    }
+
+    F256& operator^=(const F256& rhs) {
+        data = _mm256_xor_ps(data, rhs.data);
+        return *this;
+    }
+
     static constexpr std::size_t size() {
         return sizeof(decltype(data)) / sizeof(float);
     }
@@ -163,12 +215,28 @@ struct F256 {
 
 };
 
-inline bool operator==(const F256& lhs, const F256& rhs) {
-    return _mm256_movemask_ps(_mm256_cmp_ps(lhs.data, rhs.data, _CMP_EQ_OQ)) == 0xFF;
+inline F256 operator==(const F256& lhs, const F256& rhs) {
+    return F256{_mm256_cmp_ps(lhs.data, rhs.data, _CMP_EQ_OQ)};
 }
 
-inline bool operator!=(const F256& lhs, const F256& rhs) {
-    return !(lhs == rhs);
+inline F256 operator!=(const F256& lhs, const F256& rhs) {
+    return F256{_mm256_cmp_ps(lhs.data, rhs.data, _CMP_NEQ_OQ)};
+}
+
+inline F256 operator<(const F256& lhs, const F256& rhs) {
+    return F256{_mm256_cmp_ps(lhs.data, rhs.data, _CMP_LT_OQ)};
+}
+
+inline F256 operator<=(const F256& lhs, const F256& rhs) {
+    return F256{_mm256_cmp_ps(lhs.data, rhs.data, _CMP_LE_OQ)};
+}
+
+inline F256 operator>(const F256& lhs, const F256& rhs) {
+    return F256{_mm256_cmp_ps(lhs.data, rhs.data, _CMP_GT_OQ)};
+}
+
+inline F256 operator>=(const F256& lhs, const F256& rhs) {
+    return F256{_mm256_cmp_ps(lhs.data, rhs.data, _CMP_GE_OQ)};
 }
 
 inline F256 operator+(const F256& lhs, const F256& rhs) {
@@ -187,6 +255,24 @@ inline F256 operator/(const F256& lhs, const F256& rhs) {
     return F256(lhs) /= rhs;
 }
 
+inline F256 operator&(const F256& lhs, const F256& rhs) {
+    return F256(lhs) &= rhs;
+}
+
+inline F256 operator|(const F256& lhs, const F256& rhs) {
+    return F256(lhs) |= rhs;
+}
+
+inline F256 operator^(const F256& lhs, const F256& rhs) {
+    return F256(lhs) ^= rhs;
+}
+
+inline F256 operator~(const F256& rhs) {
+    F256 result{};
+    result = (result == result);  // Generate register with all 1s
+    return (result ^= rhs);
+}
+
 inline void transpose(F256& row0, F256& row1, F256& row2, F256& row3, F256& row4, F256& row5, F256& row6, F256& row7) {
     __m256 t0 = _mm256_unpacklo_ps(row0.data, row1.data);
     __m256 t1 = _mm256_unpackhi_ps(row0.data, row1.data);
@@ -214,6 +300,15 @@ inline void transpose(F256& row0, F256& row1, F256& row2, F256& row3, F256& row4
     row7.data = _mm256_permute2f128_ps(tt3, tt7, 0x31);
 }
 
+inline F256 abs(const F256& rhs) {
+    const __m256 ABSMASK = _mm256_castsi256_ps(_mm256_set1_epi32(~(1 << 31)));
+    return F256{_mm256_and_ps(ABSMASK, rhs.data)};
+}
+
+inline F256 andnot(const F256& lhs, const F256& rhs) {
+    return F256{_mm256_andnot_ps(lhs.data, rhs.data)};
+}
+
 } // namespace avx
 
 } // namespace trimd
diff --git a/Engine/Plugins/Animation/RigLogic/Source/RigLogicLib/Public/trimd/Fallback.h b/Engine/Plugins/Animation/RigLogic/Source/RigLogicLib/Public/trimd/Fallback.h
new file mode 100644
--- /dev/null
+++ b/Engine/Plugins/Animation/RigLogic/Source/RigLogicLib/Public/trimd/Fallback.h
@@ -0,0 +1,243 @@
+// Copyright Epic Games, Inc. All Rights Reserved.
+
+#pragma once
+
+namespace trimd {
+
+namespace fallback {
+
+template<typename T128>
+struct T256 {
+    using value_type = typename T128::value_type;
+
+    T128 data1;
+    T128 data2;
+
+    T256(const T128& d1, const T128& d2) : data1{d1}, data2{d2} {
+    }
+
+    T256() : data1{}, data2{} {
+    }
+
+    T256(value_type v1, value_type v2, value_type v3, value_type v4, value_type v5, value_type v6, value_type v7,
+         value_type v8) : data1{v1, v2, v3, v4}, data2{v5, v6, v7, v8} {
+    }
+
+    explicit T256(value_type value) : T256{value, value, value, value, value, value, value, value} {
+    }
+
+    template<typename U>
+    static T256 fromAlignedSource(const U* source) {
+        return T256{T128::fromAlignedSource(source), T128::fromAlignedSource(source + T128::size())};
+    }
+
+    template<typename U>
+    static T256 fromUnalignedSource(const U* source) {
+        return T256{T128::fromUnalignedSource(source), T128::fromUnalignedSource(source + T128::size())};
+    }
+
+    template<typename U>
+    static T256 loadSingleValue(const U* source) {
+        return T256{T128::loadSingleValue(source), T128{}};
+    }
+
+    template<typename U>
+    static void prefetchT0(const U* source) {
+        T128::prefetchT0(source);
+    }
+
+    template<typename U>
+    static void prefetchT1(const U* source) {
+        T128::prefetchT1(source);
+    }
+
+    template<typename U>
+    static void prefetchT2(const U* source) {
+        T128::prefetchT2(source);
+    }
+
+    template<typename U>
+    static void prefetchNTA(const U* source) {
+        T128::prefetchNTA(source);
+    }
+
+    template<typename U>
+    void alignedLoad(const U* source) {
+        data1.alignedLoad(source);
+        data2.alignedLoad(source + T128::size());
+    }
+
+    template<typename U>
+    void unalignedLoad(const U* source) {
+        data1.unalignedLoad(source);
+        data2.unalignedLoad(source + T128::size());
+    }
+
+    template<typename U>
+    void alignedStore(U* dest) const {
+        data1.alignedStore(dest);
+        data2.alignedStore(dest + T128::size());
+    }
+
+    template<typename U>
+    void unalignedStore(U* dest) const {
+        data1.unalignedStore(dest);
+        data2.unalignedStore(dest + T128::size());
+    }
+
+    value_type sum() const {
+        return data1.sum() + data2.sum();
+    }
+
+    T256& operator+=(const T256& rhs) {
+        data1 += rhs.data1;
+        data2 += rhs.data2;
+        return *this;
+    }
+
+    T256& operator-=(const T256& rhs) {
+        data1 -= rhs.data1;
+        data2 -= rhs.data2;
+        return *this;
+    }
+
+    T256& operator*=(const T256& rhs) {
+        data1 *= rhs.data1;
+        data2 *= rhs.data2;
+        return *this;
+    }
+
+    T256& operator/=(const T256& rhs) {
+        data1 /= rhs.data1;
+        data2 /= rhs.data2;
+        return *this;
+    }
+
+    T256& operator&=(const T256& rhs) {
+        data1 &= rhs.data1;
+        data2 &= rhs.data2;
+        return *this;
+    }
+
+    T256& operator|=(const T256& rhs) {
+        data1 |= rhs.data1;
+        data2 |= rhs.data2;
+        return *this;
+    }
+
+    T256& operator^=(const T256& rhs) {
+        data1 ^= rhs.data1;
+        data2 ^= rhs.data2;
+        return *this;
+    }
+
+    static constexpr std::size_t size() {
+        return T128::size() * 2ul;
+    }
+
+    static constexpr std::size_t alignment() {
+        // T128 alignment is the minimal requirement, but it might be beneficial to force here an alignment
+        // of a theoretical T256, so the autovectorizer might generate better code on platforms not directly
+        // supported by TRiMD.
+        return T128::alignment();
+    }
+
+};
+
+template<typename T128>
+inline T256<T128> operator==(const T256<T128>& lhs, const T256<T128>& rhs) {
+    return T256<T128>{lhs.data1 == rhs.data1, lhs.data2 == rhs.data2};
+}
+
+template<typename T128>
+inline T256<T128> operator!=(const T256<T128>& lhs, const T256<T128>& rhs) {
+    return T256<T128>{lhs.data1 != rhs.data1, lhs.data2 != rhs.data2};
+}
+
+template<typename T128>
+inline T256<T128> operator<(const T256<T128>& lhs, const T256<T128>& rhs) {
+    return T256<T128>{lhs.data1 < rhs.data1, lhs.data2 < rhs.data2};
+}
+
+template<typename T128>
+inline T256<T128> operator<=(const T256<T128>& lhs, const T256<T128>& rhs) {
+    return T256<T128>{lhs.data1 <= rhs.data1, lhs.data2 <= rhs.data2};
+}
+
+template<typename T128>
+inline T256<T128> operator>(const T256<T128>& lhs, const T256<T128>& rhs) {
+    return T256<T128>{lhs.data1 > rhs.data1, lhs.data2 > rhs.data2};
+}
+
+template<typename T128>
+inline T256<T128> operator>=(const T256<T128>& lhs, const T256<T128>& rhs) {
+    return T256<T128>{lhs.data1 >= rhs.data1, lhs.data2 >= rhs.data2};
+}
+
+template<typename T128>
+inline T256<T128> operator+(const T256<T128>& lhs, const T256<T128>& rhs) {
+    return T256<T128>(lhs) += rhs;
+}
+
+template<typename T128>
+inline T256<T128> operator-(const T256<T128>& lhs, const T256<T128>& rhs) {
+    return T256<T128>(lhs) -= rhs;
+}
+
+template<typename T128>
+inline T256<T128> operator*(const T256<T128>& lhs, const T256<T128>& rhs) {
+    return T256<T128>(lhs) *= rhs;
+}
+
+template<typename T128>
+inline T256<T128> operator/(const T256<T128>& lhs, const T256<T128>& rhs) {
+    return T256<T128>(lhs) /= rhs;
+}
+
+template<typename T128>
+inline T256<T128> operator&(const T256<T128>& lhs, const T256<T128>& rhs) {
+    return T256<T128>(lhs) &= rhs;
+}
+
+template<typename T128>
+inline T256<T128> operator|(const T256<T128>& lhs, const T256<T128>& rhs) {
+    return T256<T128>(lhs) |= rhs;
+}
+
+template<typename T128>
+inline T256<T128> operator^(const T256<T128>& lhs, const T256<T128>& rhs) {
+    return T256<T128>(lhs) ^= rhs;
+}
+
+template<typename T128>
+inline T256<T128> operator~(const T256<T128>& rhs) {
+    return T256<T128>{~rhs.data1, ~rhs.data2};
+}
+
+// *INDENT-OFF*
+template<typename T128>
+inline void transpose(T256<T128>& row0, T256<T128>& row1, T256<T128>& row2, T256<T128>& row3, T256<T128>& row4, T256<T128>& row5, T256<T128>& row6, T256<T128>& row7) {
+    transpose(row0.data1, row1.data1, row2.data1, row3.data1);
+    transpose(row0.data2, row1.data2, row2.data2, row3.data2);
+    transpose(row4.data1, row5.data1, row6.data1, row7.data1);
+    transpose(row4.data2, row5.data2, row6.data2, row7.data2);
+    std::swap(row0.data2, row4.data1);
+    std::swap(row1.data2, row5.data1);
+    std::swap(row2.data2, row6.data1);
+    std::swap(row3.data2, row7.data1);
+}
+// *INDENT-ON*
+
+template<typename T128>
+inline T256<T128> abs(const T256<T128>& rhs) {
+    return T256<T128>{abs(rhs.data1), abs(rhs.data2)};
+}
+
+template<typename T128>
+inline T256<T128> andnot(const T256<T128>& lhs, const T256<T128>& rhs) {
+    return T256<T128>{andnot(lhs.data1, rhs.data1), andnot(lhs.data2, rhs.data2)};
+}
+
+}  // namespace fallback
+
+}  // namespace trimd
diff --git a/Engine/Plugins/Animation/RigLogic/Source/RigLogicLib/Public/trimd/Macros.h b/Engine/Plugins/Animation/RigLogic/Source/RigLogicLib/Public/trimd/Macros.h
new file mode 100644
--- /dev/null
+++ b/Engine/Plugins/Animation/RigLogic/Source/RigLogicLib/Public/trimd/Macros.h
@@ -0,0 +1,11 @@
+// Copyright Epic Games, Inc. All Rights Reserved.
+
+#pragma once
+
+#if !defined(FORCE_INLINE)
+    #if defined(_MSC_VER)
+        #define FORCE_INLINE __forceinline
+    #else
+        #define FORCE_INLINE inline __attribute__((always_inline))
+    #endif
+#endif
diff --git a/Engine/Plugins/Animation/RigLogic/Source/RigLogicLib/Public/trimd/Polyfill.h b/Engine/Plugins/Animation/RigLogic/Source/RigLogicLib/Public/trimd/Polyfill.h
--- a/Engine/Plugins/Animation/RigLogic/Source/RigLogicLib/Public/trimd/Polyfill.h
+++ b/Engine/Plugins/Animation/RigLogic/Source/RigLogicLib/Public/trimd/Polyfill.h
@@ -7,6 +7,7 @@
 #include <immintrin.h>
 
 #include <cstdint>
+#include <cstring>
 
 #if (!defined(__clang__) && defined(_MSC_VER) && _MSC_VER <= 1900) || (!defined(__clang__) && defined(__GNUC__) && \
 __GNUC__ < 9)
@@ -15,11 +16,26 @@ __GNUC__ < 9)
 
 #if (!defined(__clang__) && defined(_MSC_VER) && _MSC_VER <= 1900) || \
 (defined(__clang__) && __clang_major__ < 8) || \
-(!defined(__clang__) && defined(__GNUC__))
+(!defined(__clang__) && defined(__GNUC__) && __GNUC__ < 11)
     inline __m128i _mm_loadu_si16(const void* source) {
         return _mm_insert_epi16(_mm_setzero_si128(), *reinterpret_cast<const std::int16_t*>(source), 0);
     }
+#endif
+
+#if !defined(__clang__) && defined(_MSC_VER) && (_MSC_VER <= 1900) && defined(_WIN32) && !defined(_WIN64)
+inline std::int64_t _mm_cvtsi128_si64(__m128i source) {
+    const std::uint32_t low32 = static_cast<std::uint32_t>(_mm_cvtsi128_si32(source));
+    const std::uint32_t high32 = static_cast<std::uint32_t>(_mm_cvtsi128_si32(_mm_srli_epi64(source, 32)));
+    return static_cast<std::int64_t>((static_cast<std::uint64_t>(high32) << 32) | low32);
+}
+#endif
 
+#if (!defined(__clang__) && defined(_MSC_VER) && _MSC_VER <= 1900) || (!defined(__clang__) && defined(__GNUC__) && \
+__GNUC__ < 9) || (defined(__clang__) && __clang_major__ < 8)
+    inline void _mm_storeu_si64(void* dest, __m128i source) {
+        const int64_t value = _mm_cvtsi128_si64(source);
+        std::memcpy(dest, &value, sizeof(value));
+    }
 #endif
 
 #endif  // defined(TRIMD_ENABLE_AVX) || defined(TRIMD_ENABLE_SSE)
diff --git a/Engine/Plugins/Animation/RigLogic/Source/RigLogicLib/Public/trimd/SSE.h b/Engine/Plugins/Animation/RigLogic/Source/RigLogicLib/Public/trimd/SSE.h
--- a/Engine/Plugins/Animation/RigLogic/Source/RigLogicLib/Public/trimd/SSE.h
+++ b/Engine/Plugins/Animation/RigLogic/Source/RigLogicLib/Public/trimd/SSE.h
@@ -4,6 +4,7 @@
 
 // *INDENT-OFF*
 #ifdef TRIMD_ENABLE_SSE
+#include "trimd/Fallback.h"
 #include "trimd/Polyfill.h"
 
 #include <immintrin.h>
@@ -13,6 +14,8 @@ namespace trimd {
 namespace sse {
 
 struct F128 {
+    using value_type = float;
+
     __m128 data;
 
     F128() : data{_mm_setzero_ps()} {
@@ -41,18 +44,15 @@ struct F128 {
 
     #ifdef TRIMD_ENABLE_F16C
         static F128 fromAlignedSource(const std::uint16_t* source) {
-            __m128i halfFloats = _mm_loadl_epi64(reinterpret_cast<const __m128i*>(source));
-            return F128{_mm_cvtph_ps(halfFloats)};
+            return F128{_mm_cvtph_ps(_mm_loadl_epi64(reinterpret_cast<const __m128i*>(source)))};
         }
 
         static F128 fromUnalignedSource(const std::uint16_t* source) {
-            __m128i halfFloats = _mm_loadu_si64(reinterpret_cast<const __m128i*>(source));
-            return F128{_mm_cvtph_ps(halfFloats)};
+            return F128{_mm_cvtph_ps(_mm_loadu_si64(reinterpret_cast<const __m128i*>(source)))};
         }
 
         static F128 loadSingleValue(const std::uint16_t* source) {
-            __m128i halfFloats = _mm_loadu_si16(source);
-            return F128{_mm_cvtph_ps(halfFloats)};
+            return F128{_mm_cvtph_ps(_mm_loadu_si16(source))};
         }
 
     #endif  // TRIMD_ENABLE_F16C
@@ -121,6 +121,44 @@ struct F128 {
         _mm_storeu_ps(dest, data);
     }
 
+    #ifdef TRIMD_ENABLE_F16C
+    void alignedLoad(const std::uint16_t* source) {
+        data = _mm_cvtph_ps(_mm_loadl_epi64(reinterpret_cast<const __m128i*>(source)));
+    }
+
+    void unalignedLoad(const std::uint16_t* source) {
+        data = _mm_cvtph_ps(_mm_loadu_si64(reinterpret_cast<const __m128i*>(source)));
+    }
+
+    void alignedStore(std::uint16_t* dest) const {
+        #if defined(__clang__) || defined(__GNUC__)
+            #pragma GCC diagnostic push
+            #pragma GCC diagnostic ignored "-Wold-style-cast"
+            #if !defined(__clang__)
+                #pragma GCC diagnostic ignored "-Wuseless-cast"
+            #endif
+        #endif
+        _mm_storel_epi64(reinterpret_cast<__m128i*>(dest), _mm_cvtps_ph(data, _MM_FROUND_CUR_DIRECTION));
+        #if defined(__clang__) || defined(__GNUC__)
+            #pragma GCC diagnostic pop
+        #endif
+    }
+
+    void unalignedStore(std::uint16_t* dest) const {
+        #if defined(__clang__) || defined(__GNUC__)
+            #pragma GCC diagnostic push
+            #pragma GCC diagnostic ignored "-Wold-style-cast"
+            #if !defined(__clang__)
+                #pragma GCC diagnostic ignored "-Wuseless-cast"
+            #endif
+        #endif
+        _mm_storeu_si64(reinterpret_cast<__m128i*>(dest), _mm_cvtps_ph(data, _MM_FROUND_CUR_DIRECTION));
+        #if defined(__clang__) || defined(__GNUC__)
+            #pragma GCC diagnostic pop
+        #endif
+    }
+    #endif  // TRIMD_ENABLE_F16C
+
     float sum() const {
         __m128 temp = _mm_movehdup_ps(data);
         __m128 result = _mm_add_ps(data, temp);
@@ -149,6 +187,21 @@ struct F128 {
         return *this;
     }
 
+    F128& operator&=(const F128& rhs) {
+        data = _mm_and_ps(data, rhs.data);
+        return *this;
+    }
+
+    F128& operator|=(const F128& rhs) {
+        data = _mm_or_ps(data, rhs.data);
+        return *this;
+    }
+
+    F128& operator^=(const F128& rhs) {
+        data = _mm_xor_ps(data, rhs.data);
+        return *this;
+    }
+
     static constexpr std::size_t size() {
         return sizeof(decltype(data)) / sizeof(float);
     }
@@ -159,12 +212,28 @@ struct F128 {
 
 };
 
-inline bool operator==(const F128& lhs, const F128& rhs) {
-    return _mm_movemask_ps(_mm_cmpeq_ps(lhs.data, rhs.data)) == 0xF;
+inline F128 operator==(const F128& lhs, const F128& rhs) {
+    return F128{_mm_cmpeq_ps(lhs.data, rhs.data)};
+}
+
+inline F128 operator!=(const F128& lhs, const F128& rhs) {
+    return F128{_mm_cmpneq_ps(lhs.data, rhs.data)};
+}
+
+inline F128 operator<(const F128& lhs, const F128& rhs) {
+    return F128{_mm_cmplt_ps(lhs.data, rhs.data)};
+}
+
+inline F128 operator<=(const F128& lhs, const F128& rhs) {
+    return F128{_mm_cmple_ps(lhs.data, rhs.data)};
+}
+
+inline F128 operator>(const F128& lhs, const F128& rhs) {
+    return F128{_mm_cmpgt_ps(lhs.data, rhs.data)};
 }
 
-inline bool operator!=(const F128& lhs, const F128& rhs) {
-    return !(lhs == rhs);
+inline F128 operator>=(const F128& lhs, const F128& rhs) {
+    return F128{_mm_cmpge_ps(lhs.data, rhs.data)};
 }
 
 inline F128 operator+(const F128& lhs, const F128& rhs) {
@@ -183,10 +252,42 @@ inline F128 operator/(const F128& lhs, const F128& rhs) {
     return F128(lhs) /= rhs;
 }
 
+inline F128 operator&(const F128& lhs, const F128& rhs) {
+    return F128(lhs) &= rhs;
+}
+
+inline F128 operator|(const F128& lhs, const F128& rhs) {
+    return F128(lhs) |= rhs;
+}
+
+inline F128 operator^(const F128& lhs, const F128& rhs) {
+    return F128(lhs) ^= rhs;
+}
+
+inline F128 operator~(const F128& rhs) {
+    static const auto ones = static_cast<int>(0xFFFFFFFF);
+    F128 result{_mm_castsi128_ps(_mm_set1_epi32(ones))};  // Generate register with all 1s
+    return (result ^= rhs);
+}
+
 inline void transpose(F128& row0, F128& row1, F128& row2, F128& row3) {
     _MM_TRANSPOSE4_PS(row0.data, row1.data, row2.data, row3.data);
 }
 
+inline F128 abs(const F128& rhs) {
+    const __m128 ABSMASK = _mm_castsi128_ps(_mm_set1_epi32(~(1 << 31)));
+    return F128{_mm_and_ps(ABSMASK, rhs.data)};
+}
+
+inline F128 andnot(const F128& lhs, const F128& rhs) {
+    return F128{_mm_andnot_ps(lhs.data, rhs.data)};
+}
+
+using F256 = fallback::T256<F128>;
+using fallback::transpose;
+using fallback::abs;
+using fallback::andnot;
+
 } // namespace sse
 
 } // namespace trimd
diff --git a/Engine/Plugins/Animation/RigLogic/Source/RigLogicLib/Public/trimd/Scalar.h b/Engine/Plugins/Animation/RigLogic/Source/RigLogicLib/Public/trimd/Scalar.h
--- a/Engine/Plugins/Animation/RigLogic/Source/RigLogicLib/Public/trimd/Scalar.h
+++ b/Engine/Plugins/Animation/RigLogic/Source/RigLogicLib/Public/trimd/Scalar.h
@@ -2,8 +2,19 @@
 
 #pragma once
 
+#include "trimd/Fallback.h"
+#include "trimd/Utils.h"
+
+#ifdef _MSC_VER
+    #pragma warning(push)
+    #pragma warning(disable : 4365 4987)
+#endif
 #include <array>
+#include <cmath>
 #include <cstddef>
+#ifdef _MSC_VER
+    #pragma warning(pop)
+#endif
 
 namespace trimd {
 
@@ -13,6 +24,8 @@ template<typename T>
 struct T128 {
     using value_type = typename std::remove_cv<T>::type;
 
+    static_assert(sizeof(value_type) == 4, "Only 32-bit types are supported");
+
     std::array<value_type, 4> data;
 
     T128() : data{} {
@@ -78,7 +91,7 @@ struct T128 {
         alignedStore(dest);
     }
 
-    float sum() const {
+    value_type sum() const {
         return data[0] + data[1] + data[2] + data[3];
     }
 
@@ -114,8 +127,32 @@ struct T128 {
         return *this;
     }
 
+    T128& operator&=(const T128& rhs) {
+        data[0] = bitcast<value_type>(bitcast<std::uint32_t>(data[0]) & bitcast<std::uint32_t>(rhs.data[0]));
+        data[1] = bitcast<value_type>(bitcast<std::uint32_t>(data[1]) & bitcast<std::uint32_t>(rhs.data[1]));
+        data[2] = bitcast<value_type>(bitcast<std::uint32_t>(data[2]) & bitcast<std::uint32_t>(rhs.data[2]));
+        data[3] = bitcast<value_type>(bitcast<std::uint32_t>(data[3]) & bitcast<std::uint32_t>(rhs.data[3]));
+        return *this;
+    }
+
+    T128& operator|=(const T128& rhs) {
+        data[0] = bitcast<value_type>(bitcast<std::uint32_t>(data[0]) | bitcast<std::uint32_t>(rhs.data[0]));
+        data[1] = bitcast<value_type>(bitcast<std::uint32_t>(data[1]) | bitcast<std::uint32_t>(rhs.data[1]));
+        data[2] = bitcast<value_type>(bitcast<std::uint32_t>(data[2]) | bitcast<std::uint32_t>(rhs.data[2]));
+        data[3] = bitcast<value_type>(bitcast<std::uint32_t>(data[3]) | bitcast<std::uint32_t>(rhs.data[3]));
+        return *this;
+    }
+
+    T128& operator^=(const T128& rhs) {
+        data[0] = bitcast<value_type>(bitcast<std::uint32_t>(data[0]) ^ bitcast<std::uint32_t>(rhs.data[0]));
+        data[1] = bitcast<value_type>(bitcast<std::uint32_t>(data[1]) ^ bitcast<std::uint32_t>(rhs.data[1]));
+        data[2] = bitcast<value_type>(bitcast<std::uint32_t>(data[2]) ^ bitcast<std::uint32_t>(rhs.data[2]));
+        data[3] = bitcast<value_type>(bitcast<std::uint32_t>(data[3]) ^ bitcast<std::uint32_t>(rhs.data[3]));
+        return *this;
+    }
+
     static constexpr std::size_t size() {
-        return sizeof(decltype(data)) / sizeof(float);
+        return sizeof(decltype(data)) / sizeof(value_type);
     }
 
     static constexpr std::size_t alignment() {
@@ -129,213 +166,109 @@ struct T128 {
 };
 
 template<typename T>
-inline bool operator==(const T128<T>& lhs, const T128<T>& rhs) {
-    return (lhs.data == rhs.data);
+inline T128<T> operator==(const T128<T>& lhs, const T128<T>& rhs) {
+    return T128<T>{
+        bitcast<T>(static_cast<std::uint32_t>(-(lhs.data[0] == rhs.data[0]))),
+        bitcast<T>(static_cast<std::uint32_t>(-(lhs.data[1] == rhs.data[1]))),
+        bitcast<T>(static_cast<std::uint32_t>(-(lhs.data[2] == rhs.data[2]))),
+        bitcast<T>(static_cast<std::uint32_t>(-(lhs.data[3] == rhs.data[3])))
+    };
 }
 
 template<typename T>
-inline bool operator!=(const T128<T>& lhs, const T128<T>& rhs) {
-    return !(lhs == rhs);
+inline T128<T> operator!=(const T128<T>& lhs, const T128<T>& rhs) {
+    return T128<T>{
+        bitcast<T>(static_cast<std::uint32_t>(-(lhs.data[0] != rhs.data[0]))),
+        bitcast<T>(static_cast<std::uint32_t>(-(lhs.data[1] != rhs.data[1]))),
+        bitcast<T>(static_cast<std::uint32_t>(-(lhs.data[2] != rhs.data[2]))),
+        bitcast<T>(static_cast<std::uint32_t>(-(lhs.data[3] != rhs.data[3])))
+    };
 }
 
 template<typename T>
-inline T128<T> operator+(const T128<T>& lhs, const T128<T>& rhs) {
-    return T128<T>(lhs) += rhs;
+inline T128<T> operator<(const T128<T>& lhs, const T128<T>& rhs) {
+    return T128<T>{
+        bitcast<T>(static_cast<std::uint32_t>(-(lhs.data[0] < rhs.data[0]))),
+        bitcast<T>(static_cast<std::uint32_t>(-(lhs.data[1] < rhs.data[1]))),
+        bitcast<T>(static_cast<std::uint32_t>(-(lhs.data[2] < rhs.data[2]))),
+        bitcast<T>(static_cast<std::uint32_t>(-(lhs.data[3] < rhs.data[3])))
+    };
 }
 
 template<typename T>
-inline T128<T> operator-(const T128<T>& lhs, const T128<T>& rhs) {
-    return T128<T>(lhs) -= rhs;
+inline T128<T> operator<=(const T128<T>& lhs, const T128<T>& rhs) {
+    return T128<T>{
+        bitcast<T>(static_cast<std::uint32_t>(-(lhs.data[0] <= rhs.data[0]))),
+        bitcast<T>(static_cast<std::uint32_t>(-(lhs.data[1] <= rhs.data[1]))),
+        bitcast<T>(static_cast<std::uint32_t>(-(lhs.data[2] <= rhs.data[2]))),
+        bitcast<T>(static_cast<std::uint32_t>(-(lhs.data[3] <= rhs.data[3])))
+    };
 }
 
 template<typename T>
-inline T128<T> operator*(const T128<T>& lhs, const T128<T>& rhs) {
-    return T128<T>(lhs) *= rhs;
+inline T128<T> operator>(const T128<T>& lhs, const T128<T>& rhs) {
+    return T128<T>{
+        bitcast<T>(static_cast<std::uint32_t>(-(lhs.data[0] > rhs.data[0]))),
+        bitcast<T>(static_cast<std::uint32_t>(-(lhs.data[1] > rhs.data[1]))),
+        bitcast<T>(static_cast<std::uint32_t>(-(lhs.data[2] > rhs.data[2]))),
+        bitcast<T>(static_cast<std::uint32_t>(-(lhs.data[3] > rhs.data[3])))
+    };
 }
 
 template<typename T>
-inline T128<T> operator/(const T128<T>& lhs, const T128<T>& rhs) {
-    return T128<T>(lhs) /= rhs;
+inline T128<T> operator>=(const T128<T>& lhs, const T128<T>& rhs) {
+    return T128<T>{
+        bitcast<T>(static_cast<std::uint32_t>(-(lhs.data[0] >= rhs.data[0]))),
+        bitcast<T>(static_cast<std::uint32_t>(-(lhs.data[1] >= rhs.data[1]))),
+        bitcast<T>(static_cast<std::uint32_t>(-(lhs.data[2] >= rhs.data[2]))),
+        bitcast<T>(static_cast<std::uint32_t>(-(lhs.data[3] >= rhs.data[3])))
+    };
 }
 
 template<typename T>
-struct T256 {
-    using value_type = typename std::remove_cv<T>::type;
-
-    std::array<value_type, 8> data;
-
-    T256() : data{} {
-    }
-
-    T256(value_type v1, value_type v2, value_type v3, value_type v4, value_type v5, value_type v6, value_type v7,
-         value_type v8) : data({v1, v2, v3, v4, v5, v6, v7,
-                                v8})
-    {
-    }
-
-    explicit T256(value_type value) : T256(value, value, value, value, value, value, value, value) {
-    }
-
-    static T256 fromAlignedSource(const value_type* source) {
-        return T256{source[0], source[1], source[2], source[3], source[4], source[5], source[6], source[7]};
-    }
-
-    static T256 fromUnalignedSource(const value_type* source) {
-        return T256::fromAlignedSource(source);
-    }
-
-    static T256 loadSingleValue(const value_type* source) {
-        return T256{source[0], value_type{}, value_type{}, value_type{}, value_type{}, value_type{}, value_type{}, value_type{}};
-    }
-
-    template<typename U>
-    static void prefetchT0(const U*  /*unused*/) {
-        // Intentionally noop
-    }
-
-    template<typename U>
-    static void prefetchT1(const U*  /*unused*/) {
-        // Intentionally noop
-    }
-
-    template<typename U>
-    static void prefetchT2(const U*  /*unused*/) {
-        // Intentionally noop
-    }
-
-    template<typename U>
-    static void prefetchNTA(const U*  /*unused*/) {
-        // Intentionally noop
-    }
-
-    void alignedLoad(const value_type* source) {
-        data[0] = source[0];
-        data[1] = source[1];
-        data[2] = source[2];
-        data[3] = source[3];
-        data[4] = source[4];
-        data[5] = source[5];
-        data[6] = source[6];
-        data[7] = source[7];
-    }
-
-    void unalignedLoad(const value_type* source) {
-        alignedLoad(source);
-    }
-
-    void alignedStore(value_type* dest) const {
-        dest[0] = data[0];
-        dest[1] = data[1];
-        dest[2] = data[2];
-        dest[3] = data[3];
-        dest[4] = data[4];
-        dest[5] = data[5];
-        dest[6] = data[6];
-        dest[7] = data[7];
-    }
-
-    void unalignedStore(value_type* dest) const {
-        alignedStore(dest);
-    }
-
-    float sum() const {
-        return data[0] + data[1] + data[2] + data[3] + data[4] + data[5] + data[6] + data[7];
-    }
-
-    T256& operator+=(const T256& rhs) {
-        data[0] += rhs.data[0];
-        data[1] += rhs.data[1];
-        data[2] += rhs.data[2];
-        data[3] += rhs.data[3];
-        data[4] += rhs.data[4];
-        data[5] += rhs.data[5];
-        data[6] += rhs.data[6];
-        data[7] += rhs.data[7];
-        return *this;
-    }
-
-    T256& operator-=(const T256& rhs) {
-        data[0] -= rhs.data[0];
-        data[1] -= rhs.data[1];
-        data[2] -= rhs.data[2];
-        data[3] -= rhs.data[3];
-        data[4] -= rhs.data[4];
-        data[5] -= rhs.data[5];
-        data[6] -= rhs.data[6];
-        data[7] -= rhs.data[7];
-        return *this;
-    }
-
-    T256& operator*=(const T256& rhs) {
-        data[0] *= rhs.data[0];
-        data[1] *= rhs.data[1];
-        data[2] *= rhs.data[2];
-        data[3] *= rhs.data[3];
-        data[4] *= rhs.data[4];
-        data[5] *= rhs.data[5];
-        data[6] *= rhs.data[6];
-        data[7] *= rhs.data[7];
-        return *this;
-    }
-
-    T256& operator/=(const T256& rhs) {
-        data[0] /= rhs.data[0];
-        data[1] /= rhs.data[1];
-        data[2] /= rhs.data[2];
-        data[3] /= rhs.data[3];
-        data[4] /= rhs.data[4];
-        data[5] /= rhs.data[5];
-        data[6] /= rhs.data[6];
-        data[7] /= rhs.data[7];
-        return *this;
-    }
-
-    static constexpr std::size_t size() {
-        return sizeof(decltype(data)) / sizeof(float);
-    }
-
-    static constexpr std::size_t alignment() {
-        #if defined(__arm__) || defined(__aarch64__) || defined(_M_ARM) || defined(_M_ARM64)
-            return std::alignment_of<std::max_align_t>::value;
-        #else
-            return sizeof(decltype(data));
-        #endif
-    }
-
-};
+inline T128<T> operator+(const T128<T>& lhs, const T128<T>& rhs) {
+    return T128<T>(lhs) += rhs;
+}
 
 template<typename T>
-inline bool operator==(const T256<T>& lhs, const T256<T>& rhs) {
-    return (lhs.data == rhs.data);
+inline T128<T> operator-(const T128<T>& lhs, const T128<T>& rhs) {
+    return T128<T>(lhs) -= rhs;
 }
 
 template<typename T>
-inline bool operator!=(const T256<T>& lhs, const T256<T>& rhs) {
-    return !(lhs == rhs);
+inline T128<T> operator*(const T128<T>& lhs, const T128<T>& rhs) {
+    return T128<T>(lhs) *= rhs;
 }
 
 template<typename T>
-inline T256<T> operator+(const T256<T>& lhs, const T256<T>& rhs) {
-    return T256<T>(lhs) += rhs;
+inline T128<T> operator/(const T128<T>& lhs, const T128<T>& rhs) {
+    return T128<T>(lhs) /= rhs;
 }
 
 template<typename T>
-inline T256<T> operator-(const T256<T>& lhs, const T256<T>& rhs) {
-    return T256<T>(lhs) -= rhs;
+inline T128<T> operator&(const T128<T>& lhs, const T128<T>& rhs) {
+    return T128<T>(lhs) &= rhs;
 }
 
 template<typename T>
-inline T256<T> operator*(const T256<T>& lhs, const T256<T>& rhs) {
-    return T256<T>(lhs) *= rhs;
+inline T128<T> operator|(const T128<T>& lhs, const T128<T>& rhs) {
+    return T128<T>(lhs) |= rhs;
 }
 
 template<typename T>
-inline T256<T> operator/(const T256<T>& lhs, const T256<T>& rhs) {
-    return T256<T>(lhs) /= rhs;
+inline T128<T> operator^(const T128<T>& lhs, const T128<T>& rhs) {
+    return T128<T>(lhs) ^= rhs;
 }
 
-using F128 = T128<float>;
-using F256 = T256<float>;
+template<typename T>
+inline T128<T> operator~(const T128<T>& rhs) {
+    return T128<T>(
+        bitcast<T>(~bitcast<std::uint32_t>(rhs.data[0])),
+        bitcast<T>(~bitcast<std::uint32_t>(rhs.data[1])),
+        bitcast<T>(~bitcast<std::uint32_t>(rhs.data[2])),
+        bitcast<T>(~bitcast<std::uint32_t>(rhs.data[3]))
+        );
+}
 
 template<typename T>
 inline void transpose(T128<T>& row0, T128<T>& row1, T128<T>& row2, T128<T>& row3) {
@@ -350,25 +283,24 @@ inline void transpose(T128<T>& row0, T128<T>& row1, T128<T>& row2, T128<T>& row3
 }
 
 template<typename T>
-inline void transpose(T256<T>& row0, T256<T>& row1, T256<T>& row2, T256<T>& row3, T256<T>& row4, T256<T>& row5, T256<T>& row6, T256<T>& row7) {
-    T256<T> transposed0{row0.data[0], row1.data[0], row2.data[0], row3.data[0], row4.data[0], row5.data[0], row6.data[0], row7.data[0]};
-    T256<T> transposed1{row0.data[1], row1.data[1], row2.data[1], row3.data[1], row4.data[1], row5.data[1], row6.data[1], row7.data[1]};
-    T256<T> transposed2{row0.data[2], row1.data[2], row2.data[2], row3.data[2], row4.data[2], row5.data[2], row6.data[2], row7.data[2]};
-    T256<T> transposed3{row0.data[3], row1.data[3], row2.data[3], row3.data[3], row4.data[3], row5.data[3], row6.data[3], row7.data[3]};
-    T256<T> transposed4{row0.data[4], row1.data[4], row2.data[4], row3.data[4], row4.data[4], row5.data[4], row6.data[4], row7.data[4]};
-    T256<T> transposed5{row0.data[5], row1.data[5], row2.data[5], row3.data[5], row4.data[5], row5.data[5], row6.data[5], row7.data[5]};
-    T256<T> transposed6{row0.data[6], row1.data[6], row2.data[6], row3.data[6], row4.data[6], row5.data[6], row6.data[6], row7.data[6]};
-    T256<T> transposed7{row0.data[7], row1.data[7], row2.data[7], row3.data[7], row4.data[7], row5.data[7], row6.data[7], row7.data[7]};
-    row0 = transposed0;
-    row1 = transposed1;
-    row2 = transposed2;
-    row3 = transposed3;
-    row4 = transposed4;
-    row5 = transposed5;
-    row6 = transposed6;
-    row7 = transposed7;
+inline T128<T> abs(const T128<T>& rhs) {
+    return {std::abs(rhs.data[0]),
+            std::abs(rhs.data[1]),
+            std::abs(rhs.data[2]),
+            std::abs(rhs.data[3])};
+}
+
+template<typename T>
+inline T128<T> andnot(const T128<T>& lhs, const T128<T>& rhs) {
+    return ~lhs & rhs;
 }
 
+using F128 = T128<float>;
+using F256 = fallback::T256<F128>;
+using fallback::transpose;
+using fallback::abs;
+using fallback::andnot;
+
 }  // namespace scalar
 
 }  // namespace trimd
diff --git a/Engine/Plugins/Animation/RigLogic/Source/RigLogicLib/Public/trimd/TRiMD.h b/Engine/Plugins/Animation/RigLogic/Source/RigLogicLib/Public/trimd/TRiMD.h
--- a/Engine/Plugins/Animation/RigLogic/Source/RigLogicLib/Public/trimd/TRiMD.h
+++ b/Engine/Plugins/Animation/RigLogic/Source/RigLogicLib/Public/trimd/TRiMD.h
@@ -2,39 +2,34 @@
 
 #pragma once
 
-#include "riglogic/system/simd/SIMD.h"
 #include "trimd/AVX.h"
 #include "trimd/SSE.h"
 #include "trimd/Scalar.h"
 
 namespace trimd {
 
-#ifdef TRIMD_ENABLE_AVX
+#if defined(TRIMD_ENABLE_AVX)
     using F256 = avx::F256;
-
-    inline void transpose(F256& row0, F256& row1, F256& row2, F256& row3, F256& row4, F256& row5, F256& row6, F256& row7) {
-        avx::transpose(row0, row1, row2, row3, row4, row5, row6, row7);
-    }
+    using avx::abs;
+    using avx::transpose;
+    using avx::andnot;
+#elif defined(TRIMD_ENABLE_SSE)
+    using F256 = sse::F256;
 #else
     using F256 = scalar::F256;
-
-    inline void transpose(F256& row0, F256& row1, F256& row2, F256& row3, F256& row4, F256& row5, F256& row6, F256& row7) {
-        scalar::transpose(row0, row1, row2, row3, row4, row5, row6, row7);
-    }
 #endif  // TRIMD_ENABLE_AVX
 
 #ifdef TRIMD_ENABLE_SSE
     using F128 = sse::F128;
-
-    inline void transpose(F128& row0, F128& row1, F128& row2, F128& row3) {
-        sse::transpose(row0, row1, row2, row3);
-    }
+    using sse::abs;
+    using sse::transpose;
+    using sse::andnot;
 #else
     using F128 = scalar::F128;
-
-    inline void transpose(F128& row0, F128& row1, F128& row2, F128& row3) {
-        scalar::transpose(row0, row1, row2, row3);
-    }
 #endif  // TRIMD_ENABLE_SSE
 
+using scalar::abs;
+using scalar::transpose;
+using scalar::andnot;
+
 }  // namespace trimd
diff --git a/Engine/Plugins/Animation/RigLogic/Source/RigLogicLib/Public/trimd/Utils.h b/Engine/Plugins/Animation/RigLogic/Source/RigLogicLib/Public/trimd/Utils.h
new file mode 100644
--- /dev/null
+++ b/Engine/Plugins/Animation/RigLogic/Source/RigLogicLib/Public/trimd/Utils.h
@@ -0,0 +1,26 @@
+// Copyright Epic Games, Inc. All Rights Reserved.
+
+#pragma once
+
+#include "trimd/Macros.h"
+
+#ifdef _MSC_VER
+    #pragma warning(push)
+    #pragma warning(disable : 4365 4987)
+#endif
+#include <cstring>
+#ifdef _MSC_VER
+    #pragma warning(pop)
+#endif
+
+namespace trimd {
+
+template<typename TTarget, typename TSource>
+FORCE_INLINE TTarget bitcast(TSource source) {
+    static_assert(sizeof(TTarget) == sizeof(TSource), "Target and source must be of equal size.");
+    TTarget target;
+    std::memcpy(&target, &source, sizeof(TSource));
+    return target;
+}
+
+}  // namespace trimd
diff --git a/Engine/Plugins/Animation/RigLogic/Source/RigLogicLib/Public/trimd/version/Version.h b/Engine/Plugins/Animation/RigLogic/Source/RigLogicLib/Public/trimd/version/Version.h
--- a/Engine/Plugins/Animation/RigLogic/Source/RigLogicLib/Public/trimd/version/Version.h
+++ b/Engine/Plugins/Animation/RigLogic/Source/RigLogicLib/Public/trimd/version/Version.h
@@ -2,7 +2,7 @@
 
 #pragma once
 
-#define TRIMD_MAJOR_VERSION 1
-#define TRIMD_MINOR_VERSION 3
-#define TRIMD_PATCH_VERSION 2
-#define TRIMD_VERSION_STRING "1.3.2"
+#define TRIMD_MAJOR_VERSION 2
+#define TRIMD_MINOR_VERSION 1
+#define TRIMD_PATCH_VERSION 4
+#define TRIMD_VERSION_STRING "2.1.4"
diff --git a/Engine/Plugins/Animation/RigLogic/Source/RigLogicLibTest/Private/trimdtests/TestT128.cpp b/Engine/Plugins/Animation/RigLogic/Source/RigLogicLibTest/Private/trimdtests/TestT128.cpp
--- a/Engine/Plugins/Animation/RigLogic/Source/RigLogicLibTest/Private/trimdtests/TestT128.cpp
+++ b/Engine/Plugins/Animation/RigLogic/Source/RigLogicLibTest/Private/trimdtests/TestT128.cpp
@@ -6,219 +6,457 @@
 
 #ifdef TRIMD_ENABLE_SSE
 using T128Types = ::testing::Types<trimd::scalar::F128, trimd::sse::F128>;
+
+bool equal(const trimd::sse::F128& lhs, const trimd::sse::F128& rhs) {
+	return (std::memcmp(&lhs.data, &rhs.data, sizeof(lhs.data)) == 0);
+}
+
 #else
 using T128Types = ::testing::Types<trimd::scalar::F128>;
 #endif  // TRIMD_ENABLE_SSE
 
+template<typename TF128>
+static TF128 frombits(uint32_t bits0, uint32_t bits1, uint32_t bits2, uint32_t bits3) {
+	return TF128{
+		trimd::bitcast<float>(bits0),
+		trimd::bitcast<float>(bits1),
+		trimd::bitcast<float>(bits2),
+		trimd::bitcast<float>(bits3)
+	};
+}
+
+static bool equal(const trimd::scalar::F128& lhs, const trimd::scalar::F128& rhs) {
+	return (std::memcmp(lhs.data.data(), rhs.data.data(), sizeof(float) * trimd::scalar::F128::size()) == 0);
+}
+
 template<typename T>
 class T128Test : public ::testing::Test {
 protected:
-    using T128 = T;
+	using T128 = T;
 
 };
 
 TYPED_TEST_SUITE(T128Test, T128Types, );
 
 TYPED_TEST(T128Test, CheckSize) {
-    ASSERT_EQ(TestFixture::T128::size(), 4ul);
+	ASSERT_EQ(TestFixture::T128::size(), 4ul);
 }
 
 TYPED_TEST(T128Test, Equality) {
-    typename TestFixture::T128 v1{1.0f, 2.0f, 3.0f, 4.0f};
-    typename TestFixture::T128 v2{1.0f, 2.0f, 3.0f, 4.0f};
+	using F128 = typename TestFixture::T128;
+	F128 v1{ 1.0f, 2.0f, 3.0f, 4.0f };
+	F128 v2{ 1.0f, 2.0f, 3.0f, 4.0f };
+	F128 v3{ 1.5f, 2.0f, 3.0f, 4.0f };
+	F128 v4{ 1.0f, 2.5f, 3.0f, 4.0f };
+	F128 v5{ 1.0f, 2.0f, 3.5f, 4.0f };
+	F128 v6{ 1.0f, 2.0f, 3.0f, 4.5f };
+
+	F128 m12 = frombits<F128>(0xFFFFFFFFu, 0xFFFFFFFFu, 0xFFFFFFFFu, 0xFFFFFFFFu);
+	F128 m13 = frombits<F128>(0x00000000u, 0xFFFFFFFFu, 0xFFFFFFFFu, 0xFFFFFFFFu);
+	F128 m14 = frombits<F128>(0xFFFFFFFFu, 0x00000000u, 0xFFFFFFFFu, 0xFFFFFFFFu);
+	F128 m15 = frombits<F128>(0xFFFFFFFFu, 0xFFFFFFFFu, 0x00000000u, 0xFFFFFFFFu);
+	F128 m16 = frombits<F128>(0xFFFFFFFFu, 0xFFFFFFFFu, 0xFFFFFFFFu, 0x00000000u);
+
+	ASSERT_TRUE(equal(v1 == v2, m12));
+	ASSERT_TRUE(equal(v1 == v3, m13));
+	ASSERT_TRUE(equal(v1 == v4, m14));
+	ASSERT_TRUE(equal(v1 == v5, m15));
+	ASSERT_TRUE(equal(v1 == v6, m16));
+}
 
-    typename TestFixture::T128 v3{1.5f, 2.0f, 3.0f, 4.0f};
-    typename TestFixture::T128 v4{1.0f, 2.5f, 3.0f, 4.0f};
-    typename TestFixture::T128 v5{1.0f, 2.0f, 3.5f, 4.0f};
-    typename TestFixture::T128 v6{1.0f, 2.0f, 3.0f, 4.5f};
+TYPED_TEST(T128Test, Inequality) {
+	using F128 = typename TestFixture::T128;
+	F128 v1{ 1.0f, 2.0f, 3.0f, 4.0f };
+	F128 v2{ 1.0f, 2.0f, 3.0f, 4.0f };
+	F128 v3{ 1.5f, 2.0f, 3.0f, 4.0f };
+	F128 v4{ 1.0f, 2.5f, 3.0f, 4.0f };
+	F128 v5{ 1.0f, 2.0f, 3.5f, 4.0f };
+	F128 v6{ 1.0f, 2.0f, 3.0f, 4.5f };
+
+	F128 m12 = frombits<F128>(0x00000000u, 0x00000000u, 0x00000000u, 0x00000000u);
+	F128 m13 = frombits<F128>(0xFFFFFFFFu, 0x00000000u, 0x00000000u, 0x00000000u);
+	F128 m14 = frombits<F128>(0x00000000u, 0xFFFFFFFFu, 0x00000000u, 0x00000000u);
+	F128 m15 = frombits<F128>(0x00000000u, 0x00000000u, 0xFFFFFFFFu, 0x00000000u);
+	F128 m16 = frombits<F128>(0x00000000u, 0x00000000u, 0x00000000u, 0xFFFFFFFFu);
+
+	ASSERT_TRUE(equal(v1 != v2, m12));
+	ASSERT_TRUE(equal(v1 != v3, m13));
+	ASSERT_TRUE(equal(v1 != v4, m14));
+	ASSERT_TRUE(equal(v1 != v5, m15));
+	ASSERT_TRUE(equal(v1 != v6, m16));
+}
+
+TYPED_TEST(T128Test, LessThan) {
+	using F128 = typename TestFixture::T128;
+	F128 v1{ 1.0f, 2.0f, 3.0f, 4.0f };
+	F128 v2{ 1.0f, 2.0f, 3.0f, 4.0f };
+	F128 v3{ 2.0f, 3.0f, 4.0f, 5.0f };
+	F128 v4{ 0.5f, 1.5f, 2.5f, 3.5f };
+
+	F128 m12 = frombits<F128>(0x00000000u, 0x00000000u, 0x00000000u, 0x00000000u);
+	F128 m13 = frombits<F128>(0xFFFFFFFFu, 0xFFFFFFFFu, 0xFFFFFFFFu, 0xFFFFFFFFu);
+	F128 m14 = frombits<F128>(0x00000000u, 0x00000000u, 0x00000000u, 0x00000000u);
 
-    ASSERT_TRUE(v1 == v2);
-    ASSERT_FALSE(v1 == v3);
-    ASSERT_FALSE(v1 == v4);
-    ASSERT_FALSE(v1 == v5);
-    ASSERT_FALSE(v1 == v6);
+	ASSERT_TRUE(equal(v1 < v2, m12));
+	ASSERT_TRUE(equal(v1 < v3, m13));
+	ASSERT_TRUE(equal(v1 < v4, m14));
 }
 
-TYPED_TEST(T128Test, Inequality) {
-    typename TestFixture::T128 v1{1.0f, 2.0f, 3.0f, 4.0f};
-    typename TestFixture::T128 v2{1.0f, 2.0f, 3.0f, 4.0f};
-    ASSERT_FALSE(v1 != v2);
+TYPED_TEST(T128Test, LessThanOrEqual) {
+	using F128 = typename TestFixture::T128;
+	F128 v1{ 1.0f, 2.0f, 3.0f, 4.0f };
+	F128 v2{ 1.0f, 2.0f, 3.0f, 4.0f };
+	F128 v3{ 2.0f, 3.0f, 4.0f, 5.0f };
+	F128 v4{ 0.5f, 1.5f, 2.5f, 3.5f };
+
+	F128 m12 = frombits<F128>(0xFFFFFFFFu, 0xFFFFFFFFu, 0xFFFFFFFFu, 0xFFFFFFFFu);
+	F128 m13 = frombits<F128>(0xFFFFFFFFu, 0xFFFFFFFFu, 0xFFFFFFFFu, 0xFFFFFFFFu);
+	F128 m14 = frombits<F128>(0x00000000u, 0x00000000u, 0x00000000u, 0x00000000u);
+
+	ASSERT_TRUE(equal(v1 <= v2, m12));
+	ASSERT_TRUE(equal(v1 <= v3, m13));
+	ASSERT_TRUE(equal(v1 <= v4, m14));
+}
+
+TYPED_TEST(T128Test, GreaterThan) {
+	using F128 = typename TestFixture::T128;
+	F128 v1{ 1.0f, 2.0f, 3.0f, 4.0f };
+	F128 v2{ 1.0f, 2.0f, 3.0f, 4.0f };
+	F128 v3{ 2.0f, 3.0f, 4.0f, 5.0f };
+	F128 v4{ 0.5f, 1.5f, 2.5f, 3.5f };
+
+	F128 m12 = frombits<F128>(0x00000000u, 0x00000000u, 0x00000000u, 0x00000000u);
+	F128 m13 = frombits<F128>(0x00000000u, 0x00000000u, 0x00000000u, 0x00000000u);
+	F128 m14 = frombits<F128>(0xFFFFFFFFu, 0xFFFFFFFFu, 0xFFFFFFFFu, 0xFFFFFFFFu);
+
+	ASSERT_TRUE(equal(v1 > v2, m12));
+	ASSERT_TRUE(equal(v1 > v3, m13));
+	ASSERT_TRUE(equal(v1 > v4, m14));
+}
+
+TYPED_TEST(T128Test, GreaterThanOrEqual) {
+	using F128 = typename TestFixture::T128;
+	F128 v1{ 1.0f, 2.0f, 3.0f, 4.0f };
+	F128 v2{ 1.0f, 2.0f, 3.0f, 4.0f };
+	F128 v3{ 2.0f, 3.0f, 4.0f, 5.0f };
+	F128 v4{ 0.5f, 1.5f, 2.5f, 3.5f };
+
+	F128 m12 = frombits<F128>(0xFFFFFFFFu, 0xFFFFFFFFu, 0xFFFFFFFFu, 0xFFFFFFFFu);
+	F128 m13 = frombits<F128>(0x00000000u, 0x00000000u, 0x00000000u, 0x00000000u);
+	F128 m14 = frombits<F128>(0xFFFFFFFFu, 0xFFFFFFFFu, 0xFFFFFFFFu, 0xFFFFFFFFu);
+
+	ASSERT_TRUE(equal(v1 >= v2, m12));
+	ASSERT_TRUE(equal(v1 >= v3, m13));
+	ASSERT_TRUE(equal(v1 >= v4, m14));
+}
+
+TYPED_TEST(T128Test, BitwiseAND) {
+	using F128 = typename TestFixture::T128;
+	F128 v{ 1.0f, 2.0f, 3.0f, 4.0f };
+
+	F128 m1 = frombits<F128>(0xFFFFFFFFu, 0x00000000u, 0x00000000u, 0x00000000u);
+	F128 m2 = frombits<F128>(0x00000000u, 0xFFFFFFFFu, 0x00000000u, 0x00000000u);
+	F128 m3 = frombits<F128>(0x00000000u, 0x00000000u, 0xFFFFFFFFu, 0x00000000u);
+	F128 m4 = frombits<F128>(0x00000000u, 0x00000000u, 0x00000000u, 0xFFFFFFFFu);
+
+	F128 vm1{ 1.0f, 0.0f, 0.0f, 0.0f };
+	F128 vm2{ 0.0f, 2.0f, 0.0f, 0.0f };
+	F128 vm3{ 0.0f, 0.0f, 3.0f, 0.0f };
+	F128 vm4{ 0.0f, 0.0f, 0.0f, 4.0f };
+
+	ASSERT_TRUE(equal(v & m1, vm1));
+	ASSERT_TRUE(equal(v & m2, vm2));
+	ASSERT_TRUE(equal(v & m3, vm3));
+	ASSERT_TRUE(equal(v & m4, vm4));
+}
+
+TYPED_TEST(T128Test, BitwiseOR) {
+	using F128 = typename TestFixture::T128;
+	F128 v1{ 0.0f, 2.0f, 0.0f, 4.0f };
+	F128 v2{ 1.0f, 0.0f, 3.0f, 0.0f };
+
+	F128 v12{ 1.0f, 2.0f, 3.0f, 4.0f };
+
+	ASSERT_TRUE(equal(v1 | v2, v12));
+}
+
+TYPED_TEST(T128Test, BitwiseXOR) {
+	using F128 = typename TestFixture::T128;
+	F128 v1{ 0.0f, 2.0f, 0.0f, 4.0f };
+	F128 v2{ 0.0f, 0.0f, 3.0f, 4.0f };
+
+	F128 v12{ 0.0f, 2.0f, 3.0f, 0.0f };
+
+	ASSERT_TRUE(equal(v1 ^ v2, v12));
+}
+
+TYPED_TEST(T128Test, BitwiseNOT) {
+	using F128 = typename TestFixture::T128;
+	F128 v1 = frombits<F128>(0x00000000u, 0x00000000u, 0x00000000u, 0x00000000u);
+	F128 v2 = frombits<F128>(0xFFFFFFFFu, 0xFFFFFFFFu, 0xFFFFFFFFu, 0xFFFFFFFFu);
+
+	ASSERT_TRUE(equal(~v1, v2));
+	ASSERT_TRUE(equal(~v2, v1));
 }
 
 TYPED_TEST(T128Test, ConstructFromArgs) {
-    typename TestFixture::T128 v{1.0f, 2.0f, 3.0f, 4.0f};
-    typename TestFixture::T128 expected{1.0f, 2.0f, 3.0f, 4.0f};
-    ASSERT_EQ(v, expected);
+	typename TestFixture::T128 v{ 1.0f, 2.0f, 3.0f, 4.0f };
+	typename TestFixture::T128 expected{ 1.0f, 2.0f, 3.0f, 4.0f };
+	ASSERT_TRUE(equal(v, expected));
 }
 
 TYPED_TEST(T128Test, ConstructFromSingleValue) {
-    typename TestFixture::T128 v{42.0f};
-    typename TestFixture::T128 expected{42.0f, 42.0f, 42.0f, 42.0f};
-    ASSERT_EQ(v, expected);
+	typename TestFixture::T128 v{ 42.0f };
+	typename TestFixture::T128 expected{ 42.0f, 42.0f, 42.0f, 42.0f };
+	ASSERT_TRUE(equal(v, expected));
 }
 
 TYPED_TEST(T128Test, FromAlignedSource) {
-    alignas(TestFixture::T128::alignment()) const float expected[] = {1.0f, 2.0f, 3.0f, 4.0f};
-    auto v = TestFixture::T128::fromAlignedSource(expected);
+	alignas(TestFixture::T128::alignment()) const float expected[] = { 1.0f, 2.0f, 3.0f, 4.0f };
+	auto v = TestFixture::T128::fromAlignedSource(expected);
 
-    alignas(TestFixture::T128::alignment()) float result[TestFixture::T128::size()];
-    v.alignedStore(result);
+	alignas(TestFixture::T128::alignment()) float result[TestFixture::T128::size()];
+	v.alignedStore(result);
 
-    ASSERT_ELEMENTS_EQ(result, expected, TestFixture::T128::size());
+	ASSERT_ELEMENTS_EQ(result, expected, TestFixture::T128::size());
 }
 
 TYPED_TEST(T128Test, AlignedLoadStore) {
-    alignas(TestFixture::T128::alignment()) const float expected[] = {1.0f, 2.0f, 3.0f, 4.0f};
-    typename TestFixture::T128 v;
-    v.alignedLoad(expected);
+	alignas(TestFixture::T128::alignment()) const float expected[] = { 1.0f, 2.0f, 3.0f, 4.0f };
+	typename TestFixture::T128 v;
+	v.alignedLoad(expected);
 
-    alignas(TestFixture::T128::alignment()) float result[TestFixture::T128::size()];
-    v.alignedStore(result);
+	alignas(TestFixture::T128::alignment()) float result[TestFixture::T128::size()];
+	v.alignedStore(result);
 
-    ASSERT_ELEMENTS_EQ(result, expected, TestFixture::T128::size());
+	ASSERT_ELEMENTS_EQ(result, expected, TestFixture::T128::size());
 }
 
 TYPED_TEST(T128Test, FromUnalignedSource) {
-    const float expected[] = {1.0f, 2.0f, 3.0f, 4.0f};
-    auto v = TestFixture::T128::fromUnalignedSource(expected);
+	const float expected[] = { 1.0f, 2.0f, 3.0f, 4.0f };
+	auto v = TestFixture::T128::fromUnalignedSource(expected);
 
-    float result[TestFixture::T128::size()];
-    v.unalignedStore(result);
+	float result[TestFixture::T128::size()];
+	v.unalignedStore(result);
 
-    ASSERT_ELEMENTS_EQ(result, expected, TestFixture::T128::size());
+	ASSERT_ELEMENTS_EQ(result, expected, TestFixture::T128::size());
 }
 
 TYPED_TEST(T128Test, UnalignedLoadStore) {
-    const float expected[] = {1.0f, 2.0f, 3.0f, 4.0f};
-    typename TestFixture::T128 v;
-    v.unalignedLoad(expected);
+	const float expected[] = { 1.0f, 2.0f, 3.0f, 4.0f };
+	typename TestFixture::T128 v;
+	v.unalignedLoad(expected);
 
-    float result[TestFixture::T128::size()];
-    v.unalignedStore(result);
+	float result[TestFixture::T128::size()];
+	v.unalignedStore(result);
 
-    ASSERT_ELEMENTS_EQ(result, expected, TestFixture::T128::size());
+	ASSERT_ELEMENTS_EQ(result, expected, TestFixture::T128::size());
 }
 
 TYPED_TEST(T128Test, LoadSingleValue) {
-    const float source[] = {42.0f, 43.0f, 44.0f, 45.0f};
-    auto v = TestFixture::T128::loadSingleValue(source);
-    typename TestFixture::T128 expected{42.0f, 0.0f, 0.0f, 0.0f};
-    ASSERT_EQ(v, expected);
+	const float source[] = { 42.0f, 43.0f, 44.0f, 45.0f };
+	auto v = TestFixture::T128::loadSingleValue(source);
+	typename TestFixture::T128 expected{ 42.0f, 0.0f, 0.0f, 0.0f };
+	ASSERT_TRUE(equal(v, expected));
 }
 
 TYPED_TEST(T128Test, Sum) {
-    typename TestFixture::T128 v{1.0f, 2.0f, 3.0f, 4.0f};
-    ASSERT_EQ(v.sum(), 10.0f);
+	typename TestFixture::T128 v{ 1.0f, 2.0f, 3.0f, 4.0f };
+	ASSERT_EQ(v.sum(), 10.0f);
 }
 
 TYPED_TEST(T128Test, CompoundAssignmentAdd) {
-    typename TestFixture::T128 v1{1.0f, 2.0f, 3.0f, 4.0f};
-    typename TestFixture::T128 v2{3.0f, 4.0f, 5.0f, 6.0f};
-    typename TestFixture::T128 expected{4.0f, 6.0f, 8.0f, 10.0f};
-    v1 += v2;
-    ASSERT_EQ(v1, expected);
+	typename TestFixture::T128 v1{ 1.0f, 2.0f, 3.0f, 4.0f };
+	typename TestFixture::T128 v2{ 3.0f, 4.0f, 5.0f, 6.0f };
+	typename TestFixture::T128 expected{ 4.0f, 6.0f, 8.0f, 10.0f };
+	v1 += v2;
+	ASSERT_TRUE(equal(v1, expected));
 }
 
 TYPED_TEST(T128Test, CompoundAssignmentSub) {
-    typename TestFixture::T128 v1{1.0f, 2.0f, 3.0f, 4.0f};
-    typename TestFixture::T128 v2{3.0f, 4.0f, 5.0f, 6.0f};
-    typename TestFixture::T128 expected{-2.0f, -2.0f, -2.0f, -2.0f};
-    v1 -= v2;
-    ASSERT_EQ(v1, expected);
+	typename TestFixture::T128 v1{ 1.0f, 2.0f, 3.0f, 4.0f };
+	typename TestFixture::T128 v2{ 3.0f, 4.0f, 5.0f, 6.0f };
+	typename TestFixture::T128 expected{ -2.0f, -2.0f, -2.0f, -2.0f };
+	v1 -= v2;
+	ASSERT_TRUE(equal(v1, expected));
 }
 
 TYPED_TEST(T128Test, CompoundAssignmentMul) {
-    typename TestFixture::T128 v1{1.0f, 2.0f, 3.0f, 4.0f};
-    typename TestFixture::T128 v2{3.0f, 4.0f, 5.0f, 6.0f};
-    typename TestFixture::T128 expected{3.0f, 8.0f, 15.0f, 24.0f};
-    v1 *= v2;
-    ASSERT_EQ(v1, expected);
+	typename TestFixture::T128 v1{ 1.0f, 2.0f, 3.0f, 4.0f };
+	typename TestFixture::T128 v2{ 3.0f, 4.0f, 5.0f, 6.0f };
+	typename TestFixture::T128 expected{ 3.0f, 8.0f, 15.0f, 24.0f };
+	v1 *= v2;
+	ASSERT_TRUE(equal(v1, expected));
 }
 
 TYPED_TEST(T128Test, CompoundAssignmentDiv) {
-    typename TestFixture::T128 v1{4.0f, 3.0f, 9.0f, 12.0f};
-    typename TestFixture::T128 v2{1.0f, 2.0f, 3.0f, 3.0f};
-    float expected[TestFixture::T128::size()] = {4.0f, 1.5f, 3.0f, 4.0f};
-    v1 /= v2;
+	typename TestFixture::T128 v1{ 4.0f, 3.0f, 9.0f, 12.0f };
+	typename TestFixture::T128 v2{ 1.0f, 2.0f, 3.0f, 3.0f };
+	float expected[TestFixture::T128::size()] = { 4.0f, 1.5f, 3.0f, 4.0f };
+	v1 /= v2;
 
-    float result[TestFixture::T128::size()];
-    v1.unalignedStore(result);
+	float result[TestFixture::T128::size()];
+	v1.unalignedStore(result);
 
-    ASSERT_ELEMENTS_NEAR(result, expected, TestFixture::T128::size(), 0.0001f);
+	ASSERT_ELEMENTS_NEAR(result, expected, TestFixture::T128::size(), 0.0001f);
 }
 
 TYPED_TEST(T128Test, OperatorAdd) {
-    typename TestFixture::T128 v1{1.0f, 2.0f, 3.0f, 4.0f};
-    typename TestFixture::T128 v2{3.0f, 4.0f, 5.0f, 6.0f};
-    typename TestFixture::T128 expected{4.0f, 6.0f, 8.0f, 10.0f};
-    auto v3 = v1 + v2;
-    ASSERT_EQ(v3, expected);
+	typename TestFixture::T128 v1{ 1.0f, 2.0f, 3.0f, 4.0f };
+	typename TestFixture::T128 v2{ 3.0f, 4.0f, 5.0f, 6.0f };
+	typename TestFixture::T128 expected{ 4.0f, 6.0f, 8.0f, 10.0f };
+	auto v3 = v1 + v2;
+	ASSERT_TRUE(equal(v3, expected));
 }
 
 TYPED_TEST(T128Test, OperatorSub) {
-    typename TestFixture::T128 v1{1.0f, 2.0f, 3.0f, 4.0f};
-    typename TestFixture::T128 v2{3.0f, 4.0f, 5.0f, 6.0f};
-    typename TestFixture::T128 expected{-2.0f, -2.0f, -2.0f, -2.0f};
-    auto v3 = v1 - v2;
-    ASSERT_EQ(v3, expected);
+	typename TestFixture::T128 v1{ 1.0f, 2.0f, 3.0f, 4.0f };
+	typename TestFixture::T128 v2{ 3.0f, 4.0f, 5.0f, 6.0f };
+	typename TestFixture::T128 expected{ -2.0f, -2.0f, -2.0f, -2.0f };
+	auto v3 = v1 - v2;
+	ASSERT_TRUE(equal(v3, expected));
 }
 
 TYPED_TEST(T128Test, OperatorMul) {
-    typename TestFixture::T128 v1{1.0f, 2.0f, 3.0f, 4.0f};
-    typename TestFixture::T128 v2{3.0f, 4.0f, 5.0f, 6.0f};
-    typename TestFixture::T128 expected{3.0f, 8.0f, 15.0f, 24.0f};
-    auto v3 = v1 * v2;
-    ASSERT_EQ(v3, expected);
+	typename TestFixture::T128 v1{ 1.0f, 2.0f, 3.0f, 4.0f };
+	typename TestFixture::T128 v2{ 3.0f, 4.0f, 5.0f, 6.0f };
+	typename TestFixture::T128 expected{ 3.0f, 8.0f, 15.0f, 24.0f };
+	auto v3 = v1 * v2;
+	ASSERT_TRUE(equal(v3, expected));
 }
 
 TYPED_TEST(T128Test, OperatorDiv) {
-    typename TestFixture::T128 v1{4.0f, 3.0f, 9.0f, 12.0f};
-    typename TestFixture::T128 v2{1.0f, 2.0f, 3.0f, 3.0f};
-    float expected[TestFixture::T128::size()] = {4.0f, 1.5f, 3.0f, 4.0f};
-    auto v3 = v1 / v2;
+	typename TestFixture::T128 v1{ 4.0f, 3.0f, 9.0f, 12.0f };
+	typename TestFixture::T128 v2{ 1.0f, 2.0f, 3.0f, 3.0f };
+	float expected[TestFixture::T128::size()] = { 4.0f, 1.5f, 3.0f, 4.0f };
+	auto v3 = v1 / v2;
 
-    float result[TestFixture::T128::size()];
-    v3.unalignedStore(result);
+	float result[TestFixture::T128::size()];
+	v3.unalignedStore(result);
 
-    ASSERT_ELEMENTS_NEAR(result, expected, TestFixture::T128::size(), 0.0001f);
+	ASSERT_ELEMENTS_NEAR(result, expected, TestFixture::T128::size(), 0.0001f);
 }
 
 TEST(T128Test, TransposeSquareScalar) {
-    trimd::scalar::F128 v1{1.0f, 2.0f, 3.0f, 4.0f};
-    trimd::scalar::F128 v2{1.0f, 2.0f, 3.0f, 4.0f};
-    trimd::scalar::F128 v3{1.0f, 2.0f, 3.0f, 4.0f};
-    trimd::scalar::F128 v4{1.0f, 2.0f, 3.0f, 4.0f};
-
-    trimd::scalar::transpose(v1, v2, v3, v4);
+	trimd::scalar::F128 v1{ 1.0f, 2.0f, 3.0f, 4.0f };
+	trimd::scalar::F128 v2{ 1.0f, 2.0f, 3.0f, 4.0f };
+	trimd::scalar::F128 v3{ 1.0f, 2.0f, 3.0f, 4.0f };
+	trimd::scalar::F128 v4{ 1.0f, 2.0f, 3.0f, 4.0f };
+
+	trimd::scalar::transpose(v1, v2, v3, v4);
+
+	trimd::scalar::F128 e1{ 1.0f, 1.0f, 1.0f, 1.0f };
+	trimd::scalar::F128 e2{ 2.0f, 2.0f, 2.0f, 2.0f };
+	trimd::scalar::F128 e3{ 3.0f, 3.0f, 3.0f, 3.0f };
+	trimd::scalar::F128 e4{ 4.0f, 4.0f, 4.0f, 4.0f };
+
+	ASSERT_TRUE(equal(v1, e1));
+	ASSERT_TRUE(equal(v2, e2));
+	ASSERT_TRUE(equal(v3, e3));
+	ASSERT_TRUE(equal(v4, e4));
+}
 
-    trimd::scalar::F128 e1{1.0f, 1.0f, 1.0f, 1.0f};
-    trimd::scalar::F128 e2{2.0f, 2.0f, 2.0f, 2.0f};
-    trimd::scalar::F128 e3{3.0f, 3.0f, 3.0f, 3.0f};
-    trimd::scalar::F128 e4{4.0f, 4.0f, 4.0f, 4.0f};
+TEST(T128Test, AbsScalar) {
+	trimd::scalar::F128 v{ -1.0f, 2.0f, -3.0f, 0.0f };
+	v = trimd::scalar::abs(v);
+	trimd::scalar::F128 e{ 1.0f, 2.0f, 3.0f, 0.0f };
+	ASSERT_TRUE(equal(v, e));
+}
 
-    ASSERT_EQ(v1, e1);
-    ASSERT_EQ(v2, e2);
-    ASSERT_EQ(v3, e3);
-    ASSERT_EQ(v4, e4);
+TEST(T128Test, AndNotScalar) {
+	trimd::scalar::F128 v{ 1.0f, 2.0f, 3.0f, 4.0f };
+	trimd::scalar::F128 mask1 = frombits<trimd::scalar::F128>(0x00000000u, 0x00000000u, 0x00000000u, 0x00000000u);
+	trimd::scalar::F128 mask2 = frombits<trimd::scalar::F128>(0xFFFFFFFFu, 0xFFFFFFFFu, 0xFFFFFFFFu, 0xFFFFFFFFu);
+	trimd::scalar::F128 result1 = trimd::scalar::andnot(mask1, v);
+	trimd::scalar::F128 result2 = trimd::scalar::andnot(mask2, v);
+	trimd::scalar::F128 e1{ 1.0f, 2.0f, 3.0f, 4.0f };
+	trimd::scalar::F128 e2{ 0.0f, 0.0f, 0.0f, 0.0f };
+	ASSERT_TRUE(equal(result1, e1));
+	ASSERT_TRUE(equal(result2, e2));
 }
 
 #ifdef TRIMD_ENABLE_SSE
 TEST(T128Test, TransposeSquareSSE) {
-    trimd::sse::F128 v1{1.0f, 2.0f, 3.0f, 4.0f};
-    trimd::sse::F128 v2{1.0f, 2.0f, 3.0f, 4.0f};
-    trimd::sse::F128 v3{1.0f, 2.0f, 3.0f, 4.0f};
-    trimd::sse::F128 v4{1.0f, 2.0f, 3.0f, 4.0f};
+	trimd::sse::F128 v1{ 1.0f, 2.0f, 3.0f, 4.0f };
+	trimd::sse::F128 v2{ 1.0f, 2.0f, 3.0f, 4.0f };
+	trimd::sse::F128 v3{ 1.0f, 2.0f, 3.0f, 4.0f };
+	trimd::sse::F128 v4{ 1.0f, 2.0f, 3.0f, 4.0f };
+
+	trimd::sse::transpose(v1, v2, v3, v4);
+
+	trimd::sse::F128 e1{ 1.0f, 1.0f, 1.0f, 1.0f };
+	trimd::sse::F128 e2{ 2.0f, 2.0f, 2.0f, 2.0f };
+	trimd::sse::F128 e3{ 3.0f, 3.0f, 3.0f, 3.0f };
+	trimd::sse::F128 e4{ 4.0f, 4.0f, 4.0f, 4.0f };
+
+	ASSERT_TRUE(equal(v1, e1));
+	ASSERT_TRUE(equal(v2, e2));
+	ASSERT_TRUE(equal(v3, e3));
+	ASSERT_TRUE(equal(v4, e4));
+}
+
+TEST(T128Test, AbsSSE) {
+	trimd::sse::F128 v{ -1.0f, 2.0f, -3.0f, 0.0f };
+	v = trimd::sse::abs(v);
+	trimd::sse::F128 e{ 1.0f, 2.0f, 3.0f, 0.0f };
+	ASSERT_TRUE(equal(v, e));
+}
+
+TEST(T128Test, AndNotSSE) {
+	trimd::sse::F128 v{ 1.0f, 2.0f, 3.0f, 4.0f };
+	trimd::sse::F128 mask1 = frombits<trimd::sse::F128>(0x00000000u, 0x00000000u, 0x00000000u, 0x00000000u);
+	trimd::sse::F128 mask2 = frombits<trimd::sse::F128>(0xFFFFFFFFu, 0xFFFFFFFFu, 0xFFFFFFFFu, 0xFFFFFFFFu);
+	trimd::sse::F128 result1 = trimd::sse::andnot(mask1, v);
+	trimd::sse::F128 result2 = trimd::sse::andnot(mask2, v);
+	trimd::sse::F128 e1{ 1.0f, 2.0f, 3.0f, 4.0f };
+	trimd::sse::F128 e2{ 0.0f, 0.0f, 0.0f, 0.0f };
+	ASSERT_TRUE(equal(result1, e1));
+	ASSERT_TRUE(equal(result2, e2));
+}
+
+#ifdef TRIMD_ENABLE_F16C
+TEST(T128Test, LoadAlignedHalfFloats) {
+	alignas(trimd::sse::F128::alignment()) const std::uint16_t halfFloats[] = { 15360, 16384, 16896, 17408 };
+	trimd::sse::F128 expected{ 1.0f, 2.0f, 3.0f, 4.0f };
+
+	auto v = trimd::sse::F128::fromAlignedSource(halfFloats);
+	ASSERT_TRUE(equal(v, expected));
+
+	trimd::sse::F128 v2;
+	v2.alignedLoad(halfFloats);
+	ASSERT_TRUE(equal(v2, expected));
+}
+
+TEST(T128Test, LoadUnalignedHalfFloats) {
+	const std::uint16_t halfFloats[] = { 15360, 16384, 16896, 17408 };
+	trimd::sse::F128 expected{ 1.0f, 2.0f, 3.0f, 4.0f };
+
+	auto v = trimd::sse::F128::fromUnalignedSource(halfFloats);
+	ASSERT_TRUE(equal(v, expected));
+
+	trimd::sse::F128 v2;
+	v2.unalignedLoad(halfFloats);
+	ASSERT_TRUE(equal(v2, expected));
+}
+
+TEST(T128Test, StoreAlignedHalfFloats) {
+	const std::uint16_t expected[] = { 15360, 16384, 16896, 17408 };
+
+	alignas(trimd::sse::F128::alignment()) std::uint16_t halfFloats[4ul] = {};
+	const trimd::sse::F128 v{ 1.0f, 2.0f, 3.0f, 4.0f };
+	v.alignedStore(halfFloats);
+
+	ASSERT_ELEMENTS_EQ(halfFloats, expected, 4ul);
+}
 
-    trimd::sse::transpose(v1, v2, v3, v4);
+TEST(T128Test, StoreUnalignedHalfFloats) {
+	const std::uint16_t expected[] = { 15360, 16384, 16896, 17408 };
 
-    trimd::sse::F128 e1{1.0f, 1.0f, 1.0f, 1.0f};
-    trimd::sse::F128 e2{2.0f, 2.0f, 2.0f, 2.0f};
-    trimd::sse::F128 e3{3.0f, 3.0f, 3.0f, 3.0f};
-    trimd::sse::F128 e4{4.0f, 4.0f, 4.0f, 4.0f};
+	std::uint16_t halfFloats[4ul] = {};
+	const trimd::sse::F128 v{ 1.0f, 2.0f, 3.0f, 4.0f };
+	v.unalignedStore(halfFloats);
 
-    ASSERT_EQ(v1, e1);
-    ASSERT_EQ(v2, e2);
-    ASSERT_EQ(v3, e3);
-    ASSERT_EQ(v4, e4);
+	ASSERT_ELEMENTS_EQ(halfFloats, expected, 4ul);
 }
+#endif  // TRIMD_ENABLE_F16C
 #endif  // TRIMD_ENABLE_SSE
diff --git a/Engine/Plugins/Animation/RigLogic/Source/RigLogicLibTest/Private/trimdtests/TestT256.cpp b/Engine/Plugins/Animation/RigLogic/Source/RigLogicLibTest/Private/trimdtests/TestT256.cpp
--- a/Engine/Plugins/Animation/RigLogic/Source/RigLogicLibTest/Private/trimdtests/TestT256.cpp
+++ b/Engine/Plugins/Animation/RigLogic/Source/RigLogicLibTest/Private/trimdtests/TestT256.cpp
@@ -4,253 +4,906 @@
 
 #include "trimd/TRiMD.h"
 
-#ifdef TRIMD_ENABLE_AVX
+#if defined(TRIMD_ENABLE_AVX) && defined(TRIMD_ENABLE_SSE)
+using T256Types = ::testing::Types<trimd::scalar::F256, trimd::sse::F256, trimd::avx::F256>;
+
+#elif defined(TRIMD_ENABLE_AVX)
 using T256Types = ::testing::Types<trimd::scalar::F256, trimd::avx::F256>;
+
+#elif defined(TRIMD_ENABLE_SSE)
+using T256Types = ::testing::Types<trimd::scalar::F256, trimd::sse::F256>;
+
 #else
 using T256Types = ::testing::Types<trimd::scalar::F256>;
 #endif  // TRIMD_ENABLE_SSE
 
+template<typename TF256>
+static TF256 frombits(uint32_t bits0,
+	uint32_t bits1,
+	uint32_t bits2,
+	uint32_t bits3,
+	uint32_t bits4,
+	uint32_t bits5,
+	uint32_t bits6,
+	uint32_t bits7) {
+	return TF256{
+		trimd::bitcast<float>(bits0),
+		trimd::bitcast<float>(bits1),
+		trimd::bitcast<float>(bits2),
+		trimd::bitcast<float>(bits3),
+		trimd::bitcast<float>(bits4),
+		trimd::bitcast<float>(bits5),
+		trimd::bitcast<float>(bits6),
+		trimd::bitcast<float>(bits7),
+	};
+}
+
+#ifdef TRIMD_ENABLE_AVX
+bool equal(const trimd::avx::F256& lhs, const trimd::avx::F256& rhs) {
+	return (std::memcmp(&lhs.data, &rhs.data, sizeof(lhs.data)) == 0);
+}
+
+#endif  // TRIMD_ENABLE_AVX
+
+#ifdef TRIMD_ENABLE_SSE
+static bool equal(const trimd::fallback::T256<trimd::sse::F128>& lhs, const trimd::fallback::T256<trimd::sse::F128>& rhs) {
+	const bool data1eq = (std::memcmp(&lhs.data1.data, &rhs.data1.data, sizeof(float) * decltype(lhs.data1)::size()) == 0);
+	const bool data2eq = (std::memcmp(&lhs.data2.data, &rhs.data2.data, sizeof(float) * decltype(lhs.data2)::size()) == 0);
+	return data1eq && data2eq;
+}
+
+#endif  // TRIMD_ENABLE_SSE
+
+static bool equal(const trimd::fallback::T256<trimd::scalar::F128>& lhs, const trimd::fallback::T256<trimd::scalar::F128>& rhs) {
+	const bool data1eq =
+		(std::memcmp(lhs.data1.data.data(), rhs.data1.data.data(), sizeof(float) * decltype(lhs.data1)::size()) == 0);
+	const bool data2eq =
+		(std::memcmp(lhs.data2.data.data(), rhs.data2.data.data(), sizeof(float) * decltype(lhs.data2)::size()) == 0);
+	return data1eq && data2eq;
+}
+
 template<typename T>
 class T256Test : public ::testing::Test {
 protected:
-    using T256 = T;
+	using T256 = T;
 
 };
 
 TYPED_TEST_SUITE(T256Test, T256Types, );
 
 TYPED_TEST(T256Test, CheckSize) {
-    ASSERT_EQ(TestFixture::T256::size(), 8ul);
+	ASSERT_EQ(TestFixture::T256::size(), 8ul);
 }
 
 TYPED_TEST(T256Test, Equality) {
-    typename TestFixture::T256 v1{1.0f, 2.0f, 3.0f, 4.0f, 5.0f, 6.0f, 7.0f, 8.0f};
-    typename TestFixture::T256 v2{1.0f, 2.0f, 3.0f, 4.0f, 5.0f, 6.0f, 7.0f, 8.0f};
-
-    typename TestFixture::T256 v3{1.5f, 2.0f, 3.0f, 4.0f, 5.0f, 6.0f, 7.0f, 8.0f};
-    typename TestFixture::T256 v4{1.0f, 2.5f, 3.0f, 4.0f, 5.0f, 6.0f, 7.0f, 8.0f};
-    typename TestFixture::T256 v5{1.0f, 2.0f, 3.5f, 4.0f, 5.0f, 6.0f, 7.0f, 8.0f};
-    typename TestFixture::T256 v6{1.0f, 2.0f, 3.0f, 4.5f, 5.0f, 6.0f, 7.0f, 8.0f};
-    typename TestFixture::T256 v7{1.0f, 2.0f, 3.0f, 4.0f, 5.5f, 6.0f, 7.0f, 8.0f};
-    typename TestFixture::T256 v8{1.0f, 2.0f, 3.0f, 4.0f, 5.0f, 6.5f, 7.0f, 8.0f};
-    typename TestFixture::T256 v9{1.0f, 2.0f, 3.0f, 4.0f, 5.0f, 6.0f, 7.5f, 8.0f};
-    typename TestFixture::T256 v10{1.0f, 2.0f, 3.0f, 4.0f, 5.0f, 6.0f, 7.0f, 8.5f};
-
-    ASSERT_TRUE(v1 == v2);
-    ASSERT_FALSE(v1 == v3);
-    ASSERT_FALSE(v1 == v4);
-    ASSERT_FALSE(v1 == v5);
-    ASSERT_FALSE(v1 == v6);
-    ASSERT_FALSE(v1 == v7);
-    ASSERT_FALSE(v1 == v8);
-    ASSERT_FALSE(v1 == v9);
-    ASSERT_FALSE(v1 == v10);
+	using F256 = typename TestFixture::T256;
+
+	F256 v1{ 1.0f, 2.0f, 3.0f, 4.0f, 5.0f, 6.0f, 7.0f, 8.0f };
+	F256 v2{ 1.0f, 2.0f, 3.0f, 4.0f, 5.0f, 6.0f, 7.0f, 8.0f };
+	F256 v3{ 1.5f, 2.0f, 3.0f, 4.0f, 5.0f, 6.0f, 7.0f, 8.0f };
+	F256 v4{ 1.0f, 2.5f, 3.0f, 4.0f, 5.0f, 6.0f, 7.0f, 8.0f };
+	F256 v5{ 1.0f, 2.0f, 3.5f, 4.0f, 5.0f, 6.0f, 7.0f, 8.0f };
+	F256 v6{ 1.0f, 2.0f, 3.0f, 4.5f, 5.0f, 6.0f, 7.0f, 8.0f };
+	F256 v7{ 1.0f, 2.0f, 3.0f, 4.0f, 5.5f, 6.0f, 7.0f, 8.0f };
+	F256 v8{ 1.0f, 2.0f, 3.0f, 4.0f, 5.0f, 6.5f, 7.0f, 8.0f };
+	F256 v9{ 1.0f, 2.0f, 3.0f, 4.0f, 5.0f, 6.0f, 7.5f, 8.0f };
+	F256 v10{ 1.0f, 2.0f, 3.0f, 4.0f, 5.0f, 6.0f, 7.0f, 8.5f };
+
+	F256 m12 = frombits<F256>(0xFFFFFFFFu,
+		0xFFFFFFFFu,
+		0xFFFFFFFFu,
+		0xFFFFFFFFu,
+		0xFFFFFFFFu,
+		0xFFFFFFFFu,
+		0xFFFFFFFFu,
+		0xFFFFFFFFu);
+	F256 m13 = frombits<F256>(0x00000000u,
+		0xFFFFFFFFu,
+		0xFFFFFFFFu,
+		0xFFFFFFFFu,
+		0xFFFFFFFFu,
+		0xFFFFFFFFu,
+		0xFFFFFFFFu,
+		0xFFFFFFFFu);
+	F256 m14 = frombits<F256>(0xFFFFFFFFu,
+		0x00000000u,
+		0xFFFFFFFFu,
+		0xFFFFFFFFu,
+		0xFFFFFFFFu,
+		0xFFFFFFFFu,
+		0xFFFFFFFFu,
+		0xFFFFFFFFu);
+	F256 m15 = frombits<F256>(0xFFFFFFFFu,
+		0xFFFFFFFFu,
+		0x00000000u,
+		0xFFFFFFFFu,
+		0xFFFFFFFFu,
+		0xFFFFFFFFu,
+		0xFFFFFFFFu,
+		0xFFFFFFFFu);
+	F256 m16 = frombits<F256>(0xFFFFFFFFu,
+		0xFFFFFFFFu,
+		0xFFFFFFFFu,
+		0x00000000u,
+		0xFFFFFFFFu,
+		0xFFFFFFFFu,
+		0xFFFFFFFFu,
+		0xFFFFFFFFu);
+	F256 m17 = frombits<F256>(0xFFFFFFFFu,
+		0xFFFFFFFFu,
+		0xFFFFFFFFu,
+		0xFFFFFFFFu,
+		0x00000000u,
+		0xFFFFFFFFu,
+		0xFFFFFFFFu,
+		0xFFFFFFFFu);
+	F256 m18 = frombits<F256>(0xFFFFFFFFu,
+		0xFFFFFFFFu,
+		0xFFFFFFFFu,
+		0xFFFFFFFFu,
+		0xFFFFFFFFu,
+		0x00000000u,
+		0xFFFFFFFFu,
+		0xFFFFFFFFu);
+	F256 m19 = frombits<F256>(0xFFFFFFFFu,
+		0xFFFFFFFFu,
+		0xFFFFFFFFu,
+		0xFFFFFFFFu,
+		0xFFFFFFFFu,
+		0xFFFFFFFFu,
+		0x00000000u,
+		0xFFFFFFFFu);
+	F256 m110 = frombits<F256>(0xFFFFFFFFu,
+		0xFFFFFFFFu,
+		0xFFFFFFFFu,
+		0xFFFFFFFFu,
+		0xFFFFFFFFu,
+		0xFFFFFFFFu,
+		0xFFFFFFFFu,
+		0x00000000u);
+
+	ASSERT_TRUE(equal(v1 == v2, m12));
+	ASSERT_TRUE(equal(v1 == v3, m13));
+	ASSERT_TRUE(equal(v1 == v4, m14));
+	ASSERT_TRUE(equal(v1 == v5, m15));
+	ASSERT_TRUE(equal(v1 == v6, m16));
+	ASSERT_TRUE(equal(v1 == v7, m17));
+	ASSERT_TRUE(equal(v1 == v8, m18));
+	ASSERT_TRUE(equal(v1 == v9, m19));
+	ASSERT_TRUE(equal(v1 == v10, m110));
 }
 
 TYPED_TEST(T256Test, Inequality) {
-    typename TestFixture::T256 v1{1.0f, 2.0f, 3.0f, 4.0f, 5.0f, 6.0f, 7.0f, 8.0f};
-    typename TestFixture::T256 v2{1.0f, 2.0f, 3.0f, 4.0f, 5.0f, 6.0f, 7.0f, 8.0f};
-    ASSERT_FALSE(v1 != v2);
+	using F256 = typename TestFixture::T256;
+
+	F256 v1{ 1.0f, 2.0f, 3.0f, 4.0f, 5.0f, 6.0f, 7.0f, 8.0f };
+	F256 v2{ 1.0f, 2.0f, 3.0f, 4.0f, 5.0f, 6.0f, 7.0f, 8.0f };
+	F256 v3{ 1.5f, 2.0f, 3.0f, 4.0f, 5.0f, 6.0f, 7.0f, 8.0f };
+	F256 v4{ 1.0f, 2.5f, 3.0f, 4.0f, 5.0f, 6.0f, 7.0f, 8.0f };
+	F256 v5{ 1.0f, 2.0f, 3.5f, 4.0f, 5.0f, 6.0f, 7.0f, 8.0f };
+	F256 v6{ 1.0f, 2.0f, 3.0f, 4.5f, 5.0f, 6.0f, 7.0f, 8.0f };
+	F256 v7{ 1.0f, 2.0f, 3.0f, 4.0f, 5.5f, 6.0f, 7.0f, 8.0f };
+	F256 v8{ 1.0f, 2.0f, 3.0f, 4.0f, 5.0f, 6.5f, 7.0f, 8.0f };
+	F256 v9{ 1.0f, 2.0f, 3.0f, 4.0f, 5.0f, 6.0f, 7.5f, 8.0f };
+	F256 v10{ 1.0f, 2.0f, 3.0f, 4.0f, 5.0f, 6.0f, 7.0f, 8.5f };
+
+	F256 m12 = frombits<F256>(0x00000000u,
+		0x00000000u,
+		0x00000000u,
+		0x00000000u,
+		0x00000000u,
+		0x00000000u,
+		0x00000000u,
+		0x00000000u);
+	F256 m13 = frombits<F256>(0xFFFFFFFFu,
+		0x00000000u,
+		0x00000000u,
+		0x00000000u,
+		0x00000000u,
+		0x00000000u,
+		0x00000000u,
+		0x00000000u);
+	F256 m14 = frombits<F256>(0x00000000u,
+		0xFFFFFFFFu,
+		0x00000000u,
+		0x00000000u,
+		0x00000000u,
+		0x00000000u,
+		0x00000000u,
+		0x00000000u);
+	F256 m15 = frombits<F256>(0x00000000u,
+		0x00000000u,
+		0xFFFFFFFFu,
+		0x00000000u,
+		0x00000000u,
+		0x00000000u,
+		0x00000000u,
+		0x00000000u);
+	F256 m16 = frombits<F256>(0x00000000u,
+		0x00000000u,
+		0x00000000u,
+		0xFFFFFFFFu,
+		0x00000000u,
+		0x00000000u,
+		0x00000000u,
+		0x00000000u);
+	F256 m17 = frombits<F256>(0x00000000u,
+		0x00000000u,
+		0x00000000u,
+		0x00000000u,
+		0xFFFFFFFFu,
+		0x00000000u,
+		0x00000000u,
+		0x00000000u);
+	F256 m18 = frombits<F256>(0x00000000u,
+		0x00000000u,
+		0x00000000u,
+		0x00000000u,
+		0x00000000u,
+		0xFFFFFFFFu,
+		0x00000000u,
+		0x00000000u);
+	F256 m19 = frombits<F256>(0x00000000u,
+		0x00000000u,
+		0x00000000u,
+		0x00000000u,
+		0x00000000u,
+		0x00000000u,
+		0xFFFFFFFFu,
+		0x00000000u);
+	F256 m110 = frombits<F256>(0x00000000u,
+		0x00000000u,
+		0x00000000u,
+		0x00000000u,
+		0x00000000u,
+		0x00000000u,
+		0x00000000u,
+		0xFFFFFFFFu);
+
+	ASSERT_TRUE(equal(v1 != v2, m12));
+	ASSERT_TRUE(equal(v1 != v3, m13));
+	ASSERT_TRUE(equal(v1 != v4, m14));
+	ASSERT_TRUE(equal(v1 != v5, m15));
+	ASSERT_TRUE(equal(v1 != v6, m16));
+	ASSERT_TRUE(equal(v1 != v7, m17));
+	ASSERT_TRUE(equal(v1 != v8, m18));
+	ASSERT_TRUE(equal(v1 != v9, m19));
+	ASSERT_TRUE(equal(v1 != v10, m110));
+}
+
+TYPED_TEST(T256Test, LessThan) {
+	using F256 = typename TestFixture::T256;
+
+	F256 v1{ 1.0f, 2.0f, 3.0f, 4.0f, 5.0f, 6.0f, 7.0f, 8.0f };
+	F256 v2{ 1.0f, 2.0f, 3.0f, 4.0f, 5.0f, 6.0f, 7.0f, 8.0f };
+	F256 v3{ 2.0f, 3.0f, 4.0f, 5.0f, 6.0f, 7.0f, 8.0f, 9.0f };
+	F256 v4{ 0.5f, 1.5f, 2.5f, 3.5f, 4.5f, 5.5f, 6.5f, 7.5f };
+
+	F256 m12 = frombits<F256>(0x00000000u,
+		0x00000000u,
+		0x00000000u,
+		0x00000000u,
+		0x00000000u,
+		0x00000000u,
+		0x00000000u,
+		0x00000000u);
+	F256 m13 = frombits<F256>(0xFFFFFFFFu,
+		0xFFFFFFFFu,
+		0xFFFFFFFFu,
+		0xFFFFFFFFu,
+		0xFFFFFFFFu,
+		0xFFFFFFFFu,
+		0xFFFFFFFFu,
+		0xFFFFFFFFu);
+	F256 m14 = frombits<F256>(0x00000000u,
+		0x00000000u,
+		0x00000000u,
+		0x00000000u,
+		0x00000000u,
+		0x00000000u,
+		0x00000000u,
+		0x00000000u);
+
+	ASSERT_TRUE(equal(v1 < v2, m12));
+	ASSERT_TRUE(equal(v1 < v3, m13));
+	ASSERT_TRUE(equal(v1 < v4, m14));
+}
+
+TYPED_TEST(T256Test, LessThanOrEqual) {
+	using F256 = typename TestFixture::T256;
+
+	F256 v1{ 1.0f, 2.0f, 3.0f, 4.0f, 5.0f, 6.0f, 7.0f, 8.0f };
+	F256 v2{ 1.0f, 2.0f, 3.0f, 4.0f, 5.0f, 6.0f, 7.0f, 8.0f };
+	F256 v3{ 2.0f, 3.0f, 4.0f, 5.0f, 6.0f, 7.0f, 8.0f, 9.0f };
+	F256 v4{ 0.5f, 1.5f, 2.5f, 3.5f, 4.5f, 5.5f, 6.5f, 7.5f };
+
+	F256 m12 = frombits<F256>(0xFFFFFFFFu,
+		0xFFFFFFFFu,
+		0xFFFFFFFFu,
+		0xFFFFFFFFu,
+		0xFFFFFFFFu,
+		0xFFFFFFFFu,
+		0xFFFFFFFFu,
+		0xFFFFFFFFu);
+	F256 m13 = frombits<F256>(0xFFFFFFFFu,
+		0xFFFFFFFFu,
+		0xFFFFFFFFu,
+		0xFFFFFFFFu,
+		0xFFFFFFFFu,
+		0xFFFFFFFFu,
+		0xFFFFFFFFu,
+		0xFFFFFFFFu);
+	F256 m14 = frombits<F256>(0x00000000u,
+		0x00000000u,
+		0x00000000u,
+		0x00000000u,
+		0x00000000u,
+		0x00000000u,
+		0x00000000u,
+		0x00000000u);
+
+	ASSERT_TRUE(equal(v1 <= v2, m12));
+	ASSERT_TRUE(equal(v1 <= v3, m13));
+	ASSERT_TRUE(equal(v1 <= v4, m14));
+}
+
+TYPED_TEST(T256Test, GreaterThan) {
+	using F256 = typename TestFixture::T256;
+
+	F256 v1{ 1.0f, 2.0f, 3.0f, 4.0f, 5.0f, 6.0f, 7.0f, 8.0f };
+	F256 v2{ 1.0f, 2.0f, 3.0f, 4.0f, 5.0f, 6.0f, 7.0f, 8.0f };
+	F256 v3{ 2.0f, 3.0f, 4.0f, 5.0f, 6.0f, 7.0f, 8.0f, 9.0f };
+	F256 v4{ 0.5f, 1.5f, 2.5f, 3.5f, 4.5f, 5.5f, 6.5f, 7.5f };
+
+	F256 m12 = frombits<F256>(0x00000000u,
+		0x00000000u,
+		0x00000000u,
+		0x00000000u,
+		0x00000000u,
+		0x00000000u,
+		0x00000000u,
+		0x00000000u);
+	F256 m13 = frombits<F256>(0x00000000u,
+		0x00000000u,
+		0x00000000u,
+		0x00000000u,
+		0x00000000u,
+		0x00000000u,
+		0x00000000u,
+		0x00000000u);
+	F256 m14 = frombits<F256>(0xFFFFFFFFu,
+		0xFFFFFFFFu,
+		0xFFFFFFFFu,
+		0xFFFFFFFFu,
+		0xFFFFFFFFu,
+		0xFFFFFFFFu,
+		0xFFFFFFFFu,
+		0xFFFFFFFFu);
+
+	ASSERT_TRUE(equal(v1 > v2, m12));
+	ASSERT_TRUE(equal(v1 > v3, m13));
+	ASSERT_TRUE(equal(v1 > v4, m14));
+}
+
+TYPED_TEST(T256Test, GreaterThanOrEqual) {
+	using F256 = typename TestFixture::T256;
+
+	F256 v1{ 1.0f, 2.0f, 3.0f, 4.0f, 5.0f, 6.0f, 7.0f, 8.0f };
+	F256 v2{ 1.0f, 2.0f, 3.0f, 4.0f, 5.0f, 6.0f, 7.0f, 8.0f };
+	F256 v3{ 2.0f, 3.0f, 4.0f, 5.0f, 6.0f, 7.0f, 8.0f, 9.0f };
+	F256 v4{ 0.5f, 1.5f, 2.5f, 3.5f, 4.5f, 5.5f, 6.5f, 7.5f };
+
+	F256 m12 = frombits<F256>(0xFFFFFFFFu,
+		0xFFFFFFFFu,
+		0xFFFFFFFFu,
+		0xFFFFFFFFu,
+		0xFFFFFFFFu,
+		0xFFFFFFFFu,
+		0xFFFFFFFFu,
+		0xFFFFFFFFu);
+	F256 m13 = frombits<F256>(0x00000000u,
+		0x00000000u,
+		0x00000000u,
+		0x00000000u,
+		0x00000000u,
+		0x00000000u,
+		0x00000000u,
+		0x00000000u);
+	F256 m14 = frombits<F256>(0xFFFFFFFFu,
+		0xFFFFFFFFu,
+		0xFFFFFFFFu,
+		0xFFFFFFFFu,
+		0xFFFFFFFFu,
+		0xFFFFFFFFu,
+		0xFFFFFFFFu,
+		0xFFFFFFFFu);
+
+	ASSERT_TRUE(equal(v1 >= v2, m12));
+	ASSERT_TRUE(equal(v1 >= v3, m13));
+	ASSERT_TRUE(equal(v1 >= v4, m14));
+}
+
+
+TYPED_TEST(T256Test, BitwiseAND) {
+	using F256 = typename TestFixture::T256;
+	F256 v{ 1.0f, 2.0f, 3.0f, 4.0f, 5.0f, 6.0f, 7.0f, 8.0f };
+
+	F256 m1 = frombits<F256>(0x00000000u,
+		0x00000000u,
+		0x00000000u,
+		0x00000000u,
+		0x00000000u,
+		0x00000000u,
+		0x00000000u,
+		0x00000000u);
+	F256 m2 = frombits<F256>(0xFFFFFFFFu,
+		0x00000000u,
+		0x00000000u,
+		0x00000000u,
+		0x00000000u,
+		0x00000000u,
+		0x00000000u,
+		0x00000000u);
+	F256 m3 = frombits<F256>(0x00000000u,
+		0xFFFFFFFFu,
+		0x00000000u,
+		0x00000000u,
+		0x00000000u,
+		0x00000000u,
+		0x00000000u,
+		0x00000000u);
+	F256 m4 = frombits<F256>(0x00000000u,
+		0x00000000u,
+		0xFFFFFFFFu,
+		0x00000000u,
+		0x00000000u,
+		0x00000000u,
+		0x00000000u,
+		0x00000000u);
+	F256 m5 = frombits<F256>(0x00000000u,
+		0x00000000u,
+		0x00000000u,
+		0xFFFFFFFFu,
+		0x00000000u,
+		0x00000000u,
+		0x00000000u,
+		0x00000000u);
+	F256 m6 = frombits<F256>(0x00000000u,
+		0x00000000u,
+		0x00000000u,
+		0x00000000u,
+		0xFFFFFFFFu,
+		0x00000000u,
+		0x00000000u,
+		0x00000000u);
+	F256 m7 = frombits<F256>(0x00000000u,
+		0x00000000u,
+		0x00000000u,
+		0x00000000u,
+		0x00000000u,
+		0xFFFFFFFFu,
+		0x00000000u,
+		0x00000000u);
+	F256 m8 = frombits<F256>(0x00000000u,
+		0x00000000u,
+		0x00000000u,
+		0x00000000u,
+		0x00000000u,
+		0x00000000u,
+		0xFFFFFFFFu,
+		0x00000000u);
+	F256 m9 = frombits<F256>(0x00000000u,
+		0x00000000u,
+		0x00000000u,
+		0x00000000u,
+		0x00000000u,
+		0x00000000u,
+		0x00000000u,
+		0xFFFFFFFFu);
+
+	F256 vm1{ 0.0f, 0.0f, 0.0f, 0.0f, 0.0f, 0.0f, 0.0f, 0.0f };
+	F256 vm2{ 1.0f, 0.0f, 0.0f, 0.0f, 0.0f, 0.0f, 0.0f, 0.0f };
+	F256 vm3{ 0.0f, 2.0f, 0.0f, 0.0f, 0.0f, 0.0f, 0.0f, 0.0f };
+	F256 vm4{ 0.0f, 0.0f, 3.0f, 0.0f, 0.0f, 0.0f, 0.0f, 0.0f };
+	F256 vm5{ 0.0f, 0.0f, 0.0f, 4.0f, 0.0f, 0.0f, 0.0f, 0.0f };
+	F256 vm6{ 0.0f, 0.0f, 0.0f, 0.0f, 5.0f, 0.0f, 0.0f, 0.0f };
+	F256 vm7{ 0.0f, 0.0f, 0.0f, 0.0f, 0.0f, 6.0f, 0.0f, 0.0f };
+	F256 vm8{ 0.0f, 0.0f, 0.0f, 0.0f, 0.0f, 0.0f, 7.0f, 0.0f };
+	F256 vm9{ 0.0f, 0.0f, 0.0f, 0.0f, 0.0f, 0.0f, 0.0f, 8.0f };
+
+	ASSERT_TRUE(equal(v & m1, vm1));
+	ASSERT_TRUE(equal(v & m2, vm2));
+	ASSERT_TRUE(equal(v & m3, vm3));
+	ASSERT_TRUE(equal(v & m4, vm4));
+	ASSERT_TRUE(equal(v & m5, vm5));
+	ASSERT_TRUE(equal(v & m6, vm6));
+	ASSERT_TRUE(equal(v & m7, vm7));
+	ASSERT_TRUE(equal(v & m8, vm8));
+	ASSERT_TRUE(equal(v & m9, vm9));
+}
+
+TYPED_TEST(T256Test, BitwiseOR) {
+	using F256 = typename TestFixture::T256;
+	F256 v1{ 0.0f, 2.0f, 0.0f, 4.0f, 0.0f, 6.0f, 0.0f, 8.0f };
+	F256 v2{ 1.0f, 0.0f, 3.0f, 0.0f, 5.0f, 0.0f, 7.0f, 0.0f };
+
+	F256 v12{ 1.0f, 2.0f, 3.0f, 4.0f, 5.0f, 6.0f, 7.0f, 8.0f };
+
+	ASSERT_TRUE(equal(v1 | v2, v12));
+}
+
+TYPED_TEST(T256Test, BitwiseXOR) {
+	using F256 = typename TestFixture::T256;
+	F256 v1{ 0.0f, 2.0f, 0.0f, 4.0f, 0.0f, 6.0f, 0.0f, 8.0f };
+	F256 v2{ 0.0f, 0.0f, 3.0f, 0.0f, 5.0f, 0.0f, 7.0f, 8.0f };
+
+	F256 v12{ 0.0f, 2.0f, 3.0f, 4.0f, 5.0f, 6.0f, 7.0f, 0.0f };
+
+	ASSERT_TRUE(equal(v1 ^ v2, v12));
+}
+
+TYPED_TEST(T256Test, BitwiseNOT) {
+	using F256 = typename TestFixture::T256;
+	F256 v1 = frombits<F256>(0x00000000u,
+		0x00000000u,
+		0x00000000u,
+		0x00000000u,
+		0x00000000u,
+		0x00000000u,
+		0x00000000u,
+		0x00000000u);
+	F256 v2 = frombits<F256>(0xFFFFFFFFu,
+		0xFFFFFFFFu,
+		0xFFFFFFFFu,
+		0xFFFFFFFFu,
+		0xFFFFFFFFu,
+		0xFFFFFFFFu,
+		0xFFFFFFFFu,
+		0xFFFFFFFFu);
+
+	ASSERT_TRUE(equal(~v1, v2));
+	ASSERT_TRUE(equal(~v2, v1));
 }
 
 TYPED_TEST(T256Test, ConstructFromArgs) {
-    typename TestFixture::T256 v{1.0f, 2.0f, 3.0f, 4.0f, 5.0f, 6.0f, 7.0f, 8.0f};
-    typename TestFixture::T256 expected{1.0f, 2.0f, 3.0f, 4.0f, 5.0f, 6.0f, 7.0f, 8.0f};
-    ASSERT_EQ(v, expected);
+	typename TestFixture::T256 v{ 1.0f, 2.0f, 3.0f, 4.0f, 5.0f, 6.0f, 7.0f, 8.0f };
+	typename TestFixture::T256 expected{ 1.0f, 2.0f, 3.0f, 4.0f, 5.0f, 6.0f, 7.0f, 8.0f };
+	ASSERT_TRUE(equal(v, expected));
 }
 
 TYPED_TEST(T256Test, ConstructFromSingleValue) {
-    typename TestFixture::T256 v{42.0f};
-    typename TestFixture::T256 expected{42.0f, 42.0f, 42.0f, 42.0f, 42.0f, 42.0f, 42.0f, 42.0f};
-    ASSERT_EQ(v, expected);
+	typename TestFixture::T256 v{ 42.0f };
+	typename TestFixture::T256 expected{ 42.0f, 42.0f, 42.0f, 42.0f, 42.0f, 42.0f, 42.0f, 42.0f };
+	ASSERT_TRUE(equal(v, expected));
 }
 
 TYPED_TEST(T256Test, FromAlignedSource) {
-    alignas(TestFixture::T256::alignment()) const float expected[] = {1.0f, 2.0f, 3.0f, 4.0f, 5.0f, 6.0f, 7.0f, 8.0f};
-    auto v = TestFixture::T256::fromAlignedSource(expected);
+	alignas(TestFixture::T256::alignment()) const float expected[] = { 1.0f, 2.0f, 3.0f, 4.0f, 5.0f, 6.0f, 7.0f, 8.0f };
+	auto v = TestFixture::T256::fromAlignedSource(expected);
 
-    alignas(TestFixture::T256::alignment()) float result[TestFixture::T256::size()];
-    v.alignedStore(result);
+	alignas(TestFixture::T256::alignment()) float result[TestFixture::T256::size()];
+	v.alignedStore(result);
 
-    ASSERT_ELEMENTS_EQ(result, expected, TestFixture::T256::size());
+	ASSERT_ELEMENTS_EQ(result, expected, TestFixture::T256::size());
 }
 
 TYPED_TEST(T256Test, AlignedLoadStore) {
-    alignas(TestFixture::T256::alignment()) const float expected[] = {1.0f, 2.0f, 3.0f, 4.0f, 5.0f, 6.0f, 7.0f, 8.0f};
-    typename TestFixture::T256 v;
-    v.alignedLoad(expected);
+	alignas(TestFixture::T256::alignment()) const float expected[] = { 1.0f, 2.0f, 3.0f, 4.0f, 5.0f, 6.0f, 7.0f, 8.0f };
+	typename TestFixture::T256 v;
+	v.alignedLoad(expected);
 
-    alignas(TestFixture::T256::alignment()) float result[TestFixture::T256::size()];
-    v.alignedStore(result);
+	alignas(TestFixture::T256::alignment()) float result[TestFixture::T256::size()];
+	v.alignedStore(result);
 
-    ASSERT_ELEMENTS_EQ(result, expected, TestFixture::T256::size());
+	ASSERT_ELEMENTS_EQ(result, expected, TestFixture::T256::size());
 }
 
 TYPED_TEST(T256Test, FromUnalignedSource) {
-    const float expected[] = {1.0f, 2.0f, 3.0f, 4.0f, 5.0f, 6.0f, 7.0f, 8.0f};
-    auto v = TestFixture::T256::fromUnalignedSource(expected);
+	const float expected[] = { 1.0f, 2.0f, 3.0f, 4.0f, 5.0f, 6.0f, 7.0f, 8.0f };
+	auto v = TestFixture::T256::fromUnalignedSource(expected);
 
-    float result[TestFixture::T256::size()];
-    v.unalignedStore(result);
+	float result[TestFixture::T256::size()];
+	v.unalignedStore(result);
 
-    ASSERT_ELEMENTS_EQ(result, expected, TestFixture::T256::size());
+	ASSERT_ELEMENTS_EQ(result, expected, TestFixture::T256::size());
 }
 
 TYPED_TEST(T256Test, UnalignedLoadStore) {
-    const float expected[] = {1.0f, 2.0f, 3.0f, 4.0f, 5.0f, 6.0f, 7.0f, 8.0f};
-    typename TestFixture::T256 v;
-    v.unalignedLoad(expected);
+	const float expected[] = { 1.0f, 2.0f, 3.0f, 4.0f, 5.0f, 6.0f, 7.0f, 8.0f };
+	typename TestFixture::T256 v;
+	v.unalignedLoad(expected);
 
-    float result[TestFixture::T256::size()];
-    v.unalignedStore(result);
+	float result[TestFixture::T256::size()];
+	v.unalignedStore(result);
 
-    ASSERT_ELEMENTS_EQ(result, expected, TestFixture::T256::size());
+	ASSERT_ELEMENTS_EQ(result, expected, TestFixture::T256::size());
 }
 
 TYPED_TEST(T256Test, LoadSingleValue) {
-    const float source[] = {42.0f, 43.0f, 44.0f, 45.0f, 46.0f, 47.0f, 48.0f, 49.0f};
-    auto v = TestFixture::T256::loadSingleValue(source);
-    typename TestFixture::T256 expected{42.0f, 0.0f, 0.0f, 0.0f, 0.0f, 0.0f, 0.0f, 0.0f};
-    ASSERT_EQ(v, expected);
+	const float source[] = { 42.0f, 43.0f, 44.0f, 45.0f, 46.0f, 47.0f, 48.0f, 49.0f };
+	auto v = TestFixture::T256::loadSingleValue(source);
+	typename TestFixture::T256 expected{ 42.0f, 0.0f, 0.0f, 0.0f, 0.0f, 0.0f, 0.0f, 0.0f };
+	ASSERT_TRUE(equal(v, expected));
 }
 
 TYPED_TEST(T256Test, Sum) {
-    typename TestFixture::T256 v{1.0f, 2.0f, 3.0f, 4.0f, 5.0f, 6.0f, 7.0f, 8.0f};
-    ASSERT_EQ(v.sum(), 36.0f);
+	typename TestFixture::T256 v{ 1.0f, 2.0f, 3.0f, 4.0f, 5.0f, 6.0f, 7.0f, 8.0f };
+	ASSERT_EQ(v.sum(), 36.0f);
 }
 
 TYPED_TEST(T256Test, CompoundAssignmentAdd) {
-    typename TestFixture::T256 v1{1.0f, 2.0f, 3.0f, 4.0f, 5.0f, 6.0f, 7.0f, 8.0f};
-    typename TestFixture::T256 v2{3.0f, 4.0f, 5.0f, 6.0f, 7.0f, 8.0f, 9.0f, 10.0f};
-    typename TestFixture::T256 expected{4.0f, 6.0f, 8.0f, 10.0f, 12.0f, 14.0f, 16.0f, 18.0f};
-    v1 += v2;
-    ASSERT_EQ(v1, expected);
+	typename TestFixture::T256 v1{ 1.0f, 2.0f, 3.0f, 4.0f, 5.0f, 6.0f, 7.0f, 8.0f };
+	typename TestFixture::T256 v2{ 3.0f, 4.0f, 5.0f, 6.0f, 7.0f, 8.0f, 9.0f, 10.0f };
+	typename TestFixture::T256 expected{ 4.0f, 6.0f, 8.0f, 10.0f, 12.0f, 14.0f, 16.0f, 18.0f };
+	v1 += v2;
+	ASSERT_TRUE(equal(v1, expected));
 }
 
 TYPED_TEST(T256Test, CompoundAssignmentSub) {
-    typename TestFixture::T256 v1{1.0f, 2.0f, 3.0f, 4.0f, 5.0f, 6.0f, 7.0f, 8.0f};
-    typename TestFixture::T256 v2{3.0f, 4.0f, 5.0f, 6.0f, 7.0f, 8.0f, 9.0f, 10.0f};
-    typename TestFixture::T256 expected{-2.0f, -2.0f, -2.0f, -2.0f, -2.0f, -2.0f, -2.0f, -2.0f};
-    v1 -= v2;
-    ASSERT_EQ(v1, expected);
+	typename TestFixture::T256 v1{ 1.0f, 2.0f, 3.0f, 4.0f, 5.0f, 6.0f, 7.0f, 8.0f };
+	typename TestFixture::T256 v2{ 3.0f, 4.0f, 5.0f, 6.0f, 7.0f, 8.0f, 9.0f, 10.0f };
+	typename TestFixture::T256 expected{ -2.0f, -2.0f, -2.0f, -2.0f, -2.0f, -2.0f, -2.0f, -2.0f };
+	v1 -= v2;
+	ASSERT_TRUE(equal(v1, expected));
 }
 
 TYPED_TEST(T256Test, CompoundAssignmentMul) {
-    typename TestFixture::T256 v1{1.0f, 2.0f, 3.0f, 4.0f, 5.0f, 6.0f, 7.0f, 8.0f};
-    typename TestFixture::T256 v2{3.0f, 4.0f, 5.0f, 6.0f, 7.0f, 8.0f, 9.0f, 10.0f};
-    typename TestFixture::T256 expected{3.0f, 8.0f, 15.0f, 24.0f, 35.0f, 48.0f, 63.0f, 80.0f};
-    v1 *= v2;
-    ASSERT_EQ(v1, expected);
+	typename TestFixture::T256 v1{ 1.0f, 2.0f, 3.0f, 4.0f, 5.0f, 6.0f, 7.0f, 8.0f };
+	typename TestFixture::T256 v2{ 3.0f, 4.0f, 5.0f, 6.0f, 7.0f, 8.0f, 9.0f, 10.0f };
+	typename TestFixture::T256 expected{ 3.0f, 8.0f, 15.0f, 24.0f, 35.0f, 48.0f, 63.0f, 80.0f };
+	v1 *= v2;
+	ASSERT_TRUE(equal(v1, expected));
 }
 
 TYPED_TEST(T256Test, CompoundAssignmentDiv) {
-    typename TestFixture::T256 v1{4.0f, 3.0f, 9.0f, 12.0f, 4.0f, 3.0f, 9.0f, 12.0f};
-    typename TestFixture::T256 v2{1.0f, 2.0f, 3.0f, 3.0f, 1.0f, 2.0f, 3.0f, 3.0f};
-    float expected[TestFixture::T256::size()] = {4.0f, 1.5f, 3.0f, 4.0f, 4.0f, 1.5f, 3.0f, 4.0f};
-    v1 /= v2;
+	typename TestFixture::T256 v1{ 4.0f, 3.0f, 9.0f, 12.0f, 4.0f, 3.0f, 9.0f, 12.0f };
+	typename TestFixture::T256 v2{ 1.0f, 2.0f, 3.0f, 3.0f, 1.0f, 2.0f, 3.0f, 3.0f };
+	float expected[TestFixture::T256::size()] = { 4.0f, 1.5f, 3.0f, 4.0f, 4.0f, 1.5f, 3.0f, 4.0f };
+	v1 /= v2;
 
-    float result[TestFixture::T256::size()];
-    v1.unalignedStore(result);
+	float result[TestFixture::T256::size()];
+	v1.unalignedStore(result);
 
-    ASSERT_ELEMENTS_NEAR(result, expected, TestFixture::T256::size(), 0.0001f);
+	ASSERT_ELEMENTS_NEAR(result, expected, TestFixture::T256::size(), 0.0001f);
 }
 
 TYPED_TEST(T256Test, OperatorAdd) {
-    typename TestFixture::T256 v1{1.0f, 2.0f, 3.0f, 4.0f, 5.0f, 6.0f, 7.0f, 8.0f};
-    typename TestFixture::T256 v2{3.0f, 4.0f, 5.0f, 6.0f, 7.0f, 8.0f, 9.0f, 10.0f};
-    typename TestFixture::T256 expected{4.0f, 6.0f, 8.0f, 10.0f, 12.0f, 14.0f, 16.0f, 18.0f};
-    auto v3 = v1 + v2;
-    ASSERT_EQ(v3, expected);
+	typename TestFixture::T256 v1{ 1.0f, 2.0f, 3.0f, 4.0f, 5.0f, 6.0f, 7.0f, 8.0f };
+	typename TestFixture::T256 v2{ 3.0f, 4.0f, 5.0f, 6.0f, 7.0f, 8.0f, 9.0f, 10.0f };
+	typename TestFixture::T256 expected{ 4.0f, 6.0f, 8.0f, 10.0f, 12.0f, 14.0f, 16.0f, 18.0f };
+	auto v3 = v1 + v2;
+	ASSERT_TRUE(equal(v3, expected));
 }
 
 TYPED_TEST(T256Test, OperatorSub) {
-    typename TestFixture::T256 v1{1.0f, 2.0f, 3.0f, 4.0f, 5.0f, 6.0f, 7.0f, 8.0f};
-    typename TestFixture::T256 v2{3.0f, 4.0f, 5.0f, 6.0f, 7.0f, 8.0f, 9.0f, 10.0f};
-    typename TestFixture::T256 expected{-2.0f, -2.0f, -2.0f, -2.0f, -2.0f, -2.0f, -2.0f, -2.0f};
-    auto v3 = v1 - v2;
-    ASSERT_EQ(v3, expected);
+	typename TestFixture::T256 v1{ 1.0f, 2.0f, 3.0f, 4.0f, 5.0f, 6.0f, 7.0f, 8.0f };
+	typename TestFixture::T256 v2{ 3.0f, 4.0f, 5.0f, 6.0f, 7.0f, 8.0f, 9.0f, 10.0f };
+	typename TestFixture::T256 expected{ -2.0f, -2.0f, -2.0f, -2.0f, -2.0f, -2.0f, -2.0f, -2.0f };
+	auto v3 = v1 - v2;
+	ASSERT_TRUE(equal(v3, expected));
 }
 
 TYPED_TEST(T256Test, OperatorMul) {
-    typename TestFixture::T256 v1{1.0f, 2.0f, 3.0f, 4.0f, 5.0f, 6.0f, 7.0f, 8.0f};
-    typename TestFixture::T256 v2{3.0f, 4.0f, 5.0f, 6.0f, 7.0f, 8.0f, 9.0f, 10.0f};
-    typename TestFixture::T256 expected{3.0f, 8.0f, 15.0f, 24.0f, 35.0f, 48.0f, 63.0f, 80.0f};
-    auto v3 = v1 * v2;
-    ASSERT_EQ(v3, expected);
+	typename TestFixture::T256 v1{ 1.0f, 2.0f, 3.0f, 4.0f, 5.0f, 6.0f, 7.0f, 8.0f };
+	typename TestFixture::T256 v2{ 3.0f, 4.0f, 5.0f, 6.0f, 7.0f, 8.0f, 9.0f, 10.0f };
+	typename TestFixture::T256 expected{ 3.0f, 8.0f, 15.0f, 24.0f, 35.0f, 48.0f, 63.0f, 80.0f };
+	auto v3 = v1 * v2;
+	ASSERT_TRUE(equal(v3, expected));
 }
 
 TYPED_TEST(T256Test, OperatorDiv) {
-    typename TestFixture::T256 v1{4.0f, 3.0f, 9.0f, 12.0f, 4.0f, 3.0f, 9.0f, 12.0f};
-    typename TestFixture::T256 v2{1.0f, 2.0f, 3.0f, 3.0f, 1.0f, 2.0f, 3.0f, 3.0f};
-    float expected[TestFixture::T256::size()] = {4.0f, 1.5f, 3.0f, 4.0f, 4.0f, 1.5f, 3.0f, 4.0f};
-    auto v3 = v1 / v2;
+	typename TestFixture::T256 v1{ 4.0f, 3.0f, 9.0f, 12.0f, 4.0f, 3.0f, 9.0f, 12.0f };
+	typename TestFixture::T256 v2{ 1.0f, 2.0f, 3.0f, 3.0f, 1.0f, 2.0f, 3.0f, 3.0f };
+	float expected[TestFixture::T256::size()] = { 4.0f, 1.5f, 3.0f, 4.0f, 4.0f, 1.5f, 3.0f, 4.0f };
+	auto v3 = v1 / v2;
 
-    float result[TestFixture::T256::size()];
-    v3.unalignedStore(result);
+	float result[TestFixture::T256::size()];
+	v3.unalignedStore(result);
 
-    ASSERT_ELEMENTS_NEAR(result, expected, TestFixture::T256::size(), 0.0001f);
+	ASSERT_ELEMENTS_NEAR(result, expected, TestFixture::T256::size(), 0.0001f);
 }
 
 TEST(T256Test, TransposeSquareScalar) {
-    trimd::scalar::F256 v1{1.0f, 2.0f, 3.0f, 4.0f, 5.0f, 6.0f, 7.0f, 8.0f};
-    trimd::scalar::F256 v2{1.0f, 2.0f, 3.0f, 4.0f, 5.0f, 6.0f, 7.0f, 8.0f};
-    trimd::scalar::F256 v3{1.0f, 2.0f, 3.0f, 4.0f, 5.0f, 6.0f, 7.0f, 8.0f};
-    trimd::scalar::F256 v4{1.0f, 2.0f, 3.0f, 4.0f, 5.0f, 6.0f, 7.0f, 8.0f};
-    trimd::scalar::F256 v5{1.0f, 2.0f, 3.0f, 4.0f, 5.0f, 6.0f, 7.0f, 8.0f};
-    trimd::scalar::F256 v6{1.0f, 2.0f, 3.0f, 4.0f, 5.0f, 6.0f, 7.0f, 8.0f};
-    trimd::scalar::F256 v7{1.0f, 2.0f, 3.0f, 4.0f, 5.0f, 6.0f, 7.0f, 8.0f};
-    trimd::scalar::F256 v8{1.0f, 2.0f, 3.0f, 4.0f, 5.0f, 6.0f, 7.0f, 8.0f};
-
-    trimd::scalar::transpose(v1, v2, v3, v4, v5, v6, v7, v8);
-
-    trimd::scalar::F256 e1{1.0f, 1.0f, 1.0f, 1.0f, 1.0f, 1.0f, 1.0f, 1.0f};
-    trimd::scalar::F256 e2{2.0f, 2.0f, 2.0f, 2.0f, 2.0f, 2.0f, 2.0f, 2.0f};
-    trimd::scalar::F256 e3{3.0f, 3.0f, 3.0f, 3.0f, 3.0f, 3.0f, 3.0f, 3.0f};
-    trimd::scalar::F256 e4{4.0f, 4.0f, 4.0f, 4.0f, 4.0f, 4.0f, 4.0f, 4.0f};
-    trimd::scalar::F256 e5{5.0f, 5.0f, 5.0f, 5.0f, 5.0f, 5.0f, 5.0f, 5.0f};
-    trimd::scalar::F256 e6{6.0f, 6.0f, 6.0f, 6.0f, 6.0f, 6.0f, 6.0f, 6.0f};
-    trimd::scalar::F256 e7{7.0f, 7.0f, 7.0f, 7.0f, 7.0f, 7.0f, 7.0f, 7.0f};
-    trimd::scalar::F256 e8{8.0f, 8.0f, 8.0f, 8.0f, 8.0f, 8.0f, 8.0f, 8.0f};
-
-    ASSERT_EQ(v1, e1);
-    ASSERT_EQ(v2, e2);
-    ASSERT_EQ(v3, e3);
-    ASSERT_EQ(v4, e4);
-    ASSERT_EQ(v5, e5);
-    ASSERT_EQ(v6, e6);
-    ASSERT_EQ(v7, e7);
-    ASSERT_EQ(v8, e8);
+	trimd::scalar::F256 v1{ 1.0f, 2.0f, 3.0f, 4.0f, 5.0f, 6.0f, 7.0f, 8.0f };
+	trimd::scalar::F256 v2{ 1.0f, 2.0f, 3.0f, 4.0f, 5.0f, 6.0f, 7.0f, 8.0f };
+	trimd::scalar::F256 v3{ 1.0f, 2.0f, 3.0f, 4.0f, 5.0f, 6.0f, 7.0f, 8.0f };
+	trimd::scalar::F256 v4{ 1.0f, 2.0f, 3.0f, 4.0f, 5.0f, 6.0f, 7.0f, 8.0f };
+	trimd::scalar::F256 v5{ 1.0f, 2.0f, 3.0f, 4.0f, 5.0f, 6.0f, 7.0f, 8.0f };
+	trimd::scalar::F256 v6{ 1.0f, 2.0f, 3.0f, 4.0f, 5.0f, 6.0f, 7.0f, 8.0f };
+	trimd::scalar::F256 v7{ 1.0f, 2.0f, 3.0f, 4.0f, 5.0f, 6.0f, 7.0f, 8.0f };
+	trimd::scalar::F256 v8{ 1.0f, 2.0f, 3.0f, 4.0f, 5.0f, 6.0f, 7.0f, 8.0f };
+
+	trimd::scalar::transpose(v1, v2, v3, v4, v5, v6, v7, v8);
+
+	trimd::scalar::F256 e1{ 1.0f, 1.0f, 1.0f, 1.0f, 1.0f, 1.0f, 1.0f, 1.0f };
+	trimd::scalar::F256 e2{ 2.0f, 2.0f, 2.0f, 2.0f, 2.0f, 2.0f, 2.0f, 2.0f };
+	trimd::scalar::F256 e3{ 3.0f, 3.0f, 3.0f, 3.0f, 3.0f, 3.0f, 3.0f, 3.0f };
+	trimd::scalar::F256 e4{ 4.0f, 4.0f, 4.0f, 4.0f, 4.0f, 4.0f, 4.0f, 4.0f };
+	trimd::scalar::F256 e5{ 5.0f, 5.0f, 5.0f, 5.0f, 5.0f, 5.0f, 5.0f, 5.0f };
+	trimd::scalar::F256 e6{ 6.0f, 6.0f, 6.0f, 6.0f, 6.0f, 6.0f, 6.0f, 6.0f };
+	trimd::scalar::F256 e7{ 7.0f, 7.0f, 7.0f, 7.0f, 7.0f, 7.0f, 7.0f, 7.0f };
+	trimd::scalar::F256 e8{ 8.0f, 8.0f, 8.0f, 8.0f, 8.0f, 8.0f, 8.0f, 8.0f };
+
+	ASSERT_TRUE(equal(v1, e1));
+	ASSERT_TRUE(equal(v2, e2));
+	ASSERT_TRUE(equal(v3, e3));
+	ASSERT_TRUE(equal(v4, e4));
+	ASSERT_TRUE(equal(v5, e5));
+	ASSERT_TRUE(equal(v6, e6));
+	ASSERT_TRUE(equal(v7, e7));
+	ASSERT_TRUE(equal(v8, e8));
+}
+
+TEST(T256Test, AbsScalar) {
+	trimd::scalar::F256 v{ -1.0f, 2.0f, -3.0f, 0.0f, -1.0f, 2.0f, -3.0f, 0.0f };
+	v = trimd::scalar::abs(v);
+	trimd::scalar::F256 e{ 1.0f, 2.0f, 3.0f, 0.0f, 1.0f, 2.0f, 3.0f, 0.0f };
+	ASSERT_TRUE(equal(v, e));
+}
+
+TEST(T256Test, AndNotScalar) {
+	trimd::scalar::F256 v{ 1.0f, 2.0f, 3.0f, 4.0f, 5.0f, 6.0f, 7.0f, 8.0f };
+	trimd::scalar::F256 mask1 = frombits<trimd::scalar::F256>(0x00000000u,
+		0x00000000u,
+		0x00000000u,
+		0x00000000u,
+		0x00000000u,
+		0x00000000u,
+		0x00000000u,
+		0x00000000u);
+	trimd::scalar::F256 mask2 = frombits<trimd::scalar::F256>(0xFFFFFFFFu,
+		0xFFFFFFFFu,
+		0xFFFFFFFFu,
+		0xFFFFFFFFu,
+		0xFFFFFFFFu,
+		0xFFFFFFFFu,
+		0xFFFFFFFFu,
+		0xFFFFFFFFu);
+	trimd::scalar::F256 result1 = trimd::scalar::andnot(mask1, v);
+	trimd::scalar::F256 result2 = trimd::scalar::andnot(mask2, v);
+	trimd::scalar::F256 e1{ 1.0f, 2.0f, 3.0f, 4.0f, 5.0f, 6.0f, 7.0f, 8.0f };
+	trimd::scalar::F256 e2{ 0.0f, 0.0f, 0.0f, 0.0f, 0.0f, 0.0f, 0.0f, 0.0f };
+	ASSERT_TRUE(equal(result1, e1));
+	ASSERT_TRUE(equal(result2, e2));
 }
 
 #ifdef TRIMD_ENABLE_AVX
 TEST(T256Test, TransposeSquareAVX) {
-    trimd::avx::F256 v1{1.0f, 2.0f, 3.0f, 4.0f, 5.0f, 6.0f, 7.0f, 8.0f};
-    trimd::avx::F256 v2{1.0f, 2.0f, 3.0f, 4.0f, 5.0f, 6.0f, 7.0f, 8.0f};
-    trimd::avx::F256 v3{1.0f, 2.0f, 3.0f, 4.0f, 5.0f, 6.0f, 7.0f, 8.0f};
-    trimd::avx::F256 v4{1.0f, 2.0f, 3.0f, 4.0f, 5.0f, 6.0f, 7.0f, 8.0f};
-    trimd::avx::F256 v5{1.0f, 2.0f, 3.0f, 4.0f, 5.0f, 6.0f, 7.0f, 8.0f};
-    trimd::avx::F256 v6{1.0f, 2.0f, 3.0f, 4.0f, 5.0f, 6.0f, 7.0f, 8.0f};
-    trimd::avx::F256 v7{1.0f, 2.0f, 3.0f, 4.0f, 5.0f, 6.0f, 7.0f, 8.0f};
-    trimd::avx::F256 v8{1.0f, 2.0f, 3.0f, 4.0f, 5.0f, 6.0f, 7.0f, 8.0f};
-
-    trimd::avx::transpose(v1, v2, v3, v4, v5, v6, v7, v8);
-
-    trimd::avx::F256 e1{1.0f, 1.0f, 1.0f, 1.0f, 1.0f, 1.0f, 1.0f, 1.0f};
-    trimd::avx::F256 e2{2.0f, 2.0f, 2.0f, 2.0f, 2.0f, 2.0f, 2.0f, 2.0f};
-    trimd::avx::F256 e3{3.0f, 3.0f, 3.0f, 3.0f, 3.0f, 3.0f, 3.0f, 3.0f};
-    trimd::avx::F256 e4{4.0f, 4.0f, 4.0f, 4.0f, 4.0f, 4.0f, 4.0f, 4.0f};
-    trimd::avx::F256 e5{5.0f, 5.0f, 5.0f, 5.0f, 5.0f, 5.0f, 5.0f, 5.0f};
-    trimd::avx::F256 e6{6.0f, 6.0f, 6.0f, 6.0f, 6.0f, 6.0f, 6.0f, 6.0f};
-    trimd::avx::F256 e7{7.0f, 7.0f, 7.0f, 7.0f, 7.0f, 7.0f, 7.0f, 7.0f};
-    trimd::avx::F256 e8{8.0f, 8.0f, 8.0f, 8.0f, 8.0f, 8.0f, 8.0f, 8.0f};
-
-    ASSERT_EQ(v1, e1);
-    ASSERT_EQ(v2, e2);
-    ASSERT_EQ(v3, e3);
-    ASSERT_EQ(v4, e4);
-    ASSERT_EQ(v5, e5);
-    ASSERT_EQ(v6, e6);
-    ASSERT_EQ(v7, e7);
-    ASSERT_EQ(v8, e8);
+	trimd::avx::F256 v1{ 1.0f, 2.0f, 3.0f, 4.0f, 5.0f, 6.0f, 7.0f, 8.0f };
+	trimd::avx::F256 v2{ 1.0f, 2.0f, 3.0f, 4.0f, 5.0f, 6.0f, 7.0f, 8.0f };
+	trimd::avx::F256 v3{ 1.0f, 2.0f, 3.0f, 4.0f, 5.0f, 6.0f, 7.0f, 8.0f };
+	trimd::avx::F256 v4{ 1.0f, 2.0f, 3.0f, 4.0f, 5.0f, 6.0f, 7.0f, 8.0f };
+	trimd::avx::F256 v5{ 1.0f, 2.0f, 3.0f, 4.0f, 5.0f, 6.0f, 7.0f, 8.0f };
+	trimd::avx::F256 v6{ 1.0f, 2.0f, 3.0f, 4.0f, 5.0f, 6.0f, 7.0f, 8.0f };
+	trimd::avx::F256 v7{ 1.0f, 2.0f, 3.0f, 4.0f, 5.0f, 6.0f, 7.0f, 8.0f };
+	trimd::avx::F256 v8{ 1.0f, 2.0f, 3.0f, 4.0f, 5.0f, 6.0f, 7.0f, 8.0f };
+
+	trimd::avx::transpose(v1, v2, v3, v4, v5, v6, v7, v8);
+
+	trimd::avx::F256 e1{ 1.0f, 1.0f, 1.0f, 1.0f, 1.0f, 1.0f, 1.0f, 1.0f };
+	trimd::avx::F256 e2{ 2.0f, 2.0f, 2.0f, 2.0f, 2.0f, 2.0f, 2.0f, 2.0f };
+	trimd::avx::F256 e3{ 3.0f, 3.0f, 3.0f, 3.0f, 3.0f, 3.0f, 3.0f, 3.0f };
+	trimd::avx::F256 e4{ 4.0f, 4.0f, 4.0f, 4.0f, 4.0f, 4.0f, 4.0f, 4.0f };
+	trimd::avx::F256 e5{ 5.0f, 5.0f, 5.0f, 5.0f, 5.0f, 5.0f, 5.0f, 5.0f };
+	trimd::avx::F256 e6{ 6.0f, 6.0f, 6.0f, 6.0f, 6.0f, 6.0f, 6.0f, 6.0f };
+	trimd::avx::F256 e7{ 7.0f, 7.0f, 7.0f, 7.0f, 7.0f, 7.0f, 7.0f, 7.0f };
+	trimd::avx::F256 e8{ 8.0f, 8.0f, 8.0f, 8.0f, 8.0f, 8.0f, 8.0f, 8.0f };
+
+	ASSERT_TRUE(equal(v1, e1));
+	ASSERT_TRUE(equal(v2, e2));
+	ASSERT_TRUE(equal(v3, e3));
+	ASSERT_TRUE(equal(v4, e4));
+	ASSERT_TRUE(equal(v5, e5));
+	ASSERT_TRUE(equal(v6, e6));
+	ASSERT_TRUE(equal(v7, e7));
+	ASSERT_TRUE(equal(v8, e8));
+}
+
+TEST(T256Test, AbsAVX) {
+	trimd::avx::F256 v{ -1.0f, 2.0f, -3.0f, 0.0f, -1.0f, 2.0f, -3.0f, 0.0f };
+	v = trimd::avx::abs(v);
+	trimd::avx::F256 e{ 1.0f, 2.0f, 3.0f, 0.0f, 1.0f, 2.0f, 3.0f, 0.0f };
+	ASSERT_TRUE(equal(v, e));
 }
+
+TEST(T256Test, AndNotAVX) {
+	trimd::avx::F256 v{ 1.0f, 2.0f, 3.0f, 4.0f, 5.0f, 6.0f, 7.0f, 8.0f };
+	trimd::avx::F256 mask1 = frombits<trimd::avx::F256>(0x00000000u,
+		0x00000000u,
+		0x00000000u,
+		0x00000000u,
+		0x00000000u,
+		0x00000000u,
+		0x00000000u,
+		0x00000000u);
+	trimd::avx::F256 mask2 = frombits<trimd::avx::F256>(0xFFFFFFFFu,
+		0xFFFFFFFFu,
+		0xFFFFFFFFu,
+		0xFFFFFFFFu,
+		0xFFFFFFFFu,
+		0xFFFFFFFFu,
+		0xFFFFFFFFu,
+		0xFFFFFFFFu);
+	trimd::avx::F256 result1 = trimd::avx::andnot(mask1, v);
+	trimd::avx::F256 result2 = trimd::avx::andnot(mask2, v);
+	trimd::avx::F256 e1{ 1.0f, 2.0f, 3.0f, 4.0f, 5.0f, 6.0f, 7.0f, 8.0f };
+	trimd::avx::F256 e2{ 0.0f, 0.0f, 0.0f, 0.0f, 0.0f, 0.0f, 0.0f, 0.0f };
+	ASSERT_TRUE(equal(result1, e1));
+	ASSERT_TRUE(equal(result2, e2));
+}
+
 #endif  // TRIMD_ENABLE_AVX
+
+#ifdef TRIMD_ENABLE_F16C
+
+#if defined(TRIMD_ENABLE_AVX) && defined(TRIMD_ENABLE_SSE)
+using T256HFTypes = ::testing::Types<trimd::sse::F256, trimd::avx::F256>;
+
+#elif defined(TRIMD_ENABLE_AVX)
+using T256HFTypes = ::testing::Types<trimd::avx::F256>;
+
+#elif defined(TRIMD_ENABLE_SSE)
+using T256HFTypes = ::testing::Types<trimd::sse::F256>;
+
+#else
+using T256HFTypes = ::testing::Types<>;
+#endif  // TRIMD_ENABLE_SSE
+
+template<typename T>
+class T256HFTest : public ::testing::Test {
+protected:
+	using T256 = T;
+
+};
+
+TYPED_TEST_SUITE(T256HFTest, T256HFTypes, );
+
+TYPED_TEST(T256HFTest, LoadAlignedHalfFloats) {
+	using F256 = typename TestFixture::T256;
+	alignas(F256::alignment()) const std::uint16_t halfFloats[] = { 15360, 16384, 16896, 17408, 17664, 17920, 18176, 18432 };
+	F256 expected{ 1.0f, 2.0f, 3.0f, 4.0f, 5.0f, 6.0f, 7.0f, 8.0f };
+
+	auto v = F256::fromAlignedSource(halfFloats);
+	ASSERT_TRUE(equal(v, expected));
+
+	F256 v2;
+	v2.alignedLoad(halfFloats);
+	ASSERT_TRUE(equal(v2, expected));
+}
+
+TYPED_TEST(T256HFTest, LoadUnalignedHalfFloats) {
+	using F256 = typename TestFixture::T256;
+	const std::uint16_t halfFloats[] = { 15360, 16384, 16896, 17408, 17664, 17920, 18176, 18432 };
+	F256 expected{ 1.0f, 2.0f, 3.0f, 4.0f, 5.0f, 6.0f, 7.0f, 8.0f };
+
+	auto v = F256::fromUnalignedSource(halfFloats);
+	ASSERT_TRUE(equal(v, expected));
+
+	F256 v2;
+	v2.unalignedLoad(halfFloats);
+	ASSERT_TRUE(equal(v2, expected));
+}
+
+TYPED_TEST(T256HFTest, StoreAlignedHalfFloats) {
+	using F256 = typename TestFixture::T256;
+	const std::uint16_t expected[] = { 15360, 16384, 16896, 17408, 17664, 17920, 18176, 18432 };
+
+	alignas(F256::alignment()) std::uint16_t halfFloats[8ul] = {};
+	const F256 v{ 1.0f, 2.0f, 3.0f, 4.0f, 5.0f, 6.0f, 7.0f, 8.0f };
+	v.alignedStore(halfFloats);
+
+	ASSERT_ELEMENTS_EQ(halfFloats, expected, 8ul);
+}
+
+TYPED_TEST(T256HFTest, StoreUnalignedHalfFloats) {
+	using F256 = typename TestFixture::T256;
+	const std::uint16_t expected[] = { 15360, 16384, 16896, 17408, 17664, 17920, 18176, 18432 };
+
+	std::uint16_t halfFloats[8ul] = {};
+	const F256 v{ 1.0f, 2.0f, 3.0f, 4.0f, 5.0f, 6.0f, 7.0f, 8.0f };
+	v.unalignedStore(halfFloats);
+
+	ASSERT_ELEMENTS_EQ(halfFloats, expected, 8ul);
+}
+
+#endif  // TRIMD_ENABLE_F16C
