Importing minor changes from //MetaHuman/Main stream

diff --git a/Engine/Plugins/Animation/RigLogic/Source/RigLogicModule/Private/DNAUtils.cpp b/Engine/Plugins/Animation/RigLogic/Source/RigLogicModule/Private/DNAUtils.cpp
--- a/Engine/Plugins/Animation/RigLogic/Source/RigLogicModule/Private/DNAUtils.cpp
+++ b/Engine/Plugins/Animation/RigLogic/Source/RigLogicModule/Private/DNAUtils.cpp
@@ -53,6 +53,19 @@ TSharedPtr<IDNAReader> ReadDNAFromBuffer(TArray<uint8>* DNABuffer, EDNADataLayer
 	return ReadDNAStream(MoveTemp(DNAStreamReader));
 }
 
+TArray<uint8> ReadStreamFromDNA(const IDNAReader* Reader, EDNADataLayer Layer)
+{
+	LLM_SCOPE_BYNAME(TEXT("Animation/RigLogic"));
+	TArray<char> DNABuffer;
+	auto DeltaDnaStream = rl4::makeScoped<trio::MemoryStream>();
+	auto DNAStreamWriter = rl4::makeScoped<dna::StreamWriter>(DeltaDnaStream.get(), FMemoryResource::Instance());
+	DNAStreamWriter->setFrom(Reader->Unwrap(), static_cast<dna::DataLayer>(Layer), FMemoryResource::Instance());
+	DNAStreamWriter->write();
+	DNABuffer.AddZeroed(DeltaDnaStream->size());
+	DeltaDnaStream->read(DNABuffer.GetData(), DeltaDnaStream->size());
+	return TArray<uint8>(DNABuffer);
+}
+
 void WriteDNAToFile(const IDNAReader* Reader, EDNADataLayer Layer, const FString& Path)
 {
 	LLM_SCOPE_BYNAME(TEXT("Animation/RigLogic"));
diff --git a/Engine/Plugins/Animation/RigLogic/Source/RigLogicModule/Public/DNAUtils.h b/Engine/Plugins/Animation/RigLogic/Source/RigLogicModule/Public/DNAUtils.h
--- a/Engine/Plugins/Animation/RigLogic/Source/RigLogicModule/Public/DNAUtils.h
+++ b/Engine/Plugins/Animation/RigLogic/Source/RigLogicModule/Public/DNAUtils.h
@@ -12,4 +12,5 @@ RIGLOGICMODULE_API TSharedPtr<IDNAReader> ReadDNAFromFile(const FString& Path, E
 RIGLOGICMODULE_API TSharedPtr<IDNAReader> ReadDNAFromFile(const FString& Path, EDNADataLayer Layer, TArrayView<uint16_t> LODs);
 RIGLOGICMODULE_API TSharedPtr<IDNAReader> ReadDNAFromBuffer(TArray<uint8>* DNABuffer, EDNADataLayer Layer = EDNADataLayer::All, uint16_t MaxLOD = 0);
 RIGLOGICMODULE_API TSharedPtr<IDNAReader> ReadDNAFromBuffer(TArray<uint8>* DNABuffer, EDNADataLayer Layer, TArrayView<uint16_t> LODs);
+RIGLOGICMODULE_API TArray<uint8> ReadStreamFromDNA(const IDNAReader* Reader, EDNADataLayer Layer);
 RIGLOGICMODULE_API void WriteDNAToFile(const IDNAReader* Reader, EDNADataLayer Layer, const FString& Path);
\ No newline at end of file
