Moved DNAUtils from GeneSplicer and MetaHuman plugin to Riglogic #rb none #preflight 632781fb121ffddebc2190c8

diff --git a/Engine/Plugins/Animation/RigLogic/Source/RigLogicModule/Private/AnimNode_RigLogic.cpp b/Engine/Plugins/Animation/RigLogic/Source/RigLogicModule/Private/AnimNode_RigLogic.cpp
--- a/Engine/Plugins/Animation/RigLogic/Source/RigLogicModule/Private/AnimNode_RigLogic.cpp
+++ b/Engine/Plugins/Animation/RigLogic/Source/RigLogicModule/Private/AnimNode_RigLogic.cpp
@@ -10,7 +10,7 @@
 #include "SharedRigRuntimeContext.h"
 #include "Animation/AnimInstanceProxy.h"
 #include "HAL/LowLevelMemTracker.h"
-#include "Runtime/Engine/Classes/Animation/MorphTarget.h"
+#include "Animation/MorphTarget.h"
 
 LLM_DEFINE_TAG(Animation_RigLogic);
 DEFINE_LOG_CATEGORY(LogRigLogicAnimNode);
diff --git a/Engine/Plugins/Animation/RigLogic/Source/RigLogicModule/Private/DNAToSkelMeshMap.cpp b/Engine/Plugins/Animation/RigLogic/Source/RigLogicModule/Private/DNAToSkelMeshMap.cpp
new file mode 100644
--- /dev/null
+++ b/Engine/Plugins/Animation/RigLogic/Source/RigLogicModule/Private/DNAToSkelMeshMap.cpp
@@ -0,0 +1,272 @@
+// Copyright Epic Games, Inc. All Rights Reserved.
+
+#include "DNAToSkelMeshMap.h"
+
+#include "DNAAsset.h"
+#include "SkelMeshDNAReader.h"
+
+#include "Serialization/BufferArchive.h"
+#include "Serialization/MemoryReader.h"
+#include "Misc/Paths.h"
+#include "Misc/FileHelper.h"
+#include "UObject/UObjectGlobals.h"
+#include "Rendering/SkeletalMeshModel.h"
+#include "Rendering/SkeletalMeshLODModel.h"
+#include "Animation/SmartName.h"
+#include "Engine/SkeletalMesh.h"
+
+#include "riglogic/RigLogic.h"
+
+DEFINE_LOG_CATEGORY(LogDNAToSkelMeshMap);
+
+#if WITH_EDITORONLY_DATA
+/** Creates mappings between source DNA and target SkelMesh. */
+void FDNAToSkelMeshMap::InitBaseMesh(IDNAReader* InSourceDNAReader, USkeletalMesh* InTargetSkeletalMesh)
+{
+	TargetSkelMesh = InTargetSkeletalMesh; //maps are created for specific SkelMesh, so we can memorize it for clarity
+
+	InitVertexMap(InSourceDNAReader);
+
+	//NOTE: this is just init for updating base mesh
+	// before updating joints and morph targets, MapJoints and MapMorphTarget must be called respectively
+}
+#endif
+
+
+/** Creates and initiates DNAReader and creates mapping between DNA and provided SkelMesh. */
+bool FDNAToSkelMeshMap::InitFromDNAAsset(USkeletalMesh* InSkelMesh)
+{
+#if WITH_EDITORONLY_DATA
+	// Fetch DNAAsset from SkelMesh and use its StreamReader
+	UAssetUserData* UserData = InSkelMesh->GetAssetUserDataOfClass(UDNAAsset::StaticClass());
+	if (UserData == nullptr)
+	{
+		FString PathName = GetPathNameSafe(InSkelMesh);
+		UE_LOG(LogDNAToSkelMeshMap, Warning, TEXT("Could not find DNAAsset user data for %s"), *PathName);
+		return false;
+	}
+	UDNAAsset* DNAAsset = Cast<UDNAAsset>(UserData);
+	FSkelMeshDNAReader SkelMeshDNAReader{ DNAAsset };
+	InitBaseMesh(&SkelMeshDNAReader, InSkelMesh);
+
+	return true;
+#else
+	return false;
+#endif //WITH_EDITORONLY_DATA
+}
+
+#if WITH_EDITORONLY_DATA
+/** Maps vertices and meshes from SkelMesh to DNA. */
+void FDNAToSkelMeshMap::InitVertexMap(IDNAReader* InDNAReader)
+{
+	double StartTime = FPlatformTime::Seconds();
+
+	ImportVtxToDNAMeshIndex.Empty();
+	ImportVtxToDNAVtxIndex.Empty();
+
+	FSkeletalMeshModel* ImportedModel = TargetSkelMesh->GetImportedModel();
+	int32 NumEngineLods = ImportedModel->LODModels.Num();
+	ImportVtxToDNAMeshIndex.AddZeroed(NumEngineLods);
+	ImportVtxToDNAVtxIndex.AddZeroed(NumEngineLods);
+	for (int32 i = 0; i < NumEngineLods; i++)
+	{
+		FSkeletalMeshLODModel& LODModel = ImportedModel->LODModels[i];
+		ImportVtxToDNAMeshIndex[i].AddZeroed(LODModel.MeshToImportVertexMap.Num());
+		ImportVtxToDNAVtxIndex[i].AddZeroed(LODModel.MeshToImportVertexMap.Num());
+	}
+
+	TArray<TArray<int32>> MeshStartIndices; // For each LOD, store the start index offset of each mesh.
+	TArray<int32> TotalVertices; // Counts total indices for each mesh in DNA, per LOD.
+	uint16 NumLODs = InDNAReader->GetLODCount();
+	MeshStartIndices.AddZeroed(NumLODs);
+	TotalVertices.AddZeroed(NumLODs);
+	ImportDNAVtxToUEVtxIndex.AddZeroed(NumLODs);
+	// First find start indices for each mesh and for each LOD in DNA.
+	uint16 TotalMeshCount = InDNAReader->GetMeshCount();
+	for (uint16 LODIndex = 0; LODIndex < NumLODs; LODIndex++)
+	{
+		const TArrayView<const uint16> LODMeshIndices = InDNAReader->GetMeshIndicesForLOD(LODIndex);
+
+		ImportDNAVtxToUEVtxIndex[LODIndex].AddZeroed(TotalMeshCount);
+
+		for (size_t MeshIndex = 0; MeshIndex < LODMeshIndices.Num(); MeshIndex++)
+		{
+			MeshStartIndices[LODIndex].Add(TotalVertices[LODIndex]);
+
+			const int32 VertexCount = InDNAReader->GetVertexPositionCount(LODMeshIndices[MeshIndex]);
+			TotalVertices[LODIndex] += VertexCount;
+			ImportDNAVtxToUEVtxIndex[LODIndex][LODMeshIndices[MeshIndex]].Init(INDEX_NONE, VertexCount);
+		}
+		MeshStartIndices[LODIndex].Add(TotalVertices[LODIndex]);
+	}
+	// Then for each LOD and each Vertex in LODModel Vertex Map find corresponding DNA Vertex.
+	int32 ImportedVertexCount = 0;
+	for (int32 LODIndex = 0; LODIndex < ImportedModel->LODModels.Num(); LODIndex++)
+	{
+		const FSkeletalMeshLODModel& LODModel = ImportedModel->LODModels[LODIndex];
+		const TArrayView<const uint16> LODMeshIndices = InDNAReader->GetMeshIndicesForLOD(LODIndex);
+		const int32 StartIndicesCount = MeshStartIndices[LODIndex].Num();
+
+		int32 LODMeshVtxCount = LODModel.MeshToImportVertexMap.Num();
+		for (int32 LODMeshVtxIndex = 0; LODMeshVtxIndex < LODMeshVtxCount; LODMeshVtxIndex++)
+		{
+			int32 FbxVertexIndex = LODModel.MeshToImportVertexMap[LODMeshVtxIndex];
+			bool bFound = false;
+			// For each LODModel Vertex try to find the corresponding DNA Vertex Index by checking prepared list of Start Indices.
+			for (int k = 1; k < StartIndicesCount; k++)
+			{
+				if (FbxVertexIndex < MeshStartIndices[LODIndex][k])
+				{
+					bFound = true;
+					int32 MeshIndex = LODMeshIndices[k - 1];
+					ImportVtxToDNAMeshIndex[LODIndex][LODMeshVtxIndex] = MeshIndex;
+					const int32 DNAVertexIndex = FbxVertexIndex - MeshStartIndices[LODIndex][k - 1];
+					ImportVtxToDNAVtxIndex[LODIndex][LODMeshVtxIndex] = DNAVertexIndex;
+					ImportDNAVtxToUEVtxIndex[LODIndex][MeshIndex][DNAVertexIndex] = LODMeshVtxIndex;
+					break;
+				}
+			}
+			// If no matching DNA Vertex Index is found mark LODModel Vertex Index with -1.
+			if (!bFound)
+			{
+#ifdef DEBUG
+				UE_LOG(LogDNAToSkelMeshMap, Warning, TEXT("Not sorted fbx vertex found %d"), FbxVertexIndex);
+#endif
+				ImportVtxToDNAMeshIndex[LODIndex][LODMeshVtxIndex] = INDEX_NONE;
+				ImportVtxToDNAVtxIndex[LODIndex][LODMeshVtxIndex] = INDEX_NONE;
+			}
+		}
+	}
+
+	// Find and map overlapping indices and vertex to section indices.
+	OverlappingVertices.Empty(NumEngineLods);
+	OverlappingVertices.AddZeroed(NumEngineLods);
+	UEVertexToSectionIndices.Empty(NumEngineLods);
+	UEVertexToSectionIndices.AddZeroed(NumEngineLods);
+	for (int32 i = 0; i < NumEngineLods; i++)
+	{
+		FSkeletalMeshLODModel& LODModel = ImportedModel->LODModels[i];
+
+		UEVertexToSectionIndices[i].Init(INDEX_NONE, LODModel.NumVertices);
+
+		const int32 NumSections = LODModel.Sections.Num();
+		OverlappingVertices[i].AddZeroed(NumSections);
+		for (int32 SectionIdx = 0; SectionIdx < NumSections; SectionIdx++)
+		{
+			const FSkelMeshSection& Section = LODModel.Sections[SectionIdx];
+			const int32 NumVertices = Section.GetNumVertices();			
+
+			TArray<bool> VerticesCovered;
+			VerticesCovered.AddZeroed(NumVertices);
+
+			OverlappingVertices[i][SectionIdx].AddZeroed(NumVertices);
+			for (int32 VertexIndex = 0; VertexIndex < NumVertices; VertexIndex++)
+			{
+				const int32 VertexBufferIndex = Section.GetVertexBufferIndex() + VertexIndex;
+				UEVertexToSectionIndices[i][VertexBufferIndex] = SectionIdx;
+
+				const int32 DNAVertexIndex = ImportVtxToDNAVtxIndex[i][VertexBufferIndex];
+				if (DNAVertexIndex >= 0)
+				{
+					VerticesCovered[VertexIndex] = true;
+
+					if (Section.OverlappingVertices.Contains(VertexIndex))
+					{
+						const int32 OverlappingCount = Section.OverlappingVertices[VertexIndex].Num();
+						for (int32 OverlappingIndex = 0; OverlappingIndex < OverlappingCount; ++OverlappingIndex)
+						{
+							const int32 OverlappingVertexIndex = Section.OverlappingVertices[VertexIndex][OverlappingIndex];
+							if (!VerticesCovered[OverlappingVertexIndex]) // No need to update the same vertex twice.
+							{
+								OverlappingVertices[i][SectionIdx][VertexIndex].Add(OverlappingVertexIndex);
+							}
+						}
+					}
+				}
+			}
+		}
+	}
+
+	double TimeToMap = FPlatformTime::Seconds();
+	TimeToMap -= StartTime;
+
+	UE_LOG(LogDNAToSkelMeshMap, Log, TEXT("	InitVertexMap:	%.6f"), TimeToMap * 1000.0f);
+}
+
+/** Makes a map of all Joints from DNA to Bones in Reference Skeleton.
+ * Additionally makes a map of all Skinned joints. */
+void FDNAToSkelMeshMap::MapJoints(IDNAReader* InDNAReader)
+{
+
+	double StartTime = FPlatformTime::Seconds();
+
+	FReferenceSkeleton* RefSkeleton = &TargetSkelMesh->GetRefSkeleton();
+	uint32 JointCount = InDNAReader->GetJointCount();
+
+	// Map Joints to Bones.
+	RLJointToUEBoneIndices.Empty();
+	RLJointToUEBoneIndices.Reserve(JointCount);
+	for (uint32 JntIndex = 0; JntIndex < JointCount; ++JntIndex)
+	{
+		const FString boneNameFStr = InDNAReader->GetJointName(JntIndex);
+		const FName BoneName = FName{ *boneNameFStr };
+		const int32 BoneIndex = RefSkeleton->FindBoneIndex(BoneName);
+
+		// BoneIndex can be INDEX_NONE;
+		// We can safely put it into the map with other indices, it will be handled in the Evaluate method.
+		RLJointToUEBoneIndices.Add(BoneIndex);
+	}
+
+
+	double TimeToMap = FPlatformTime::Seconds();
+	TimeToMap -= StartTime;
+
+	UE_LOG(LogDNAToSkelMeshMap, Log, TEXT("	Map joints:	%.6f"), TimeToMap * 1000.0f);
+}
+
+void FDNAToSkelMeshMap::MapMorphTargets(IDNAReader* InDNAReader)
+{
+	double StartTime = FPlatformTime::Seconds();
+
+	uint16 MeshCount = InDNAReader->GetMeshCount();
+	// Assemble the reverse map, for each MorphTarget index save corresponding FDNABlendShapeTarget Index
+	// Keeping MeshIndex and TargetIndex for each MorphTarget.
+	MeshBlendShapeTargets.Empty();
+	MeshBlendShapeTargets.Init(FDNABlendShapeTarget(), TargetSkelMesh->GetMorphTargets().Num());
+
+	for (int16 MeshIndex = 0; MeshIndex < MeshCount; MeshIndex++)
+	{
+		uint16 TargetCount = InDNAReader->GetBlendShapeTargetCount(MeshIndex);
+		for (int16 TargetIndex = 0; TargetIndex < TargetCount; TargetIndex++)
+		{
+			// Construct MorphTarget name by combining Blend Shape Channel name and Mesh name from DNA.
+			const uint16 ChannelIndex = InDNAReader->GetBlendShapeChannelIndex(MeshIndex, TargetIndex);
+			const FString BlendShapeStr = InDNAReader->GetBlendShapeChannelName(ChannelIndex);
+			const FString MeshStr = InDNAReader->GetMeshName(MeshIndex);
+			const FString MorphTargetStr = MeshStr + TEXT("__") + BlendShapeStr;
+			const FName MorphTargetName = FName{ *MorphTargetStr };
+
+			// Find MorphTarget from SkelMesh by previously constructed name.
+			int32 MorphTargetIndex;
+			UMorphTarget* MorphTarget = TargetSkelMesh->FindMorphTargetAndIndex(MorphTargetName, MorphTargetIndex);
+			if (MorphTarget != nullptr)
+			{
+				// Store MeshIndex and TargetIndex for found MorphTarget.
+				MeshBlendShapeTargets[MorphTargetIndex].MeshIndex = MeshIndex;
+				MeshBlendShapeTargets[MorphTargetIndex].TargetIndex = TargetIndex;
+			}
+			else
+			{
+#ifdef DEBUG
+				UE_LOG(LogDNAToSkelMeshMap, Error, TEXT("Could not find morph target %s (probably below threshold) Channel %d, Target %d"), *MorphTargetStr, ChannelIndex, TargetIndex);
+#endif
+			}
+		}
+	}
+
+	double TimeToMap = FPlatformTime::Seconds();
+	TimeToMap -= StartTime;
+
+	//UE_LOG(LogDNAToSkelMeshMap, Log, TEXT("	MapMorphTargets:	%.6f"), TimeToMap * 1000.0f);
+}
+#endif // WITH_EDITORONLY_DATA
\ No newline at end of file
diff --git a/Engine/Plugins/Animation/RigLogic/Source/RigLogicModule/Private/SkelMeshDNAReader.cpp b/Engine/Plugins/Animation/RigLogic/Source/RigLogicModule/Private/SkelMeshDNAReader.cpp
new file mode 100644
--- /dev/null
+++ b/Engine/Plugins/Animation/RigLogic/Source/RigLogicModule/Private/SkelMeshDNAReader.cpp
@@ -0,0 +1,616 @@
+// Copyright Epic Games, Inc. All Rights Reserved.
+
+#include "SkelMeshDNAReader.h"
+
+#include "DNAAsset.h"
+#include "DNAReader.h"
+
+#if WITH_EDITORONLY_DATA
+#include "EditorFramework/AssetImportData.h"
+#endif
+#include "Engine/AssetUserData.h"
+#include "Serialization/BufferArchive.h"
+#include "Serialization/MemoryReader.h"
+#include "Misc/Paths.h"
+#include "Misc/FileHelper.h"
+#include "UObject/UObjectGlobals.h"
+
+DEFINE_LOG_CATEGORY(LogSkelMeshDNAReader);
+
+FSkelMeshDNAReader::FSkelMeshDNAReader(UDNAAsset* DNAAsset) : GeometryStreamReader{ nullptr }
+{
+	BehaviorStreamReader = DNAAsset->GetBehaviorReader().Get();
+#if WITH_EDITORONLY_DATA
+	GeometryStreamReader = DNAAsset->GetGeometryReader().Get();
+#endif
+}
+
+dna::Reader* FSkelMeshDNAReader::Unwrap() const
+{
+	return nullptr;  // Unused in SkelMeshDNAReader
+}
+
+/** DESCRIPTOR READER **/
+
+FString FSkelMeshDNAReader::GetName() const
+{
+	return BehaviorStreamReader->GetName();
+}
+
+EArchetype FSkelMeshDNAReader::GetArchetype() const
+{
+	return BehaviorStreamReader->GetArchetype();
+}
+
+EGender FSkelMeshDNAReader::GetGender() const
+{
+	return BehaviorStreamReader->GetGender();
+}
+
+uint16 FSkelMeshDNAReader::GetAge() const
+{
+	return BehaviorStreamReader->GetAge();
+}
+
+uint32 FSkelMeshDNAReader::GetMetaDataCount() const
+{
+	return BehaviorStreamReader->GetMetaDataCount();
+}
+
+FString FSkelMeshDNAReader::GetMetaDataKey(uint32 Index) const
+{
+	return BehaviorStreamReader->GetMetaDataKey(Index);
+}
+
+FString FSkelMeshDNAReader::GetMetaDataValue(const FString& Key) const
+{
+	return BehaviorStreamReader->GetMetaDataValue(Key);
+}
+
+ETranslationUnit FSkelMeshDNAReader::GetTranslationUnit() const
+{
+	return BehaviorStreamReader->GetTranslationUnit();
+}
+
+ERotationUnit FSkelMeshDNAReader::GetRotationUnit() const
+{
+	return BehaviorStreamReader->GetRotationUnit();
+}
+
+FCoordinateSystem FSkelMeshDNAReader::GetCoordinateSystem() const
+{
+	return BehaviorStreamReader->GetCoordinateSystem();
+}
+
+uint16 FSkelMeshDNAReader::GetLODCount() const
+{
+	return BehaviorStreamReader->GetLODCount();
+}
+
+uint16 FSkelMeshDNAReader::GetDBMaxLOD() const
+{
+	return BehaviorStreamReader->GetDBMaxLOD();
+}
+
+FString FSkelMeshDNAReader::GetDBComplexity() const
+{
+	return BehaviorStreamReader->GetDBComplexity();
+}
+
+FString FSkelMeshDNAReader::GetDBName() const
+{
+	return BehaviorStreamReader->GetDBName();
+}
+
+/** DEFINITION READER **/
+
+uint16 FSkelMeshDNAReader::GetGUIControlCount() const
+{
+	return BehaviorStreamReader->GetGUIControlCount();
+}
+
+FString FSkelMeshDNAReader::GetGUIControlName(uint16 Index) const
+{
+	return BehaviorStreamReader->GetGUIControlName(Index);
+}
+
+uint16 FSkelMeshDNAReader::GetRawControlCount() const
+{
+	return BehaviorStreamReader->GetRawControlCount();
+}
+
+FString FSkelMeshDNAReader::GetRawControlName(uint16 Index) const
+{
+	return BehaviorStreamReader->GetRawControlName(Index);
+}
+
+uint16 FSkelMeshDNAReader::GetJointCount() const
+{
+	return BehaviorStreamReader->GetJointCount();
+}
+
+FString FSkelMeshDNAReader::GetJointName(uint16 Index) const
+{
+	return BehaviorStreamReader->GetJointName(Index);
+}
+
+uint16 FSkelMeshDNAReader::GetJointIndexListCount() const
+{
+	return BehaviorStreamReader->GetJointIndexListCount();
+}
+
+TArrayView<const uint16> FSkelMeshDNAReader::GetJointIndicesForLOD(uint16 LOD) const
+{
+	return BehaviorStreamReader->GetJointIndicesForLOD(LOD);
+}
+
+uint16 FSkelMeshDNAReader::GetBlendShapeChannelCount() const
+{
+	return BehaviorStreamReader->GetBlendShapeChannelCount();
+}
+
+FString FSkelMeshDNAReader::GetBlendShapeChannelName(uint16 Index) const
+{
+	return BehaviorStreamReader->GetBlendShapeChannelName(Index);
+}
+
+uint16 FSkelMeshDNAReader::GetBlendShapeChannelIndexListCount() const
+{
+	return BehaviorStreamReader->GetBlendShapeChannelIndexListCount();
+}
+
+TArrayView<const uint16> FSkelMeshDNAReader::GetBlendShapeChannelIndicesForLOD(uint16 LOD) const
+{
+	return BehaviorStreamReader->GetBlendShapeChannelIndicesForLOD(LOD);
+}
+
+uint16 FSkelMeshDNAReader::GetAnimatedMapCount() const
+{
+	return BehaviorStreamReader->GetAnimatedMapCount();
+}
+
+FString FSkelMeshDNAReader::GetAnimatedMapName(uint16 Index) const
+{
+	return BehaviorStreamReader->GetAnimatedMapName(Index);
+}
+
+uint16 FSkelMeshDNAReader::GetAnimatedMapIndexListCount() const
+{
+	return BehaviorStreamReader->GetAnimatedMapIndexListCount();
+}
+
+TArrayView<const uint16> FSkelMeshDNAReader::GetAnimatedMapIndicesForLOD(uint16 LOD) const
+{
+	return BehaviorStreamReader->GetAnimatedMapIndicesForLOD(LOD);
+}
+
+uint16 FSkelMeshDNAReader::GetMeshCount() const
+{
+	return BehaviorStreamReader->GetMeshCount();
+}
+
+FString FSkelMeshDNAReader::GetMeshName(uint16 Index) const
+{
+	return BehaviorStreamReader->GetMeshName(Index);
+}
+
+uint16 FSkelMeshDNAReader::GetMeshIndexListCount() const
+{
+	return BehaviorStreamReader->GetMeshIndexListCount();
+}
+
+TArrayView<const uint16> FSkelMeshDNAReader::GetMeshIndicesForLOD(uint16 LOD) const
+{
+	return BehaviorStreamReader->GetMeshIndicesForLOD(LOD);
+}
+
+uint16 FSkelMeshDNAReader::GetMeshBlendShapeChannelMappingCount() const
+{
+	return BehaviorStreamReader->GetMeshBlendShapeChannelMappingCount();
+}
+
+FMeshBlendShapeChannelMapping FSkelMeshDNAReader::GetMeshBlendShapeChannelMapping(uint16 Index) const
+{
+	return BehaviorStreamReader->GetMeshBlendShapeChannelMapping(Index);
+}
+
+TArrayView<const uint16> FSkelMeshDNAReader::GetMeshBlendShapeChannelMappingIndicesForLOD(uint16 LOD) const
+{
+	return BehaviorStreamReader->GetMeshBlendShapeChannelMappingIndicesForLOD(LOD);
+}
+
+FVector FSkelMeshDNAReader::GetNeutralJointTranslation(uint16 Index) const
+{
+	return BehaviorStreamReader->GetNeutralJointTranslation(Index);
+}
+
+FVector FSkelMeshDNAReader::GetNeutralJointRotation(uint16 Index) const
+{
+	return BehaviorStreamReader->GetNeutralJointRotation(Index);
+}
+
+uint16 FSkelMeshDNAReader::GetJointParentIndex(uint16 Index) const
+{
+	return BehaviorStreamReader->GetJointParentIndex(Index);
+}
+
+/** BEHAVIOR READER **/
+
+TArrayView<const uint16> FSkelMeshDNAReader::GetGUIToRawInputIndices() const
+{
+	return BehaviorStreamReader->GetGUIToRawInputIndices();
+}
+
+TArrayView<const uint16> FSkelMeshDNAReader::GetGUIToRawOutputIndices() const
+{
+	return BehaviorStreamReader->GetGUIToRawOutputIndices();
+}
+
+TArrayView<const float> FSkelMeshDNAReader::GetGUIToRawFromValues() const
+{
+	return BehaviorStreamReader->GetGUIToRawFromValues();
+}
+
+TArrayView<const float> FSkelMeshDNAReader::GetGUIToRawToValues() const
+{
+	return  BehaviorStreamReader->GetGUIToRawToValues();
+}
+
+TArrayView<const float> FSkelMeshDNAReader::GetGUIToRawSlopeValues() const
+{
+	return BehaviorStreamReader->GetGUIToRawSlopeValues();
+}
+
+TArrayView<const float> FSkelMeshDNAReader::GetGUIToRawCutValues() const
+{
+	return BehaviorStreamReader->GetGUIToRawCutValues();
+}
+
+uint16 FSkelMeshDNAReader::GetPSDCount() const
+{
+	return BehaviorStreamReader->GetPSDCount();
+}
+
+TArrayView<const uint16> FSkelMeshDNAReader::GetPSDRowIndices() const
+{
+	return BehaviorStreamReader->GetPSDRowIndices();
+}
+
+TArrayView<const uint16> FSkelMeshDNAReader::GetPSDColumnIndices() const
+{
+	return BehaviorStreamReader->GetPSDColumnIndices();
+}
+
+TArrayView<const float> FSkelMeshDNAReader::GetPSDValues() const
+{
+	return BehaviorStreamReader->GetPSDValues();
+}
+
+uint16 FSkelMeshDNAReader::GetJointRowCount() const
+{
+	return BehaviorStreamReader->GetJointRowCount();
+}
+
+uint16 FSkelMeshDNAReader::GetJointColumnCount() const
+{
+	return BehaviorStreamReader->GetJointColumnCount();
+}
+
+TArrayView<const uint16> FSkelMeshDNAReader::GetJointGroupJointIndices(uint16 JointGroupIndex) const
+{
+	return BehaviorStreamReader->GetJointGroupJointIndices(JointGroupIndex);
+}
+
+TArrayView<const uint16> FSkelMeshDNAReader::GetJointVariableAttributeIndices(uint16 LOD) const
+{
+	return BehaviorStreamReader->GetJointVariableAttributeIndices(LOD);
+}
+
+uint16 FSkelMeshDNAReader::GetJointGroupCount() const
+{
+	return BehaviorStreamReader->GetJointGroupCount();
+}
+
+TArrayView<const uint16> FSkelMeshDNAReader::GetJointGroupLODs(uint16 JointGroupIndex) const
+{
+	return BehaviorStreamReader->GetJointGroupLODs(JointGroupIndex);
+}
+
+TArrayView<const uint16> FSkelMeshDNAReader::GetJointGroupInputIndices(uint16 JointGroupIndex) const
+{
+	return BehaviorStreamReader->GetJointGroupInputIndices(JointGroupIndex);
+}
+
+TArrayView<const uint16> FSkelMeshDNAReader::GetJointGroupOutputIndices(uint16 JointGroupIndex) const
+{
+	return BehaviorStreamReader->GetJointGroupOutputIndices(JointGroupIndex);
+}
+
+TArrayView<const float> FSkelMeshDNAReader::GetJointGroupValues(uint16 JointGroupIndex) const
+{
+	return BehaviorStreamReader->GetJointGroupValues(JointGroupIndex);
+}
+
+TArrayView<const uint16> FSkelMeshDNAReader::GetBlendShapeChannelLODs() const
+{
+	return BehaviorStreamReader->GetBlendShapeChannelLODs();
+}
+
+TArrayView<const uint16> FSkelMeshDNAReader::GetBlendShapeChannelInputIndices() const
+{
+	return BehaviorStreamReader->GetBlendShapeChannelInputIndices();
+}
+
+TArrayView<const uint16> FSkelMeshDNAReader::GetBlendShapeChannelOutputIndices() const
+{
+	return BehaviorStreamReader->GetBlendShapeChannelOutputIndices();
+}
+
+TArrayView<const uint16> FSkelMeshDNAReader::GetAnimatedMapLODs() const
+{
+	return BehaviorStreamReader->GetAnimatedMapLODs();
+}
+
+TArrayView<const uint16> FSkelMeshDNAReader::GetAnimatedMapInputIndices() const
+{
+	return BehaviorStreamReader->GetAnimatedMapInputIndices();
+}
+
+TArrayView<const uint16> FSkelMeshDNAReader::GetAnimatedMapOutputIndices() const
+{
+	return BehaviorStreamReader->GetAnimatedMapOutputIndices();
+}
+
+TArrayView<const float> FSkelMeshDNAReader::GetAnimatedMapFromValues() const
+{
+	return BehaviorStreamReader->GetAnimatedMapFromValues();
+}
+
+TArrayView<const float> FSkelMeshDNAReader::GetAnimatedMapToValues() const
+{
+	return BehaviorStreamReader->GetAnimatedMapToValues();
+}
+
+TArrayView<const float> FSkelMeshDNAReader::GetAnimatedMapSlopeValues() const
+{
+	return BehaviorStreamReader->GetAnimatedMapSlopeValues();
+}
+
+TArrayView<const float> FSkelMeshDNAReader::GetAnimatedMapCutValues() const
+{
+	return BehaviorStreamReader->GetAnimatedMapCutValues();
+}
+
+/** GEOMETRY READER **/
+uint32 FSkelMeshDNAReader::GetVertexPositionCount(uint16 MeshIndex) const
+{
+	if (GeometryStreamReader)
+	{
+		return GeometryStreamReader->GetVertexPositionCount(MeshIndex);
+	}
+	return {};
+}
+
+FVector FSkelMeshDNAReader::GetVertexPosition(uint16 MeshIndex, uint32 PositionIndex) const
+{
+	if (GeometryStreamReader)
+	{
+		return GeometryStreamReader->GetVertexPosition(MeshIndex, PositionIndex);
+	}
+	return {};
+}
+
+TArrayView<const float> FSkelMeshDNAReader::GetVertexPositionXs(uint16 MeshIndex) const
+{
+	if (GeometryStreamReader)
+	{
+		return GeometryStreamReader->GetVertexPositionXs(MeshIndex);
+	}
+	return {};
+}
+
+TArrayView<const float> FSkelMeshDNAReader::GetVertexPositionYs(uint16 MeshIndex) const
+{
+	if (GeometryStreamReader)
+	{
+		return GeometryStreamReader->GetVertexPositionYs(MeshIndex);
+	}
+	return {};
+}
+
+TArrayView<const float> FSkelMeshDNAReader::GetVertexPositionZs(uint16 MeshIndex) const
+{
+	if (GeometryStreamReader)
+	{
+		return GeometryStreamReader->GetVertexPositionZs(MeshIndex);
+	}
+	return {};
+}
+
+uint32 FSkelMeshDNAReader::GetVertexTextureCoordinateCount(uint16 MeshIndex) const
+{
+	if (GeometryStreamReader)
+	{
+		return GeometryStreamReader->GetVertexTextureCoordinateCount(MeshIndex);
+	}
+	return {};
+}
+
+FTextureCoordinate FSkelMeshDNAReader::GetVertexTextureCoordinate(uint16 MeshIndex, uint32 TextureCoordinateIndex) const
+{
+	if (GeometryStreamReader)
+	{
+		return GeometryStreamReader->GetVertexTextureCoordinate(MeshIndex, TextureCoordinateIndex);
+	}
+	return {};
+}
+
+uint32 FSkelMeshDNAReader::GetVertexNormalCount(uint16 MeshIndex) const
+{
+	if (GeometryStreamReader)
+	{
+		return GeometryStreamReader->GetVertexNormalCount(MeshIndex);
+	}
+	return {};
+}
+
+FVector FSkelMeshDNAReader::GetVertexNormal(uint16 MeshIndex, uint32 NormalIndex) const
+{
+	if (GeometryStreamReader)
+	{
+		return GeometryStreamReader->GetVertexNormal(MeshIndex, NormalIndex);
+	}
+	return {};
+}
+
+/* not needed for gene splicer */
+uint32 FSkelMeshDNAReader::GetVertexLayoutCount(uint16 MeshIndex) const
+{
+	if (GeometryStreamReader)
+	{
+		return GeometryStreamReader->GetVertexLayoutCount(MeshIndex);
+	}
+	return {};
+}
+
+/* not needed for gene splicer */
+FVertexLayout FSkelMeshDNAReader::GetVertexLayout(uint16 MeshIndex, uint32 VertexIndex) const
+{
+	if (GeometryStreamReader)
+	{
+		return GeometryStreamReader->GetVertexLayout(MeshIndex, VertexIndex);
+	}
+	return {};
+}
+
+/* not needed for gene splicer */
+uint32 FSkelMeshDNAReader::GetFaceCount(uint16 MeshIndex) const
+{
+	if (GeometryStreamReader)
+	{
+		return GeometryStreamReader->GetFaceCount(MeshIndex);
+	}
+	return {};
+}
+
+/* not needed for gene splicer */
+TArrayView<const uint32> FSkelMeshDNAReader::GetFaceVertexLayoutIndices(uint16 MeshIndex, uint32 FaceIndex) const
+{
+	if (GeometryStreamReader)
+	{
+		return GeometryStreamReader->GetFaceVertexLayoutIndices(MeshIndex, FaceIndex);
+	}
+	return {};
+}
+
+uint16 FSkelMeshDNAReader::GetMaximumInfluencePerVertex(uint16 MeshIndex) const
+{
+	if (GeometryStreamReader)
+	{
+		return GeometryStreamReader->GetMaximumInfluencePerVertex(MeshIndex);
+	}
+	return {};
+}
+
+uint32 FSkelMeshDNAReader::GetSkinWeightsCount(uint16 MeshIndex) const
+{
+	if (GeometryStreamReader)
+	{
+		return GeometryStreamReader->GetSkinWeightsCount(MeshIndex);
+	}
+	return {};
+}
+
+TArrayView<const float> FSkelMeshDNAReader::GetSkinWeightsValues(uint16 MeshIndex, uint32 VertexIndex) const
+{
+	if (GeometryStreamReader)
+	{
+		return GeometryStreamReader->GetSkinWeightsValues(MeshIndex, VertexIndex);
+	}
+	return {};
+}
+
+TArrayView<const uint16> FSkelMeshDNAReader::GetSkinWeightsJointIndices(uint16 MeshIndex, uint32 VertexIndex) const
+{
+	if (GeometryStreamReader)
+	{
+		return GeometryStreamReader->GetSkinWeightsJointIndices(MeshIndex, VertexIndex);
+	}
+	return {};
+}
+
+uint16 FSkelMeshDNAReader::GetBlendShapeTargetCount(uint16 MeshIndex) const
+{
+	if (GeometryStreamReader)
+	{
+		return GeometryStreamReader->GetBlendShapeTargetCount(MeshIndex);
+	}
+	return {};
+}
+
+uint16 FSkelMeshDNAReader::GetBlendShapeChannelIndex(uint16 MeshIndex, uint16 BlendShapeTargetIndex) const
+{
+	if (GeometryStreamReader)
+	{
+		return GeometryStreamReader->GetBlendShapeChannelIndex(MeshIndex, BlendShapeTargetIndex);
+	}
+	return {};
+}
+
+uint32 FSkelMeshDNAReader::GetBlendShapeTargetDeltaCount(uint16 MeshIndex, uint16 BlendShapeIndex) const
+{
+	if (GeometryStreamReader)
+	{
+		return GeometryStreamReader->GetBlendShapeTargetDeltaCount(MeshIndex, BlendShapeIndex);
+	}
+	return {};
+}
+
+FVector FSkelMeshDNAReader::GetBlendShapeTargetDelta(uint16 MeshIndex, uint16 BlendShapeIndex, uint32 DeltaIndex) const
+{
+	if (GeometryStreamReader)
+	{
+		return GeometryStreamReader->GetBlendShapeTargetDelta(MeshIndex, BlendShapeIndex, DeltaIndex);
+	}
+	return {};
+}
+
+TArrayView<const float> FSkelMeshDNAReader::GetBlendShapeTargetDeltaXs(uint16 MeshIndex, uint16 BlendShapeTargetIndex) const
+{
+	if (GeometryStreamReader)
+	{
+		return GeometryStreamReader->GetBlendShapeTargetDeltaXs(MeshIndex, BlendShapeTargetIndex);
+	}
+	return {};
+}
+
+TArrayView<const float> FSkelMeshDNAReader::GetBlendShapeTargetDeltaYs(uint16 MeshIndex, uint16 BlendShapeTargetIndex) const
+{
+	if (GeometryStreamReader)
+	{
+		return GeometryStreamReader->GetBlendShapeTargetDeltaYs(MeshIndex, BlendShapeTargetIndex);
+	}
+	return {};
+}
+
+TArrayView<const float> FSkelMeshDNAReader::GetBlendShapeTargetDeltaZs(uint16 MeshIndex, uint16 BlendShapeTargetIndex) const
+{
+	if (GeometryStreamReader)
+	{
+		return GeometryStreamReader->GetBlendShapeTargetDeltaZs(MeshIndex, BlendShapeTargetIndex);
+	}
+	return {};
+}
+
+TArrayView<const uint32> FSkelMeshDNAReader::GetBlendShapeTargetVertexIndices(uint16 MeshIndex, uint16 BlendShapeIndex) const
+{
+	if (GeometryStreamReader)
+	{
+		return GeometryStreamReader->GetBlendShapeTargetVertexIndices(MeshIndex, BlendShapeIndex);
+	}
+	return {};
+}
+
+void FSkelMeshDNAReader::Unload(EDNADataLayer Layer) {
+	ensureMsgf(false, TEXT("Assest are not unloadable"));
+}
\ No newline at end of file
diff --git a/Engine/Plugins/Animation/RigLogic/Source/RigLogicModule/Private/SkelMeshDNAUtils.cpp b/Engine/Plugins/Animation/RigLogic/Source/RigLogicModule/Private/SkelMeshDNAUtils.cpp
new file mode 100644
--- /dev/null
+++ b/Engine/Plugins/Animation/RigLogic/Source/RigLogicModule/Private/SkelMeshDNAUtils.cpp
@@ -0,0 +1,483 @@
+// Copyright Epic Games, Inc. All Rights Reserved.
+#include "SkelMeshDNAUtils.h"
+#include "DNAToSkelMeshMap.h"
+
+#include "RenderResource.h"
+#include "RHICommandList.h"
+#include "Async/ParallelFor.h"
+#include "Components/SkeletalMeshComponent.h"
+#include "Engine/StaticMesh.h"
+#include "Misc/Paths.h"
+#include "ProfilingDebugging/ScopedTimers.h"
+#include "Misc/FileHelper.h"
+#include "UObject/ConstructorHelpers.h"
+#include "Engine/SkeletalMesh.h"
+#include "ComponentReregisterContext.h"
+#include "Rendering/SkeletalMeshModel.h"
+#include "Rendering/SkeletalMeshLODModel.h"
+#include "Rendering/SkeletalMeshRenderData.h"
+#include "Rendering/SkeletalMeshLODRenderData.h"
+#if WITH_EDITORONLY_DATA
+#include "MeshUtilities.h"
+#endif // WITH_EDITORONLY_DATA
+#include "Modules/ModuleManager.h"
+
+#include "AnimationRuntime.h"
+
+#include "riglogic/RigLogic.h"
+
+DEFINE_LOG_CATEGORY(LogDNAUtils);
+/** compare based on base mesh source vertex indices */
+struct FCompareMorphTargetDeltas
+{
+	FORCEINLINE bool operator()(const FMorphTargetDelta& A, const FMorphTargetDelta& B) const
+	{
+		return ((int32)A.SourceIdx - (int32)B.SourceIdx) < 0;
+	}
+};
+
+USkelMeshDNAUtils::USkelMeshDNAUtils(const FObjectInitializer& ObjectInitializer)
+	: Super(ObjectInitializer)
+{
+}
+
+FDNAToSkelMeshMap* USkelMeshDNAUtils::CreateMapForUpdatingNeutralMesh(IDNAReader* DNAReader, USkeletalMesh* SkelMesh)
+{
+#if WITH_EDITORONLY_DATA
+	FDNAToSkelMeshMap* DNAToSkelMeshMap = new FDNAToSkelMeshMap();
+
+	//only vertex map is initialized in this pass so we can mix neutral meshes fast (e.g. on slider move);
+	//playing animation on a such a mesh requires updating joints and skin weights
+	//getting full quality animation requires mixing morph targets too
+	DNAToSkelMeshMap->InitBaseMesh(DNAReader, SkelMesh);
+	return DNAToSkelMeshMap;
+#else
+	return nullptr;
+#endif // WITH_EDITORONLY_DATA
+}
+
+#if WITH_EDITORONLY_DATA
+/** Updates bind pose using joint positions from DNA. */
+void USkelMeshDNAUtils::UpdateJoints(USkeletalMesh* SkelMesh, IDNAReader* DNAReader, FDNAToSkelMeshMap* DNAToSkelMeshMap)
+{
+	{	// Scoping of RefSkelModifier
+		FReferenceSkeletonModifier RefSkelModifier(SkelMesh->GetRefSkeleton(), SkelMesh->GetSkeleton());
+
+		// copy here
+		TArray<FTransform> RawBonePose = SkelMesh->GetRefSkeleton().GetRawRefBonePose();
+
+		// When we are mounting the head to different bodies than female average, we need to use
+		// component space, as the joint to which the head root is snapped to will be on a different
+		// position than in the head rig!
+
+		// calculate component space ahead of current transform
+		TArray<FTransform> ComponentTransforms;
+		FAnimationRuntime::FillUpComponentSpaceTransforms(SkelMesh->GetRefSkeleton(), RawBonePose, ComponentTransforms);
+
+		const TArray<FMeshBoneInfo>& RawBoneInfo = SkelMesh->GetRefSkeleton().GetRawRefBoneInfo();
+
+		// Skipping root joint (index 0) to avoid blinking of the mesh due to bounding box issue
+		for (uint16 JointIndex = 0; JointIndex < DNAReader->GetJointCount(); JointIndex++)
+		{
+			int32 BoneIndex = DNAToSkelMeshMap->GetUEBoneIndex(JointIndex);
+
+			FTransform DNATransform = FTransform::Identity;
+
+			// Updating bind pose affects just translations.
+			FVector Translate = DNAReader->GetNeutralJointTranslation(JointIndex);
+			DNATransform.SetTranslation(Translate);
+			FVector RotationVector = DNAReader->GetNeutralJointRotation(JointIndex);
+			FRotator Rotation(RotationVector.X, RotationVector.Y, RotationVector.Z);
+
+			// Joint 0 (spine_04) is a root of GeneSplicer joint hierarchy, and is a special case
+			// 1) it is parent to itself
+			// 2) it is in DNA space, so we need to rotate it 90 degs on x axis to UE4 space
+			// 3) the head joints below it in the skeletal mesh are not spliced, as they are not in DNA,
+			//    so they will retain female average transforms
+
+			if (DNAReader->GetJointParentIndex(JointIndex) == JointIndex)  //parent to itself
+			{
+				Rotation.Pitch += 90;
+				DNATransform.SetRotation(Rotation.Quaternion());
+
+				DNATransform.SetTranslation(FVector(Translate.X, Translate.Z, -Translate.Y));
+
+				ComponentTransforms[BoneIndex] = DNATransform;
+
+			}
+			else
+			{
+				DNATransform.SetRotation(Rotation.Quaternion());
+
+				if (ensure(RawBoneInfo[BoneIndex].ParentIndex != INDEX_NONE))
+				{
+					ComponentTransforms[BoneIndex] = DNATransform * ComponentTransforms[RawBoneInfo[BoneIndex].ParentIndex];
+				}
+			}
+
+			ComponentTransforms[BoneIndex].NormalizeRotation();
+		}
+
+		for (uint16 BoneIndex = 0; BoneIndex < RawBoneInfo.Num(); BoneIndex++)
+		{
+			FTransform LocalTransform;
+
+			if (BoneIndex == 0)
+			{
+				LocalTransform = ComponentTransforms[BoneIndex];
+			}
+			else
+			{
+				LocalTransform = ComponentTransforms[BoneIndex].GetRelativeTransform(ComponentTransforms[RawBoneInfo[BoneIndex].ParentIndex]);
+			}
+
+			LocalTransform.NormalizeRotation();
+
+			RefSkelModifier.UpdateRefPoseTransform(BoneIndex, LocalTransform);
+		}
+	}
+
+	SkelMesh->GetRefBasesInvMatrix().Reset();
+	SkelMesh->CalculateInvRefMatrices(); // Needs to be called after RefSkelModifier is destroyed
+}
+
+/** Updates base mesh vertices using data from DNA. */
+void USkelMeshDNAUtils::UpdateBaseMesh(USkeletalMesh* SkelMesh, IDNAReader* DNAReader, FDNAToSkelMeshMap* DNAToSkelMeshMap, ELodUpdateOption UpdateOption)
+{
+
+	FSkeletalMeshModel* ImportedModel = SkelMesh->GetImportedModel();
+	// Expects vertex map to be initialized beforehand	
+	int32 LODStart;
+	int32 LODRangeSize;
+	GetLODRange(UpdateOption, ImportedModel->LODModels.Num(), LODStart, LODRangeSize);
+	for (int32 LODIndex = LODStart; LODIndex < LODRangeSize; LODIndex++)
+	{
+		FSkeletalMeshLODModel& LODModel = ImportedModel->LODModels[LODIndex];
+		int32 SectionIndex = 0;
+		for (FSkelMeshSection& Section : LODModel.Sections)
+		{
+			int32& DNAMeshIndex = DNAToSkelMeshMap->ImportVtxToDNAMeshIndex[LODIndex][Section.GetVertexBufferIndex()];
+
+			const int32 NumSoftVertices = Section.GetNumVertices();
+			auto& OverlappingMap = DNAToSkelMeshMap->OverlappingVertices[LODIndex][SectionIndex];
+			int32 VertexBufferIndex = Section.GetVertexBufferIndex();
+			for (int32 VertexIndex = 0; VertexIndex < NumSoftVertices; VertexIndex++)
+			{
+				int32& DNAVertexIndex = DNAToSkelMeshMap->ImportVtxToDNAVtxIndex[LODIndex][VertexBufferIndex];
+
+				if (DNAVertexIndex >= 0)
+				{
+					const FVector Position = DNAReader->GetVertexPosition(DNAMeshIndex, DNAVertexIndex);
+					FSoftSkinVertex& Vertex = Section.SoftVertices[VertexIndex];
+					Vertex.Position = FVector3f{ Position };
+
+					// Check if the current vertex has overlapping vertices, and then update them as well.
+					TArray<int32>& OverlappedIndices = OverlappingMap[VertexIndex];
+					int32 OverlappingCount = OverlappedIndices.Num();
+					for (int32 OverlappingIndex = 0; OverlappingIndex < OverlappingCount; ++OverlappingIndex)
+					{
+						int32 OverlappingVertexIndex = OverlappedIndices[OverlappingIndex];
+						FSoftSkinVertex& OverlappingVertex = Section.SoftVertices[OverlappingVertexIndex];
+						OverlappingVertex.Position = FVector3f{ Position };
+					}
+				}
+				VertexBufferIndex++;
+			}
+			SectionIndex++;
+		}
+	}
+}
+
+/** Updates Morph Targets using Blend Shapes from DNA.  */
+void USkelMeshDNAUtils::UpdateMorphTargets(USkeletalMesh* SkelMesh, IDNAReader* DNAReader, FDNAToSkelMeshMap* DNAToSkelMeshMap, ELodUpdateOption UpdateOption)
+{
+	TArray<FDNABlendShapeTarget>& MeshBlendShapeTargets = DNAToSkelMeshMap->GetMeshBlendShapeTargets();
+	if (MeshBlendShapeTargets.Num() == 0)
+	{
+		UE_LOG(LogDNAUtils, Warning, TEXT("No morph targets updated!"));
+		return;
+	}
+	FSkeletalMeshModel* ImportedModel = SkelMesh->GetImportedModel();
+
+	uint16 ChannelCount = DNAReader->GetBlendShapeChannelCount();
+	ParallelFor(SkelMesh->GetMorphTargets().Num(), [&](int32 MorphIndex)
+		{
+			UMorphTarget* MorphTarget = SkelMesh->GetMorphTargets()[MorphIndex];
+			const FDNABlendShapeTarget& MeshTarget = MeshBlendShapeTargets[MorphIndex];
+			// First get all DNA deltas for current Morph Target.
+			TArrayView<const uint32> BlendShapeVertexIndices = DNAReader->GetBlendShapeTargetVertexIndices(MeshTarget.MeshIndex, MeshTarget.TargetIndex);
+			const int32 NumOfDeltas = BlendShapeVertexIndices.Num();
+			if (NumOfDeltas > 0)
+			{
+				int32 LODStart;
+				int32 LODRangeSize;
+				GetLODRange(UpdateOption, MorphTarget->GetMorphLODModels().Num(), LODStart, LODRangeSize);
+				for (int32 LODIndex = LODStart; LODIndex < LODRangeSize; LODIndex++)
+				{
+					// MorphTarget vertex indices refer to full vertex index buffer of imported mesh.
+					FMorphTargetLODModel& MorphLODModel = MorphTarget->GetMorphLODModels()[LODIndex];
+					MorphLODModel.NumBaseMeshVerts = NumOfDeltas;
+					MorphLODModel.bGeneratedByEngine = false;
+					MorphLODModel.SectionIndices.Empty();
+					MorphLODModel.Vertices.Reset(NumOfDeltas);
+					TArrayView<const float> DeltaXs = DNAReader->GetBlendShapeTargetDeltaXs(MeshTarget.MeshIndex, MeshTarget.TargetIndex);
+					TArrayView<const float> DeltaYs = DNAReader->GetBlendShapeTargetDeltaYs(MeshTarget.MeshIndex, MeshTarget.TargetIndex);
+					TArrayView<const float> DeltaZs = DNAReader->GetBlendShapeTargetDeltaZs(MeshTarget.MeshIndex, MeshTarget.TargetIndex);
+
+					for (int32 DeltaIndex = 0; DeltaIndex < NumOfDeltas; DeltaIndex++)
+					{
+						FVector PositionDelta = FVector(DeltaXs[DeltaIndex], DeltaYs[DeltaIndex], DeltaZs[DeltaIndex]);
+						FMorphTargetDelta MorphDelta;
+						int32 DNAVertexIndex = BlendShapeVertexIndices[DeltaIndex];
+						int32 UEVertexIndex = DNAToSkelMeshMap->ImportDNAVtxToUEVtxIndex[LODIndex][MeshTarget.MeshIndex][DNAVertexIndex];
+						if (UEVertexIndex > INDEX_NONE)
+						{
+							MorphDelta.SourceIdx = (uint32)UEVertexIndex;
+							MorphDelta.PositionDelta = FVector3f{ PositionDelta };
+							MorphDelta.TangentZDelta = FVector3f::ZeroVector;
+							MorphLODModel.Vertices.Add(MorphDelta);
+							// Find section indices that are involved in these morph deltas.
+							int32 SectionIdx = DNAToSkelMeshMap->UEVertexToSectionIndices[LODIndex][UEVertexIndex];
+							if (!MorphLODModel.SectionIndices.Contains(SectionIdx) && SectionIdx > INDEX_NONE)
+							{
+								MorphLODModel.SectionIndices.Add(SectionIdx);
+							}
+						}
+					}
+					MorphLODModel.Vertices.Sort(FCompareMorphTargetDeltas());
+				}
+			}
+			else
+			{
+				int32 LODStart;
+				int32 LODRangeSize;
+				GetLODRange(UpdateOption, MorphTarget->GetMorphLODModels().Num(), LODStart, LODRangeSize);
+				for (int32 LODIndex = LODStart; LODIndex < LODRangeSize; LODIndex++)
+				{
+					FMorphTargetLODModel& MorphLODModel = MorphTarget->GetMorphLODModels()[LODIndex];
+					MorphLODModel.Reset();
+				}
+#ifdef DEBUG
+				UE_LOG(LogDNAUtils, Warning, TEXT(" 0 deltas found for mesh %d and blend shape target %d"), MeshTarget.MeshIndex, MeshTarget.TargetIndex);
+#endif 
+			}
+		});
+}
+
+/* Updates Bone influences using Skin Weights from DNA. */
+void USkelMeshDNAUtils::UpdateSkinWeights(USkeletalMesh* SkelMesh, IDNAReader* DNAReader, FDNAToSkelMeshMap* DNAToSkelMeshMap, ELodUpdateOption UpdateOption)
+{
+	FSkeletalMeshModel* ImportedModel = SkelMesh->GetImportedModel();
+
+	bool InfluenceMismatch = false;
+	//Set a threshold a bit smaller then 1/255
+	constexpr float MINWEIGHT = 0.9999f / 255.0f;
+	int32 LODStart;
+	int32 LODRangeSize;
+	GetLODRange(UpdateOption, ImportedModel->LODModels.Num(), LODStart, LODRangeSize);
+	for (int32 LODIndex = LODStart; LODIndex < LODRangeSize; ++LODIndex)
+	{
+		FSkeletalMeshLODModel& LODModel = ImportedModel->LODModels[LODIndex];
+		for (int32 SectionIndex = 0; SectionIndex < LODModel.Sections.Num(); SectionIndex++)
+		{
+			FSkelMeshSection& Section = LODModel.Sections[SectionIndex];
+			int32 DNAMeshIndex = DNAToSkelMeshMap->ImportVtxToDNAMeshIndex[LODIndex][Section.GetVertexBufferIndex()];
+
+			const int32 NumEngineVertices = Section.GetNumVertices();
+			for (int32 VertexIndex = 0; VertexIndex < NumEngineVertices; VertexIndex++)
+			{
+				const int32 VertexBufferIndex = VertexIndex + Section.GetVertexBufferIndex();
+				int32 DNAVertexIndex = DNAToSkelMeshMap->ImportVtxToDNAVtxIndex[LODIndex][VertexBufferIndex];
+
+				if (DNAVertexIndex < 0) continue; // Skip vertex not in DNA.
+
+				TArrayView<const float> DNASkinWeights = DNAReader->GetSkinWeightsValues(DNAMeshIndex, DNAVertexIndex);
+				TArrayView<const uint16> DNASkinJoints = DNAReader->GetSkinWeightsJointIndices(DNAMeshIndex, DNAVertexIndex);
+				uint16 SkinJointNum = DNASkinJoints.Num();
+
+				uint32 TotalWeight = 0;  // store all influences to vertex to ensure they add up to 255 (fix rounding errors)
+				uint16 MaxInfluenceIndex = 0;
+				uint8 MaxInfluenceWeight = 0;
+
+				FSoftSkinVertex& Vertex = Section.SoftVertices[VertexIndex];
+
+				// Update skin weights only around eyes where blue channel is not 0.
+				if (Vertex.Color.B == 0) {
+					//UE_LOG(LogDNAUtils, Log, TEXT("Skipping vertex UE: %d, DNA: %d"), VertexIndex, DNAVertexIndex );
+					continue;
+				}
+				TArray<float> NormalizedSkinWeights;
+				NormalizedSkinWeights.Reserve(SkinJointNum);
+				TArray<uint16> NormalizedJoints;
+				NormalizedJoints.Reserve(SkinJointNum);
+				float TotalInfluence = 0.0;
+				// Second step is to filter out all the influences below the MIN WEIGHT or above total influence limit.
+				for (uint16 i = 0; i < SkinJointNum; i++)
+				{
+					if (DNASkinWeights[i] > MINWEIGHT && i < MAX_TOTAL_INFLUENCES)
+					{
+						NormalizedSkinWeights.Add(DNASkinWeights[i]);
+						NormalizedJoints.Add(DNASkinJoints[i]);
+						TotalInfluence += DNASkinWeights[i];
+					}
+				}
+				SkinJointNum = NormalizedJoints.Num();
+				if (SkinJointNum > 0 && (TotalInfluence != 1.0f))
+				{
+					// Missing fractions of influence weights have to be assigned equally along the existing influences.
+					float OneOverTotalWeight = 1.f / TotalInfluence;
+					for (uint16 r = 0; r < SkinJointNum; r++)
+					{
+						NormalizedSkinWeights[r] *= OneOverTotalWeight;
+					}
+				}
+				// Reset all influences that are not covered by DNA data.
+				for (uint16 i = SkinJointNum; i < MAX_TOTAL_INFLUENCES; i++)
+				{
+					Vertex.InfluenceBones[i] = 0;
+					Vertex.InfluenceWeights[i] = 0;
+				}
+				for (uint16 InfluenceIndex = 0; InfluenceIndex < SkinJointNum; ++InfluenceIndex)
+				{
+					uint8 EngineWeight = 0;
+					// Find Engine bone for corresponding DNAJoint for the same influence.
+					int32 UpdatedBoneId = DNAToSkelMeshMap->GetUEBoneIndex(NormalizedJoints[InfluenceIndex]);
+					// BoneMap holds subset of bones belonging to current section.
+					int32 BoneMapIndex = Section.BoneMap.Find(UpdatedBoneId);
+
+					// Update which bone in the subset influences this vertex.
+					Vertex.InfluenceBones[InfluenceIndex] = BoneMapIndex;
+					if (BoneMapIndex != INDEX_NONE)
+					{
+						// Update influence weight.
+						EngineWeight = (uint8)(NormalizedSkinWeights[InfluenceIndex] * 255.0f);
+					}
+
+					Vertex.InfluenceWeights[InfluenceIndex] = EngineWeight;
+					TotalWeight += Vertex.InfluenceWeights[InfluenceIndex];
+
+					if (EngineWeight > MaxInfluenceWeight)
+					{
+						MaxInfluenceIndex = InfluenceIndex;
+						MaxInfluenceWeight = EngineWeight;
+					}
+				}
+				// Add missing fraction to fill up to 255.
+				Vertex.InfluenceWeights[MaxInfluenceIndex] += 255 - TotalWeight;
+			}
+		}
+	}
+}
+
+/** Rebuilds render data from LODModel and inits resources. */
+void USkelMeshDNAUtils::RebuildRenderData(USkeletalMesh* SkelMesh)
+{
+	FPlatformTime::InitTiming();
+
+	double StartTime = FPlatformTime::Seconds();
+	{
+		SkelMesh->FlushRenderState();
+	}
+	double TimeToFlush = FPlatformTime::Seconds();
+	{
+		FSkeletalMeshRenderData* RenderData = SkelMesh->GetResourceForRendering();
+		int32 LODIndex = 0;
+
+		for (FSkeletalMeshLODRenderData& LODRenderData : RenderData->LODRenderData)
+		{
+			FSkeletalMeshLODModel& LODModelRef = SkelMesh->GetImportedModel()->LODModels[LODIndex];
+			for (int32 i = 0; i < LODModelRef.Sections.Num(); i++)
+			{
+				FSkelMeshSection& ModelSection = LODModelRef.Sections[i];
+				ModelSection.CalcMaxBoneInfluences();
+				ModelSection.CalcUse16BitBoneIndex();
+			}
+
+			const FSkeletalMeshLODModel* LODModelPtr = &LODModelRef;
+			LODRenderData.BuildFromLODModel(LODModelPtr, 0);
+			LODIndex++;
+		}
+	}
+	double TimeToRebuildModel = FPlatformTime::Seconds();
+	{
+		if (FApp::CanEverRender())
+		{
+			// Reinitialize the static mesh's resources.
+			SkelMesh->InitResources();
+		}
+	}
+	double TimeToInitResources = FPlatformTime::Seconds();
+	{
+		// Re-register scope
+		TArray<UActorComponent*> ComponentsToReregister;
+		for (TObjectIterator<USkeletalMeshComponent> It; It; ++It)
+		{
+			USkeletalMeshComponent* MeshComponent = *It;
+			if (MeshComponent && !MeshComponent->IsTemplate() && MeshComponent->GetSkeletalMeshAsset() == SkelMesh)
+			{
+				ComponentsToReregister.Add(*It);
+			}
+		}
+		FMultiComponentReregisterContext ReregisterContext(ComponentsToReregister);
+	}
+}
+
+void USkelMeshDNAUtils::RebuildRenderData_VertexPosition(USkeletalMesh* SkelMesh)
+{
+	if (!FApp::CanEverRender())
+	{
+		return;
+	}
+
+	{
+		FSkeletalMeshModel* MeshModel = SkelMesh->GetImportedModel();
+		FSkeletalMeshRenderData* RenderData = SkelMesh->GetResourceForRendering();
+
+		for (int32 LODIdx = 0; LODIdx < RenderData->LODRenderData.Num(); ++LODIdx)
+		{
+			FSkeletalMeshLODModel& LODModel = MeshModel->LODModels[LODIdx];
+			FSkeletalMeshLODRenderData& LODRenderData = RenderData->LODRenderData[LODIdx];
+
+			ENQUEUE_RENDER_COMMAND(FSkelMeshDNAUpdatePositions)
+				([&LODModel, &LODRenderData](FRHICommandListImmediate& RHICmdList)
+					{
+						LLM_SCOPE(ELLMTag::SkeletalMesh);
+
+						TArray<FSoftSkinVertex> Vertices;
+						LODModel.GetVertices(Vertices);
+						check(Vertices.Num() == LODRenderData.StaticVertexBuffers.PositionVertexBuffer.GetNumVertices());
+						LODRenderData.StaticVertexBuffers.PositionVertexBuffer.Init(Vertices.Num());
+						for (int32 i = 0; i < Vertices.Num(); i++)
+						{
+							LODRenderData.StaticVertexBuffers.PositionVertexBuffer.VertexPosition(i) = Vertices[i].Position;
+						}
+
+						auto& VertexBuffer = LODRenderData.StaticVertexBuffers.PositionVertexBuffer;
+						void* VertexBufferData = RHILockBuffer(VertexBuffer.VertexBufferRHI, 0, VertexBuffer.GetNumVertices() * VertexBuffer.GetStride(), RLM_WriteOnly);
+						FMemory::Memcpy(VertexBufferData, VertexBuffer.GetVertexData(), VertexBuffer.GetNumVertices() * VertexBuffer.GetStride());
+						RHIUnlockBuffer(VertexBuffer.VertexBufferRHI);
+					});
+		}
+	}
+}
+
+void USkelMeshDNAUtils::UpdateJointBehavior(USkeletalMeshComponent* SkelMeshComponent)
+{
+	//DNAAsset->SetBehaviorReader is called before calling this method
+	//it is not here to avoid having DNAAsset in the API, as in future we might want
+	//to generalize SkelMeshUpdate to be dna-independent
+
+	//the rig behavior has changed, we need to force re-initializing of RigLogic
+	//this will set RigLogic RigUnit to initial state
+	SkelMeshComponent->InitAnim(true);
+}
+
+void USkelMeshDNAUtils::UpdateSourceData(USkeletalMesh* SkelMesh)
+{
+	IMeshUtilities& MeshUtilities = FModuleManager::Get().LoadModuleChecked<IMeshUtilities>("MeshUtilities");
+	// Source data must be updated during cooking.
+	SkelMesh->EmptyAllImportData();
+	MeshUtilities.CreateImportDataFromLODModel(SkelMesh);
+}
+
+#endif // WITH_EDITORONLY_DATA
diff --git a/Engine/Plugins/Animation/RigLogic/Source/RigLogicModule/Public/DNAToSkelMeshMap.h b/Engine/Plugins/Animation/RigLogic/Source/RigLogicModule/Public/DNAToSkelMeshMap.h
new file mode 100644
--- /dev/null
+++ b/Engine/Plugins/Animation/RigLogic/Source/RigLogicModule/Public/DNAToSkelMeshMap.h
@@ -0,0 +1,83 @@
+// Copyright Epic Games, Inc. All Rights Reserved.
+
+#pragma once
+
+#include "CoreMinimal.h"
+#include "HAL/FileManager.h"
+#include "Serialization/Archive.h"
+#include "UObject/NoExportTypes.h"
+
+#include "DNAReader.h"
+
+DECLARE_LOG_CATEGORY_EXTERN(LogDNAToSkelMeshMap, Log, All);
+
+class USkeletalMesh;
+
+
+struct FDNABlendShapeTarget
+{
+	FDNABlendShapeTarget()
+		: MeshIndex(INDEX_NONE)
+		, TargetIndex(INDEX_NONE)
+	{
+	}
+
+	FDNABlendShapeTarget(const int16 InMeshIndex, const int16 InTargetIndex)
+		: MeshIndex(InMeshIndex)
+		, TargetIndex(InTargetIndex)
+	{
+	}
+
+	int16 MeshIndex;
+	int16 TargetIndex;
+};
+
+/** An object holding the mappings from DNA to UE indices, needed for updating a character from DNA file */
+class RIGLOGICMODULE_API FDNAToSkelMeshMap
+{
+public:
+
+	FDNAToSkelMeshMap() {}
+
+#if WITH_EDITORONLY_DATA
+	/** Creates vertex map between source DNA and target SkelMesh, and stores DNAReader
+	  * Enough for updating neutral meshes, but not for updating joints and morphtargets; call the next method for that
+	  * TODO: When ImportedModel is added to runtime, WITH_EDITORONLY_DATA should be removed**/
+	void InitBaseMesh(IDNAReader* InDNAReader, USkeletalMesh* SkelMesh);
+	/** Creates map for updating joints; needs to be called before playing animation on the rig **/
+	void MapJoints(IDNAReader* InDNAReader);
+	/** Creates map for updating morph targets; needs to be called before playing animation on the rig **/
+	void MapMorphTargets(IDNAReader* InDNAReader);
+#endif //WITH_EDITORONLY_DATA
+	/** Creates mappings between source DNA stored in DNAAsset AssetUserData attached to SkeletalMesh;
+		Uses FSkeletalMeshDNAReader to be able to read geometry which, in this case, is NOT present in the DNAAsset **/
+	bool InitFromDNAAsset(USkeletalMesh* SkelMesh);
+	void InitVertexMap(IDNAReader* InDNAReader);
+
+
+	int32 GetUEBoneIndex(int32 RLJointIndex)
+	{
+		return RLJointToUEBoneIndices[RLJointIndex];
+	}
+
+	TArray<FDNABlendShapeTarget>& GetMeshBlendShapeTargets()
+	{
+		return MeshBlendShapeTargets;
+	}
+
+	TArray<TArray<int32>> ImportVtxToDNAMeshIndex;
+	TArray<TArray<int32>> ImportVtxToDNAVtxIndex;
+	TArray<TArray<TArray<int32>>> ImportDNAVtxToUEVtxIndex;
+	/** Map for keeping all overlapping vertices for each vertex of each section of each LOD Model */
+	TArray<TArray<TArray<TArray<int32>>>> OverlappingVertices;
+	TArray<TArray<int32>> UEVertexToSectionIndices;
+
+private:
+
+	USkeletalMesh* TargetSkelMesh = nullptr;
+
+	TArray<int32> RLJointToUEBoneIndices;  // UFC joint index to UE bone index
+	TArray<FDNABlendShapeTarget> MeshBlendShapeTargets; // Reverse mapping.
+	TArray<uint16> DNASkinnedJointIndices;
+
+};
diff --git a/Engine/Plugins/Animation/RigLogic/Source/RigLogicModule/Public/SkelMeshDNAReader.h b/Engine/Plugins/Animation/RigLogic/Source/RigLogicModule/Public/SkelMeshDNAReader.h
new file mode 100644
--- /dev/null
+++ b/Engine/Plugins/Animation/RigLogic/Source/RigLogicModule/Public/SkelMeshDNAReader.h
@@ -0,0 +1,135 @@
+// Copyright Epic Games, Inc. All Rights Reserved.
+
+#pragma once
+
+#include "DNAReader.h"
+
+#include "CoreMinimal.h"
+#include "Engine/AssetUserData.h"
+#include "Interfaces/Interface_AssetUserData.h"
+#include "HAL/FileManager.h"
+#include "Serialization/Archive.h"
+#include "UObject/NoExportTypes.h"
+
+class UDNAAsset;
+
+DECLARE_LOG_CATEGORY_EXTERN(LogSkelMeshDNAReader, Log, All);
+
+class RIGLOGICMODULE_API FSkelMeshDNAReader: public IDNAReader
+{
+public:
+	explicit FSkelMeshDNAReader(UDNAAsset* DNAAsset);
+
+	// Descriptor
+	FString GetName() const override;
+	EArchetype GetArchetype() const override;
+	EGender GetGender() const override;
+	uint16 GetAge() const override;
+	uint32 GetMetaDataCount() const override;
+	FString GetMetaDataKey(uint32 Index) const override;
+	FString GetMetaDataValue(const FString& Key) const override;
+	ETranslationUnit GetTranslationUnit() const override;
+	ERotationUnit GetRotationUnit() const override;
+	FCoordinateSystem GetCoordinateSystem() const override;
+	uint16 GetLODCount() const override;
+	uint16 GetDBMaxLOD() const override;
+	FString GetDBComplexity() const override;
+	FString GetDBName() const override;
+	// Definition
+	uint16 GetGUIControlCount() const override;
+	FString GetGUIControlName(uint16 Index) const override;
+	uint16 GetRawControlCount() const override;
+	FString GetRawControlName(uint16 Index) const override;
+	uint16 GetJointCount() const override;
+	FString GetJointName(uint16 Index) const override;
+	uint16 GetJointIndexListCount() const override;
+	TArrayView<const uint16> GetJointIndicesForLOD(uint16 LOD) const override;
+	uint16 GetJointParentIndex(uint16 Index) const override;
+	uint16 GetBlendShapeChannelCount() const override;
+	FString GetBlendShapeChannelName(uint16 Index) const override;
+	uint16 GetBlendShapeChannelIndexListCount() const override;
+	TArrayView<const uint16> GetBlendShapeChannelIndicesForLOD(uint16 LOD) const override;
+	uint16 GetAnimatedMapCount() const override;
+	FString GetAnimatedMapName(uint16 Index) const override;
+	uint16 GetAnimatedMapIndexListCount() const override;
+	TArrayView<const uint16> GetAnimatedMapIndicesForLOD(uint16 LOD) const override;
+	uint16 GetMeshCount() const override;
+	FString GetMeshName(uint16 Index) const override;
+	uint16 GetMeshIndexListCount() const override;
+	TArrayView<const uint16> GetMeshIndicesForLOD(uint16 LOD) const override;
+	uint16 GetMeshBlendShapeChannelMappingCount() const override;
+	FMeshBlendShapeChannelMapping GetMeshBlendShapeChannelMapping(uint16 Index) const override;
+	TArrayView<const uint16> GetMeshBlendShapeChannelMappingIndicesForLOD(uint16 LOD) const override;
+	FVector GetNeutralJointTranslation(uint16 Index) const override;
+	FVector GetNeutralJointRotation(uint16 Index) const override;
+	// Behavior
+	TArrayView<const uint16> GetGUIToRawInputIndices() const override;
+	TArrayView<const uint16> GetGUIToRawOutputIndices() const override;
+	TArrayView<const float> GetGUIToRawFromValues() const override;
+	TArrayView<const float> GetGUIToRawToValues() const override;
+	TArrayView<const float> GetGUIToRawSlopeValues() const override;
+	TArrayView<const float> GetGUIToRawCutValues() const override;
+	uint16 GetPSDCount() const override;
+	TArrayView<const uint16> GetPSDRowIndices() const override;
+	TArrayView<const uint16> GetPSDColumnIndices() const override;
+	TArrayView<const float> GetPSDValues() const override;
+	uint16 GetJointRowCount() const override;
+	uint16 GetJointColumnCount() const override;
+	TArrayView<const uint16> GetJointVariableAttributeIndices(uint16 LOD) const override;
+	uint16 GetJointGroupCount() const override;
+	TArrayView<const uint16> GetJointGroupLODs(uint16 JointGroupIndex) const override;
+	TArrayView<const uint16> GetJointGroupInputIndices(uint16 JointGroupIndex) const override;
+	TArrayView<const uint16> GetJointGroupOutputIndices(uint16 JointGroupIndex) const override;
+	TArrayView<const float> GetJointGroupValues(uint16 JointGroupIndex) const override;
+	TArrayView<const uint16> GetJointGroupJointIndices(uint16 JointGroupIndex) const override;
+	TArrayView<const uint16> GetBlendShapeChannelLODs() const override;
+	TArrayView<const uint16> GetBlendShapeChannelOutputIndices() const override;
+	TArrayView<const uint16> GetBlendShapeChannelInputIndices() const override;
+	TArrayView<const uint16> GetAnimatedMapLODs() const override;
+	TArrayView<const uint16> GetAnimatedMapInputIndices() const override;
+	TArrayView<const uint16> GetAnimatedMapOutputIndices() const override;
+	TArrayView<const float> GetAnimatedMapFromValues() const override;
+	TArrayView<const float> GetAnimatedMapToValues() const override;
+	TArrayView<const float> GetAnimatedMapSlopeValues() const override;
+	TArrayView<const float> GetAnimatedMapCutValues() const override;
+	// Geometry
+	uint32 GetVertexPositionCount(uint16 MeshIndex) const override;
+	FVector GetVertexPosition(uint16 MeshIndex, uint32 VertexIndex) const override;
+	TArrayView<const float> GetVertexPositionXs(uint16 MeshIndex) const override;
+	TArrayView<const float> GetVertexPositionYs(uint16 MeshIndex) const override;
+	TArrayView<const float> GetVertexPositionZs(uint16 MeshIndex) const override;
+	uint32 GetVertexTextureCoordinateCount(uint16 MeshIndex) const override;
+	FTextureCoordinate GetVertexTextureCoordinate(uint16 MeshIndex, uint32 TextureCoordinateIndex) const override;
+	uint32 GetVertexNormalCount(uint16 MeshIndex) const override;
+	FVector GetVertexNormal(uint16 MeshIndex, uint32 NormalIndex) const override;
+	uint32 GetFaceCount(uint16 MeshIndex) const override;
+	TArrayView<const uint32> GetFaceVertexLayoutIndices(uint16 MeshIndex, uint32 FaceIndex) const override;
+	uint32 GetVertexLayoutCount(uint16 MeshIndex) const override;
+	FVertexLayout GetVertexLayout(uint16 MeshIndex, uint32 LayoutIndex) const override;
+	uint16 GetMaximumInfluencePerVertex(uint16 MeshIndex) const override;
+	uint32 GetSkinWeightsCount(uint16 MeshIndex) const override;
+	TArrayView<const float> GetSkinWeightsValues(uint16 MeshIndex, uint32 VertexIndex) const override;
+	TArrayView<const uint16> GetSkinWeightsJointIndices(uint16 MeshIndex, uint32 VertexIndex) const override;
+	uint16 GetBlendShapeTargetCount(uint16 MeshIndex) const override;
+	uint16 GetBlendShapeChannelIndex(uint16 MeshIndex, uint16 BlendShapeTargetIndex) const override;
+	uint32 GetBlendShapeTargetDeltaCount(uint16 MeshIndex, uint16 BlendShapeTargetIndex) const override;
+	FVector GetBlendShapeTargetDelta(uint16 MeshIndex, uint16 BlendShapeTargetIndex, uint32 DeltaIndex) const override;
+	TArrayView<const float> GetBlendShapeTargetDeltaXs(uint16 MeshIndex, uint16 BlendShapeTargetIndex) const override;
+	TArrayView<const float> GetBlendShapeTargetDeltaYs(uint16 MeshIndex, uint16 BlendShapeTargetIndex) const override;
+	TArrayView<const float> GetBlendShapeTargetDeltaZs(uint16 MeshIndex, uint16 BlendShapeTargetIndex) const override;
+	TArrayView<const uint32> GetBlendShapeTargetVertexIndices(uint16 MeshIndex, uint16 BlendShapeTargetIndex) const override;
+	void Unload(EDNADataLayer /*unused*/) override;
+
+private:
+	dna::Reader* Unwrap() const override;
+
+private:
+	/** Both BehaviorReader and DesignDataStreamReader are StreamReaders from DNAAsset 	   
+	  * split out into run-time and in-editor parts from a full DNA that is either:
+	  * 1) imported manually into SkeletalMesh asset through ContentBrowser
+	  *	2) overwritten by GeneSplicer (GeneSplicerDNAReader) in a transient SkeletalMesh copy
+	  * They both just borrow DNAAsset's readers and are not owned by SkelMeshDNAReader
+	 **/
+	IBehaviorReader* BehaviorStreamReader = nullptr;
+	IGeometryReader* GeometryStreamReader = nullptr;
+};
diff --git a/Engine/Plugins/Animation/RigLogic/Source/RigLogicModule/Public/SkelMeshDNAUtils.h b/Engine/Plugins/Animation/RigLogic/Source/RigLogicModule/Public/SkelMeshDNAUtils.h
new file mode 100644
--- /dev/null
+++ b/Engine/Plugins/Animation/RigLogic/Source/RigLogicModule/Public/SkelMeshDNAUtils.h
@@ -0,0 +1,84 @@
+// Copyright Epic Games, Inc. All Rights Reserved.
+
+#pragma once
+
+#include "DNAToSkelMeshMap.h"
+#include "Engine/SkeletalMesh.h"
+#include "SkelMeshDNAUtils.generated.h"
+
+DECLARE_LOG_CATEGORY_EXTERN(LogDNAUtils, Log, All);
+
+class IDNAReader;
+class FDNAToSkelMeshMap;
+
+UENUM()
+enum class ELodUpdateOption : uint8
+{
+	LOD0Only, 		// LOD0 Only
+	LOD1AndHigher, 	// LOD1 and higher
+	All				// All LODs
+};
+
+/** A utility class for updating SkeletalMesh joints, base mesh, morph targets and skin weights according to DNA data.
+  * After the update, the render data is re-chunked.
+ **/
+UCLASS(transient)
+class RIGLOGICMODULE_API USkelMeshDNAUtils: public UObject
+{
+	GENERATED_UCLASS_BODY()
+
+public:
+
+	/** Prepare context object that will allow mapping of DNA structures to SkelMesh ones for updating **/
+	static FDNAToSkelMeshMap* CreateMapForUpdatingNeutralMesh(IDNAReader* DNAReader, USkeletalMesh* SkelMesh);
+
+#if WITH_EDITORONLY_DATA
+	/** Updates the positions, orientation and scale in the joint hierarchy using the data from DNA file **/
+	static void UpdateJoints(USkeletalMesh* SkelMesh, IDNAReader* DNAReader, FDNAToSkelMeshMap* DNAToSkelMeshMap);
+	/** Updates the base mesh vertex positions for all mesh sections of all LODs, using the data from DNA file 
+	  * NOTE: Not calling RebuildRenderData automatically, it needs to be called explicitly after the first update
+	  *       As the topology doesn't change, for subsequent updates it can be ommited to gain performance **/
+	static void UpdateBaseMesh(USkeletalMesh* SkelMesh, IDNAReader* DNAReader, FDNAToSkelMeshMap* DNAToSkelMeshMap, ELodUpdateOption UpdateOption = ELodUpdateOption::LOD0Only);
+	/** Updates the morph targets for all mesh sections of LODs, using the data from DNA file **/
+	static void UpdateMorphTargets(USkeletalMesh* SkelMesh, IDNAReader* DNAReader, FDNAToSkelMeshMap* DNAToSkelMeshMap, ELodUpdateOption UpdateOption = ELodUpdateOption::LOD0Only);
+	/** Updates the skin weights for all LODs using the data from DNA file **/
+	static void UpdateSkinWeights(USkeletalMesh* SkelMesh, IDNAReader* DNAReader, FDNAToSkelMeshMap* DNAToSkelMeshMap, ELodUpdateOption UpdateOption = ELodUpdateOption::LOD0Only);
+	/** Rechunks the mesh after the update **/
+	static void RebuildRenderData(USkeletalMesh* SkelMesh);
+	/** Re-initialize vertex positions for rendering after the update **/
+	static void RebuildRenderData_VertexPosition(USkeletalMesh* SkelMesh);
+	/** Update joint behavior **/
+	static void UpdateJointBehavior(USkeletalMeshComponent* SkelMeshComponent);
+#endif // WITH_EDITORONLY_DATA
+	/** Converts DNA vertex coordinates to UE4 coordinate system **/
+	inline static FVector ConvertDNAVertexToUE4CoordSystem(FVector VertexPositionInDNA)
+	{
+		return FVector{-VertexPositionInDNA.X, VertexPositionInDNA.Y, -VertexPositionInDNA.Z};
+	}
+
+	/** Converts UE4 coordinate system to DNA vertex coordinates **/
+	inline static FVector ConvertUE4CoordSystemToDNAVertex(FVector VertexPositionInUE4)
+	{
+		return FVector{-VertexPositionInUE4.X, VertexPositionInUE4.Y, -VertexPositionInUE4.Z};
+	}
+
+	/** Updates source skeleton data for the purpose of character cooking and export.  */
+	static void UpdateSourceData(USkeletalMesh* SkelMesh);
+	static void UpdateSourceData(USkeletalMesh* SkelMesh, IDNAReader* DNAReader, FDNAToSkelMeshMap* DNAToSkelMeshMap);
+	
+private:
+
+	inline static void GetLODRange(ELodUpdateOption UpdateOption, const int32& LODNum, int32& LODStart, int32& LODRangeSize)
+	{
+		LODStart = 0;
+		LODRangeSize = LODNum;
+		if (UpdateOption == ELodUpdateOption::LOD1AndHigher)
+		{
+			LODStart = 1;
+		}
+		else if (UpdateOption == ELodUpdateOption::LOD0Only && LODRangeSize > 0)
+		{
+			LODRangeSize = 1;
+		}
+	}
+};
diff --git a/Engine/Plugins/Animation/RigLogic/Source/RigLogicModule/RigLogicModule.Build.cs b/Engine/Plugins/Animation/RigLogic/Source/RigLogicModule/RigLogicModule.Build.cs
--- a/Engine/Plugins/Animation/RigLogic/Source/RigLogicModule/RigLogicModule.Build.cs
+++ b/Engine/Plugins/Animation/RigLogic/Source/RigLogicModule/RigLogicModule.Build.cs
@@ -30,9 +30,10 @@ namespace UnrealBuildTool.Rules
                     "RigLogicLib",
                     "RigVM",
                     "Projects"
-                }
+				}
             );
 
+
             if (Target.Type == TargetType.Editor)
             {
                 PublicDependencyModuleNames.Add("UnrealEd");
@@ -42,9 +43,12 @@ namespace UnrealBuildTool.Rules
 
             PrivateDependencyModuleNames.AddRange(
                 new string[]
-                {
-                    "AnimationCore"
-                }
+				{
+					"AnimationCore",
+					"SkeletalMeshUtilitiesCommon",
+					"RHI",
+					"RenderCore",
+				}
             );
         }
     }
