RigLogicModule: moved some of private dependencies to be included only in Editor.

diff --git a/Engine/Plugins/Animation/RigLogic/Source/RigLogicModule/RigLogicModule.Build.cs b/Engine/Plugins/Animation/RigLogic/Source/RigLogicModule/RigLogicModule.Build.cs
--- a/Engine/Plugins/Animation/RigLogic/Source/RigLogicModule/RigLogicModule.Build.cs
+++ b/Engine/Plugins/Animation/RigLogic/Source/RigLogicModule/RigLogicModule.Build.cs
@@ -39,15 +39,16 @@ namespace UnrealBuildTool.Rules
                 PublicDependencyModuleNames.Add("UnrealEd");
                 PublicDependencyModuleNames.Add("EditorFramework");
                 PublicDependencyModuleNames.Add("MessageLog");
-            }
+
+				PrivateDependencyModuleNames.Add("SkeletalMeshUtilitiesCommon");
+				PrivateDependencyModuleNames.Add("RHI");
+				PrivateDependencyModuleNames.Add("RenderCore");
+			}
 
             PrivateDependencyModuleNames.AddRange(
                 new string[]
 				{
-					"AnimationCore",
-					"SkeletalMeshUtilitiesCommon",
-					"RHI",
-					"RenderCore",
+					"AnimationCore"
 				}
             );
         }
