Fix for static analysis warnings.

#rnx

diff --git a/Engine/Plugins/Animation/RigLogic/Source/RigLogicModule/Private/SkelMeshDNAUtils.cpp b/Engine/Plugins/Animation/RigLogic/Source/RigLogicModule/Private/SkelMeshDNAUtils.cpp
--- a/Engine/Plugins/Animation/RigLogic/Source/RigLogicModule/Private/SkelMeshDNAUtils.cpp
+++ b/Engine/Plugins/Animation/RigLogic/Source/RigLogicModule/Private/SkelMeshDNAUtils.cpp
@@ -272,7 +272,6 @@ void USkelMeshDNAUtils::UpdateSkinWeights(USkeletalMesh* SkelMesh, IDNAReader* D
 
 	bool InfluenceMismatch = false;
 	//Set a threshold a bit smaller then 1/255
-	constexpr float MINWEIGHT = 0.9999f / 255.0f;
 	int32 LODStart;
 	int32 LODRangeSize;
 	GetLODRange(UpdateOption, ImportedModel->LODModels.Num(), LODStart, LODRangeSize);
@@ -282,23 +281,19 @@ void USkelMeshDNAUtils::UpdateSkinWeights(USkeletalMesh* SkelMesh, IDNAReader* D
 		for (int32 SectionIndex = 0; SectionIndex < LODModel.Sections.Num(); SectionIndex++)
 		{
 			FSkelMeshSection& Section = LODModel.Sections[SectionIndex];
-			int32 DNAMeshIndex = DNAToSkelMeshMap->ImportVtxToDNAMeshIndex[LODIndex][Section.GetVertexBufferIndex()];
+			const int32 DNAMeshIndex = DNAToSkelMeshMap->ImportVtxToDNAMeshIndex[LODIndex][Section.GetVertexBufferIndex()];
 
 			const int32 NumEngineVertices = Section.GetNumVertices();
 			for (int32 VertexIndex = 0; VertexIndex < NumEngineVertices; VertexIndex++)
 			{
 				const int32 VertexBufferIndex = VertexIndex + Section.GetVertexBufferIndex();
-				int32 DNAVertexIndex = DNAToSkelMeshMap->ImportVtxToDNAVtxIndex[LODIndex][VertexBufferIndex];
+				const int32 DNAVertexIndex = DNAToSkelMeshMap->ImportVtxToDNAVtxIndex[LODIndex][VertexBufferIndex];
 
 				if (DNAVertexIndex < 0) continue; // Skip vertex not in DNA.
 
 				TArrayView<const float> DNASkinWeights = DNAReader->GetSkinWeightsValues(DNAMeshIndex, DNAVertexIndex);
 				TArrayView<const uint16> DNASkinJoints = DNAReader->GetSkinWeightsJointIndices(DNAMeshIndex, DNAVertexIndex);
-				uint16 SkinJointNum = DNASkinJoints.Num();
-
-				uint32 TotalWeight = 0;  // store all influences to vertex to ensure they add up to 255 (fix rounding errors)
-				uint16 MaxInfluenceIndex = 0;
-				uint8 MaxInfluenceWeight = 0;
+				int32 SkinJointNum = DNASkinJoints.Num();
 
 				FSoftSkinVertex& Vertex = Section.SoftVertices[VertexIndex];
 
@@ -307,64 +302,39 @@ void USkelMeshDNAUtils::UpdateSkinWeights(USkeletalMesh* SkelMesh, IDNAReader* D
 					//UE_LOG(LogDNAUtils, Log, TEXT("Skipping vertex UE: %d, DNA: %d"), VertexIndex, DNAVertexIndex );
 					continue;
 				}
-				TArray<float> NormalizedSkinWeights;
-				NormalizedSkinWeights.Reserve(SkinJointNum);
-				TArray<uint16> NormalizedJoints;
-				NormalizedJoints.Reserve(SkinJointNum);
-				float TotalInfluence = 0.0;
-				// Second step is to filter out all the influences below the MIN WEIGHT or above total influence limit.
-				for (uint16 i = 0; i < SkinJointNum; i++)
-				{
-					if (DNASkinWeights[i] > MINWEIGHT && i < MAX_TOTAL_INFLUENCES)
-					{
-						NormalizedSkinWeights.Add(DNASkinWeights[i]);
-						NormalizedJoints.Add(DNASkinJoints[i]);
-						TotalInfluence += DNASkinWeights[i];
-					}
-				}
-				SkinJointNum = NormalizedJoints.Num();
-				if (SkinJointNum > 0 && (TotalInfluence != 1.0f))
-				{
-					// Missing fractions of influence weights have to be assigned equally along the existing influences.
-					float OneOverTotalWeight = 1.f / TotalInfluence;
-					for (uint16 r = 0; r < SkinJointNum; r++)
-					{
-						NormalizedSkinWeights[r] *= OneOverTotalWeight;
-					}
-				}
-				// Reset all influences that are not covered by DNA data.
-				for (uint16 i = SkinJointNum; i < MAX_TOTAL_INFLUENCES; i++)
-				{
-					Vertex.InfluenceBones[i] = 0;
-					Vertex.InfluenceWeights[i] = 0;
-				}
-				for (uint16 InfluenceIndex = 0; InfluenceIndex < SkinJointNum; ++InfluenceIndex)
+				
+				FBoneIndexType InfluenceBones[MAX_TOTAL_INFLUENCES];
+				float InfluenceWeights[MAX_TOTAL_INFLUENCES];
+
+				int32 InfluenceIndex = 0;
+				for (int32 Index = 0; Index < FMath::Min(SkinJointNum, MAX_TOTAL_INFLUENCES); ++Index)
 				{
-					uint8 EngineWeight = 0;
 					// Find Engine bone for corresponding DNAJoint for the same influence.
-					int32 UpdatedBoneId = DNAToSkelMeshMap->GetUEBoneIndex(NormalizedJoints[InfluenceIndex]);
+					const int32 UpdatedBoneId = DNAToSkelMeshMap->GetUEBoneIndex(DNASkinJoints[Index]);
 					// BoneMap holds subset of bones belonging to current section.
-					int32 BoneMapIndex = Section.BoneMap.Find(UpdatedBoneId);
+					const int32 BoneMapIndex = Section.BoneMap.Find(UpdatedBoneId);
 
 					// Update which bone in the subset influences this vertex.
-					Vertex.InfluenceBones[InfluenceIndex] = BoneMapIndex;
 					if (BoneMapIndex != INDEX_NONE)
 					{
-						// Update influence weight.
-						EngineWeight = (uint16)(NormalizedSkinWeights[InfluenceIndex] * UE::AnimationCore::MaxRawBoneWeightFloat);
+						InfluenceBones[InfluenceIndex] = BoneMapIndex;
+						InfluenceWeights[InfluenceIndex] = DNASkinWeights[Index];
+						InfluenceIndex++;
 					}
+				}
 
-					Vertex.InfluenceWeights[InfluenceIndex] = EngineWeight;
-					TotalWeight += Vertex.InfluenceWeights[InfluenceIndex];
+				using namespace UE::AnimationCore;
+				FMemory::Memzero(Vertex.InfluenceBones);
+				FMemory::Memzero(Vertex.InfluenceWeights);
 
-					if (EngineWeight > MaxInfluenceWeight)
-					{
-						MaxInfluenceIndex = InfluenceIndex;
-						MaxInfluenceWeight = EngineWeight;
-					}
+				FBoneWeights BoneWeights = FBoneWeights::Create(InfluenceBones, InfluenceWeights, InfluenceIndex);
+				InfluenceIndex = 0;
+				for (FBoneWeight BoneWeight: BoneWeights)
+				{
+					Vertex.InfluenceBones[InfluenceIndex] = BoneWeight.GetBoneIndex();
+					Vertex.InfluenceWeights[InfluenceIndex] = BoneWeight.GetRawWeight();
+					InfluenceIndex++;
 				}
-				// Add missing fraction to fill up the bone weights.
-				Vertex.InfluenceWeights[MaxInfluenceIndex] += UE::AnimationCore::MaxRawBoneWeight - TotalWeight;
 			}
 		}
 	}
diff --git a/Engine/Source/Runtime/Engine/Private/SkeletalMeshLODRenderData.cpp b/Engine/Source/Runtime/Engine/Private/SkeletalMeshLODRenderData.cpp
--- a/Engine/Source/Runtime/Engine/Private/SkeletalMeshLODRenderData.cpp
+++ b/Engine/Source/Runtime/Engine/Private/SkeletalMeshLODRenderData.cpp
@@ -384,7 +384,7 @@ void FSkeletalMeshLODRenderData::BuildFromLODModel(const FSkeletalMeshLODModel*
 		{
 			uint32	TotalInfluenceWeight = 0;
 			int32	MaxInfluenceIndex = 0;
-			for (; Vertex.InfluenceWeights[MaxInfluenceIndex] && MaxInfluenceIndex < MAX_TOTAL_INFLUENCES; MaxInfluenceIndex++)
+			for (; MaxInfluenceIndex < MAX_TOTAL_INFLUENCES && Vertex.InfluenceWeights[MaxInfluenceIndex]; MaxInfluenceIndex++)
 			{
 				InfluenceWeights[MaxInfluenceIndex] = static_cast<uint8>(Vertex.InfluenceWeights[MaxInfluenceIndex] >> 8);
 				if (InfluenceWeights[MaxInfluenceIndex] == 0)
