Removed duplicate of SkelMeshDNAUtils classes from MetaHuman plugin

diff --git a/Engine/Plugins/Animation/RigLogic/Source/RigLogicModule/Private/SkelMeshDNAUtils.cpp b/Engine/Plugins/Animation/RigLogic/Source/RigLogicModule/Private/SkelMeshDNAUtils.cpp
--- a/Engine/Plugins/Animation/RigLogic/Source/RigLogicModule/Private/SkelMeshDNAUtils.cpp
+++ b/Engine/Plugins/Animation/RigLogic/Source/RigLogicModule/Private/SkelMeshDNAUtils.cpp
@@ -42,7 +42,7 @@ USkelMeshDNAUtils::USkelMeshDNAUtils(const FObjectInitializer& ObjectInitializer
 {
 }
 
-FDNAToSkelMeshMap* USkelMeshDNAUtils::CreateMapForUpdatingNeutralMesh(IDNAReader* DNAReader, USkeletalMesh* SkelMesh)
+FDNAToSkelMeshMap* USkelMeshDNAUtils::CreateMapForUpdatingNeutralMesh(IDNAReader* InDNAReader, USkeletalMesh* InSkelMesh)
 {
 #if WITH_EDITORONLY_DATA
 	FDNAToSkelMeshMap* DNAToSkelMeshMap = new FDNAToSkelMeshMap();
@@ -50,7 +50,19 @@ FDNAToSkelMeshMap* USkelMeshDNAUtils::CreateMapForUpdatingNeutralMesh(IDNAReader
 	//only vertex map is initialized in this pass so we can mix neutral meshes fast (e.g. on slider move);
 	//playing animation on a such a mesh requires updating joints and skin weights
 	//getting full quality animation requires mixing morph targets too
-	DNAToSkelMeshMap->InitBaseMesh(DNAReader, SkelMesh);
+	DNAToSkelMeshMap->InitBaseMesh(InDNAReader, InSkelMesh);
+	return DNAToSkelMeshMap;
+#else
+	return nullptr;
+#endif // WITH_EDITORONLY_DATA
+}
+
+FDNAToSkelMeshMap* USkelMeshDNAUtils::CreateMapForUpdatingNeutralMesh(USkeletalMesh* InSkelMesh)
+{
+#if WITH_EDITORONLY_DATA
+	FDNAToSkelMeshMap* DNAToSkelMeshMap = new FDNAToSkelMeshMap();
+
+	DNAToSkelMeshMap->InitFromDNAAsset(InSkelMesh);
 	return DNAToSkelMeshMap;
 #else
 	return nullptr;
@@ -59,13 +71,13 @@ FDNAToSkelMeshMap* USkelMeshDNAUtils::CreateMapForUpdatingNeutralMesh(IDNAReader
 
 #if WITH_EDITORONLY_DATA
 /** Updates bind pose using joint positions from DNA. */
-void USkelMeshDNAUtils::UpdateJoints(USkeletalMesh* SkelMesh, IDNAReader* DNAReader, FDNAToSkelMeshMap* DNAToSkelMeshMap)
+void USkelMeshDNAUtils::UpdateJoints(USkeletalMesh* InSkelMesh, IDNAReader* InDNAReader, FDNAToSkelMeshMap* InDNAToSkelMeshMap)
 {
 	{	// Scoping of RefSkelModifier
-		FReferenceSkeletonModifier RefSkelModifier(SkelMesh->GetRefSkeleton(), SkelMesh->GetSkeleton());
+		FReferenceSkeletonModifier RefSkelModifier(InSkelMesh->GetRefSkeleton(), InSkelMesh->GetSkeleton());
 
 		// copy here
-		TArray<FTransform> RawBonePose = SkelMesh->GetRefSkeleton().GetRawRefBonePose();
+		TArray<FTransform> RawBonePose = InSkelMesh->GetRefSkeleton().GetRawRefBonePose();
 
 		// When we are mounting the head to different bodies than female average, we need to use
 		// component space, as the joint to which the head root is snapped to will be on a different
@@ -73,21 +85,21 @@ void USkelMeshDNAUtils::UpdateJoints(USkeletalMesh* SkelMesh, IDNAReader* DNARea
 
 		// calculate component space ahead of current transform
 		TArray<FTransform> ComponentTransforms;
-		FAnimationRuntime::FillUpComponentSpaceTransforms(SkelMesh->GetRefSkeleton(), RawBonePose, ComponentTransforms);
+		FAnimationRuntime::FillUpComponentSpaceTransforms(InSkelMesh->GetRefSkeleton(), RawBonePose, ComponentTransforms);
 
-		const TArray<FMeshBoneInfo>& RawBoneInfo = SkelMesh->GetRefSkeleton().GetRawRefBoneInfo();
+		const TArray<FMeshBoneInfo>& RawBoneInfo = InSkelMesh->GetRefSkeleton().GetRawRefBoneInfo();
 
 		// Skipping root joint (index 0) to avoid blinking of the mesh due to bounding box issue
-		for (uint16 JointIndex = 0; JointIndex < DNAReader->GetJointCount(); JointIndex++)
+		for (uint16 JointIndex = 0; JointIndex < InDNAReader->GetJointCount(); JointIndex++)
 		{
-			int32 BoneIndex = DNAToSkelMeshMap->GetUEBoneIndex(JointIndex);
+			int32 BoneIndex = InDNAToSkelMeshMap->GetUEBoneIndex(JointIndex);
 
 			FTransform DNATransform = FTransform::Identity;
 
 			// Updating bind pose affects just translations.
-			FVector Translate = DNAReader->GetNeutralJointTranslation(JointIndex);
+			FVector Translate = InDNAReader->GetNeutralJointTranslation(JointIndex);
 			DNATransform.SetTranslation(Translate);
-			FVector RotationVector = DNAReader->GetNeutralJointRotation(JointIndex);
+			FVector RotationVector = InDNAReader->GetNeutralJointRotation(JointIndex);
 			FRotator Rotation(RotationVector.X, RotationVector.Y, RotationVector.Z);
 
 			// Joint 0 (spine_04) is a root of GeneSplicer joint hierarchy, and is a special case
@@ -96,7 +108,7 @@ void USkelMeshDNAUtils::UpdateJoints(USkeletalMesh* SkelMesh, IDNAReader* DNARea
 			// 3) the head joints below it in the skeletal mesh are not spliced, as they are not in DNA,
 			//    so they will retain female average transforms
 
-			if (DNAReader->GetJointParentIndex(JointIndex) == JointIndex)  //parent to itself
+			if (InDNAReader->GetJointParentIndex(JointIndex) == JointIndex)  //parent to itself
 			{
 				Rotation.Pitch += 90;
 				DNATransform.SetRotation(Rotation.Quaternion());
@@ -138,37 +150,37 @@ void USkelMeshDNAUtils::UpdateJoints(USkeletalMesh* SkelMesh, IDNAReader* DNARea
 		}
 	}
 
-	SkelMesh->GetRefBasesInvMatrix().Reset();
-	SkelMesh->CalculateInvRefMatrices(); // Needs to be called after RefSkelModifier is destroyed
+	InSkelMesh->GetRefBasesInvMatrix().Reset();
+	InSkelMesh->CalculateInvRefMatrices(); // Needs to be called after RefSkelModifier is destroyed
 }
 
 /** Updates base mesh vertices using data from DNA. */
-void USkelMeshDNAUtils::UpdateBaseMesh(USkeletalMesh* SkelMesh, IDNAReader* DNAReader, FDNAToSkelMeshMap* DNAToSkelMeshMap, ELodUpdateOption UpdateOption)
+void USkelMeshDNAUtils::UpdateBaseMesh(USkeletalMesh* InSkelMesh, IDNAReader* InDNAReader, FDNAToSkelMeshMap* InDNAToSkelMeshMap, ELodUpdateOption InUpdateOption)
 {
 
-	FSkeletalMeshModel* ImportedModel = SkelMesh->GetImportedModel();
+	FSkeletalMeshModel* ImportedModel = InSkelMesh->GetImportedModel();
 	// Expects vertex map to be initialized beforehand	
 	int32 LODStart;
 	int32 LODRangeSize;
-	GetLODRange(UpdateOption, ImportedModel->LODModels.Num(), LODStart, LODRangeSize);
+	GetLODRange(InUpdateOption, ImportedModel->LODModels.Num(), LODStart, LODRangeSize);
 	for (int32 LODIndex = LODStart; LODIndex < LODRangeSize; LODIndex++)
 	{
 		FSkeletalMeshLODModel& LODModel = ImportedModel->LODModels[LODIndex];
 		int32 SectionIndex = 0;
 		for (FSkelMeshSection& Section : LODModel.Sections)
 		{
-			int32& DNAMeshIndex = DNAToSkelMeshMap->ImportVtxToDNAMeshIndex[LODIndex][Section.GetVertexBufferIndex()];
+			int32& DNAMeshIndex = InDNAToSkelMeshMap->ImportVtxToDNAMeshIndex[LODIndex][Section.GetVertexBufferIndex()];
 
 			const int32 NumSoftVertices = Section.GetNumVertices();
-			auto& OverlappingMap = DNAToSkelMeshMap->OverlappingVertices[LODIndex][SectionIndex];
+			auto& OverlappingMap = InDNAToSkelMeshMap->OverlappingVertices[LODIndex][SectionIndex];
 			int32 VertexBufferIndex = Section.GetVertexBufferIndex();
 			for (int32 VertexIndex = 0; VertexIndex < NumSoftVertices; VertexIndex++)
 			{
-				int32& DNAVertexIndex = DNAToSkelMeshMap->ImportVtxToDNAVtxIndex[LODIndex][VertexBufferIndex];
+				int32& DNAVertexIndex = InDNAToSkelMeshMap->ImportVtxToDNAVtxIndex[LODIndex][VertexBufferIndex];
 
 				if (DNAVertexIndex >= 0)
 				{
-					const FVector Position = DNAReader->GetVertexPosition(DNAMeshIndex, DNAVertexIndex);
+					const FVector Position = InDNAReader->GetVertexPosition(DNAMeshIndex, DNAVertexIndex);
 					FSoftSkinVertex& Vertex = Section.SoftVertices[VertexIndex];
 					Vertex.Position = FVector3f{ Position };
 
@@ -190,29 +202,29 @@ void USkelMeshDNAUtils::UpdateBaseMesh(USkeletalMesh* SkelMesh, IDNAReader* DNAR
 }
 
 /** Updates Morph Targets using Blend Shapes from DNA.  */
-void USkelMeshDNAUtils::UpdateMorphTargets(USkeletalMesh* SkelMesh, IDNAReader* DNAReader, FDNAToSkelMeshMap* DNAToSkelMeshMap, ELodUpdateOption UpdateOption)
+void USkelMeshDNAUtils::UpdateMorphTargets(USkeletalMesh* InSkelMesh, IDNAReader* InDNAReader, FDNAToSkelMeshMap* InDNAToSkelMeshMap, ELodUpdateOption InUpdateOption)
 {
-	TArray<FDNABlendShapeTarget>& MeshBlendShapeTargets = DNAToSkelMeshMap->GetMeshBlendShapeTargets();
+	TArray<FDNABlendShapeTarget>& MeshBlendShapeTargets = InDNAToSkelMeshMap->GetMeshBlendShapeTargets();
 	if (MeshBlendShapeTargets.Num() == 0)
 	{
 		UE_LOG(LogDNAUtils, Warning, TEXT("No morph targets updated!"));
 		return;
 	}
-	FSkeletalMeshModel* ImportedModel = SkelMesh->GetImportedModel();
+	FSkeletalMeshModel* ImportedModel = InSkelMesh->GetImportedModel();
 
-	uint16 ChannelCount = DNAReader->GetBlendShapeChannelCount();
-	ParallelFor(SkelMesh->GetMorphTargets().Num(), [&](int32 MorphIndex)
+	uint16 ChannelCount = InDNAReader->GetBlendShapeChannelCount();
+	ParallelFor(InSkelMesh->GetMorphTargets().Num(), [&](int32 MorphIndex)
 		{
-			UMorphTarget* MorphTarget = SkelMesh->GetMorphTargets()[MorphIndex];
+			UMorphTarget* MorphTarget = InSkelMesh->GetMorphTargets()[MorphIndex];
 			const FDNABlendShapeTarget& MeshTarget = MeshBlendShapeTargets[MorphIndex];
 			// First get all DNA deltas for current Morph Target.
-			TArrayView<const uint32> BlendShapeVertexIndices = DNAReader->GetBlendShapeTargetVertexIndices(MeshTarget.MeshIndex, MeshTarget.TargetIndex);
+			TArrayView<const uint32> BlendShapeVertexIndices = InDNAReader->GetBlendShapeTargetVertexIndices(MeshTarget.MeshIndex, MeshTarget.TargetIndex);
 			const int32 NumOfDeltas = BlendShapeVertexIndices.Num();
 			if (NumOfDeltas > 0)
 			{
 				int32 LODStart;
 				int32 LODRangeSize;
-				GetLODRange(UpdateOption, MorphTarget->GetMorphLODModels().Num(), LODStart, LODRangeSize);
+				GetLODRange(InUpdateOption, MorphTarget->GetMorphLODModels().Num(), LODStart, LODRangeSize);
 				for (int32 LODIndex = LODStart; LODIndex < LODRangeSize; LODIndex++)
 				{
 					// MorphTarget vertex indices refer to full vertex index buffer of imported mesh.
@@ -221,16 +233,16 @@ void USkelMeshDNAUtils::UpdateMorphTargets(USkeletalMesh* SkelMesh, IDNAReader*
 					MorphLODModel.bGeneratedByEngine = false;
 					MorphLODModel.SectionIndices.Empty();
 					MorphLODModel.Vertices.Reset(NumOfDeltas);
-					TArrayView<const float> DeltaXs = DNAReader->GetBlendShapeTargetDeltaXs(MeshTarget.MeshIndex, MeshTarget.TargetIndex);
-					TArrayView<const float> DeltaYs = DNAReader->GetBlendShapeTargetDeltaYs(MeshTarget.MeshIndex, MeshTarget.TargetIndex);
-					TArrayView<const float> DeltaZs = DNAReader->GetBlendShapeTargetDeltaZs(MeshTarget.MeshIndex, MeshTarget.TargetIndex);
+					TArrayView<const float> DeltaXs = InDNAReader->GetBlendShapeTargetDeltaXs(MeshTarget.MeshIndex, MeshTarget.TargetIndex);
+					TArrayView<const float> DeltaYs = InDNAReader->GetBlendShapeTargetDeltaYs(MeshTarget.MeshIndex, MeshTarget.TargetIndex);
+					TArrayView<const float> DeltaZs = InDNAReader->GetBlendShapeTargetDeltaZs(MeshTarget.MeshIndex, MeshTarget.TargetIndex);
 
 					for (int32 DeltaIndex = 0; DeltaIndex < NumOfDeltas; DeltaIndex++)
 					{
 						FVector PositionDelta = FVector(DeltaXs[DeltaIndex], DeltaYs[DeltaIndex], DeltaZs[DeltaIndex]);
 						FMorphTargetDelta MorphDelta;
 						int32 DNAVertexIndex = BlendShapeVertexIndices[DeltaIndex];
-						int32 UEVertexIndex = DNAToSkelMeshMap->ImportDNAVtxToUEVtxIndex[LODIndex][MeshTarget.MeshIndex][DNAVertexIndex];
+						int32 UEVertexIndex = InDNAToSkelMeshMap->ImportDNAVtxToUEVtxIndex[LODIndex][MeshTarget.MeshIndex][DNAVertexIndex];
 						if (UEVertexIndex > INDEX_NONE)
 						{
 							MorphDelta.SourceIdx = (uint32)UEVertexIndex;
@@ -238,7 +250,7 @@ void USkelMeshDNAUtils::UpdateMorphTargets(USkeletalMesh* SkelMesh, IDNAReader*
 							MorphDelta.TangentZDelta = FVector3f::ZeroVector;
 							MorphLODModel.Vertices.Add(MorphDelta);
 							// Find section indices that are involved in these morph deltas.
-							int32 SectionIdx = DNAToSkelMeshMap->UEVertexToSectionIndices[LODIndex][UEVertexIndex];
+							int32 SectionIdx = InDNAToSkelMeshMap->UEVertexToSectionIndices[LODIndex][UEVertexIndex];
 							if (!MorphLODModel.SectionIndices.Contains(SectionIdx) && SectionIdx > INDEX_NONE)
 							{
 								MorphLODModel.SectionIndices.Add(SectionIdx);
@@ -252,7 +264,7 @@ void USkelMeshDNAUtils::UpdateMorphTargets(USkeletalMesh* SkelMesh, IDNAReader*
 			{
 				int32 LODStart;
 				int32 LODRangeSize;
-				GetLODRange(UpdateOption, MorphTarget->GetMorphLODModels().Num(), LODStart, LODRangeSize);
+				GetLODRange(InUpdateOption, MorphTarget->GetMorphLODModels().Num(), LODStart, LODRangeSize);
 				for (int32 LODIndex = LODStart; LODIndex < LODRangeSize; LODIndex++)
 				{
 					FMorphTargetLODModel& MorphLODModel = MorphTarget->GetMorphLODModels()[LODIndex];
@@ -266,33 +278,33 @@ void USkelMeshDNAUtils::UpdateMorphTargets(USkeletalMesh* SkelMesh, IDNAReader*
 }
 
 /* Updates Bone influences using Skin Weights from DNA. */
-void USkelMeshDNAUtils::UpdateSkinWeights(USkeletalMesh* SkelMesh, IDNAReader* DNAReader, FDNAToSkelMeshMap* DNAToSkelMeshMap, ELodUpdateOption UpdateOption)
+void USkelMeshDNAUtils::UpdateSkinWeights(USkeletalMesh* InSkelMesh, IDNAReader* InDNAReader, FDNAToSkelMeshMap* InDNAToSkelMeshMap, ELodUpdateOption InUpdateOption)
 {
-	FSkeletalMeshModel* ImportedModel = SkelMesh->GetImportedModel();
+	FSkeletalMeshModel* ImportedModel = InSkelMesh->GetImportedModel();
 
 	bool InfluenceMismatch = false;
 	//Set a threshold a bit smaller then 1/255
 	int32 LODStart;
 	int32 LODRangeSize;
-	GetLODRange(UpdateOption, ImportedModel->LODModels.Num(), LODStart, LODRangeSize);
+	GetLODRange(InUpdateOption, ImportedModel->LODModels.Num(), LODStart, LODRangeSize);
 	for (int32 LODIndex = LODStart; LODIndex < LODRangeSize; ++LODIndex)
 	{
 		FSkeletalMeshLODModel& LODModel = ImportedModel->LODModels[LODIndex];
 		for (int32 SectionIndex = 0; SectionIndex < LODModel.Sections.Num(); SectionIndex++)
 		{
 			FSkelMeshSection& Section = LODModel.Sections[SectionIndex];
-			const int32 DNAMeshIndex = DNAToSkelMeshMap->ImportVtxToDNAMeshIndex[LODIndex][Section.GetVertexBufferIndex()];
+			const int32 DNAMeshIndex = InDNAToSkelMeshMap->ImportVtxToDNAMeshIndex[LODIndex][Section.GetVertexBufferIndex()];
 
 			const int32 NumEngineVertices = Section.GetNumVertices();
 			for (int32 VertexIndex = 0; VertexIndex < NumEngineVertices; VertexIndex++)
 			{
 				const int32 VertexBufferIndex = VertexIndex + Section.GetVertexBufferIndex();
-				const int32 DNAVertexIndex = DNAToSkelMeshMap->ImportVtxToDNAVtxIndex[LODIndex][VertexBufferIndex];
+				const int32 DNAVertexIndex = InDNAToSkelMeshMap->ImportVtxToDNAVtxIndex[LODIndex][VertexBufferIndex];
 
 				if (DNAVertexIndex < 0) continue; // Skip vertex not in DNA.
 
-				TArrayView<const float> DNASkinWeights = DNAReader->GetSkinWeightsValues(DNAMeshIndex, DNAVertexIndex);
-				TArrayView<const uint16> DNASkinJoints = DNAReader->GetSkinWeightsJointIndices(DNAMeshIndex, DNAVertexIndex);
+				TArrayView<const float> DNASkinWeights = InDNAReader->GetSkinWeightsValues(DNAMeshIndex, DNAVertexIndex);
+				TArrayView<const uint16> DNASkinJoints = InDNAReader->GetSkinWeightsJointIndices(DNAMeshIndex, DNAVertexIndex);
 				int32 SkinJointNum = DNASkinJoints.Num();
 
 				FSoftSkinVertex& Vertex = Section.SoftVertices[VertexIndex];
@@ -310,7 +322,7 @@ void USkelMeshDNAUtils::UpdateSkinWeights(USkeletalMesh* SkelMesh, IDNAReader* D
 				for (int32 Index = 0; Index < FMath::Min(SkinJointNum, MAX_TOTAL_INFLUENCES); ++Index)
 				{
 					// Find Engine bone for corresponding DNAJoint for the same influence.
-					const int32 UpdatedBoneId = DNAToSkelMeshMap->GetUEBoneIndex(DNASkinJoints[Index]);
+					const int32 UpdatedBoneId = InDNAToSkelMeshMap->GetUEBoneIndex(NormalizedJoints[InfluenceIndex]);
 					// BoneMap holds subset of bones belonging to current section.
 					const int32 BoneMapIndex = Section.BoneMap.Find(UpdatedBoneId);
 
@@ -341,22 +353,22 @@ void USkelMeshDNAUtils::UpdateSkinWeights(USkeletalMesh* SkelMesh, IDNAReader* D
 }
 
 /** Rebuilds render data from LODModel and inits resources. */
-void USkelMeshDNAUtils::RebuildRenderData(USkeletalMesh* SkelMesh)
+void USkelMeshDNAUtils::RebuildRenderData(USkeletalMesh* InSkelMesh)
 {
 	FPlatformTime::InitTiming();
 
 	double StartTime = FPlatformTime::Seconds();
 	{
-		SkelMesh->FlushRenderState();
+		InSkelMesh->FlushRenderState();
 	}
 	double TimeToFlush = FPlatformTime::Seconds();
 	{
-		FSkeletalMeshRenderData* RenderData = SkelMesh->GetResourceForRendering();
+		FSkeletalMeshRenderData* RenderData = InSkelMesh->GetResourceForRendering();
 		int32 LODIndex = 0;
 
 		for (FSkeletalMeshLODRenderData& LODRenderData : RenderData->LODRenderData)
 		{
-			FSkeletalMeshLODModel& LODModelRef = SkelMesh->GetImportedModel()->LODModels[LODIndex];
+			FSkeletalMeshLODModel& LODModelRef = InSkelMesh->GetImportedModel()->LODModels[LODIndex];
 			for (int32 i = 0; i < LODModelRef.Sections.Num(); i++)
 			{
 				FSkelMeshSection& ModelSection = LODModelRef.Sections[i];
@@ -374,7 +386,7 @@ void USkelMeshDNAUtils::RebuildRenderData(USkeletalMesh* SkelMesh)
 		if (FApp::CanEverRender())
 		{
 			// Reinitialize the static mesh's resources.
-			SkelMesh->InitResources();
+			InSkelMesh->InitResources();
 		}
 	}
 	double TimeToInitResources = FPlatformTime::Seconds();
@@ -384,7 +396,7 @@ void USkelMeshDNAUtils::RebuildRenderData(USkeletalMesh* SkelMesh)
 		for (TObjectIterator<USkeletalMeshComponent> It; It; ++It)
 		{
 			USkeletalMeshComponent* MeshComponent = *It;
-			if (MeshComponent && !MeshComponent->IsTemplate() && MeshComponent->GetSkeletalMeshAsset() == SkelMesh)
+			if (MeshComponent && !MeshComponent->IsTemplate() && MeshComponent->GetSkeletalMeshAsset() == InSkelMesh)
 			{
 				ComponentsToReregister.Add(*It);
 			}
@@ -393,7 +405,7 @@ void USkelMeshDNAUtils::RebuildRenderData(USkeletalMesh* SkelMesh)
 	}
 }
 
-void USkelMeshDNAUtils::RebuildRenderData_VertexPosition(USkeletalMesh* SkelMesh)
+void USkelMeshDNAUtils::RebuildRenderData_VertexPosition(USkeletalMesh* InSkelMesh)
 {
 	if (!FApp::CanEverRender())
 	{
@@ -401,8 +413,8 @@ void USkelMeshDNAUtils::RebuildRenderData_VertexPosition(USkeletalMesh* SkelMesh
 	}
 
 	{
-		FSkeletalMeshModel* MeshModel = SkelMesh->GetImportedModel();
-		FSkeletalMeshRenderData* RenderData = SkelMesh->GetResourceForRendering();
+		FSkeletalMeshModel* MeshModel = InSkelMesh->GetImportedModel();
+		FSkeletalMeshRenderData* RenderData = InSkelMesh->GetResourceForRendering();
 
 		for (int32 LODIdx = 0; LODIdx < RenderData->LODRenderData.Num(); ++LODIdx)
 		{
@@ -432,7 +444,7 @@ void USkelMeshDNAUtils::RebuildRenderData_VertexPosition(USkeletalMesh* SkelMesh
 	}
 }
 
-void USkelMeshDNAUtils::UpdateJointBehavior(USkeletalMeshComponent* SkelMeshComponent)
+void USkelMeshDNAUtils::UpdateJointBehavior(USkeletalMeshComponent* InSkelMeshComponent)
 {
 	//DNAAsset->SetBehaviorReader is called before calling this method
 	//it is not here to avoid having DNAAsset in the API, as in future we might want
@@ -440,15 +452,15 @@ void USkelMeshDNAUtils::UpdateJointBehavior(USkeletalMeshComponent* SkelMeshComp
 
 	//the rig behavior has changed, we need to force re-initializing of RigLogic
 	//this will set RigLogic RigUnit to initial state
-	SkelMeshComponent->InitAnim(true);
+	InSkelMeshComponent->InitAnim(true);
 }
 
-void USkelMeshDNAUtils::UpdateSourceData(USkeletalMesh* SkelMesh)
+void USkelMeshDNAUtils::UpdateSourceData(USkeletalMesh* InSkelMesh)
 {
 	IMeshUtilities& MeshUtilities = FModuleManager::Get().LoadModuleChecked<IMeshUtilities>("MeshUtilities");
 	// Source data must be updated during cooking.
-	SkelMesh->EmptyAllImportData();
-	MeshUtilities.CreateImportDataFromLODModel(SkelMesh);
+	InSkelMesh->EmptyAllImportData();
+	MeshUtilities.CreateImportDataFromLODModel(InSkelMesh);
 }
 
 #endif // WITH_EDITORONLY_DATA
diff --git a/Engine/Plugins/Animation/RigLogic/Source/RigLogicModule/Public/DNAToSkelMeshMap.h b/Engine/Plugins/Animation/RigLogic/Source/RigLogicModule/Public/DNAToSkelMeshMap.h
--- a/Engine/Plugins/Animation/RigLogic/Source/RigLogicModule/Public/DNAToSkelMeshMap.h
+++ b/Engine/Plugins/Animation/RigLogic/Source/RigLogicModule/Public/DNAToSkelMeshMap.h
@@ -43,7 +43,7 @@ public:
 	/** Creates vertex map between source DNA and target SkelMesh, and stores DNAReader
 	  * Enough for updating neutral meshes, but not for updating joints and morphtargets; call the next method for that
 	  * TODO: When ImportedModel is added to runtime, WITH_EDITORONLY_DATA should be removed**/
-	void InitBaseMesh(IDNAReader* InDNAReader, USkeletalMesh* SkelMesh);
+	void InitBaseMesh(IDNAReader* InSourceDNAReader, USkeletalMesh* InTargetSkeletalMesh);
 	/** Creates map for updating joints; needs to be called before playing animation on the rig **/
 	void MapJoints(IDNAReader* InDNAReader);
 	/** Creates map for updating morph targets; needs to be called before playing animation on the rig **/
@@ -51,13 +51,13 @@ public:
 #endif //WITH_EDITORONLY_DATA
 	/** Creates mappings between source DNA stored in DNAAsset AssetUserData attached to SkeletalMesh;
 		Uses FSkeletalMeshDNAReader to be able to read geometry which, in this case, is NOT present in the DNAAsset **/
-	bool InitFromDNAAsset(USkeletalMesh* SkelMesh);
+	bool InitFromDNAAsset(USkeletalMesh* InSkelMesh);
 	void InitVertexMap(IDNAReader* InDNAReader);
 
 
-	int32 GetUEBoneIndex(int32 RLJointIndex)
+	int32 GetUEBoneIndex(int32 InRLJointIndex)
 	{
-		return RLJointToUEBoneIndices[RLJointIndex];
+		return RLJointToUEBoneIndices[InRLJointIndex];
 	}
 
 	TArray<FDNABlendShapeTarget>& GetMeshBlendShapeTargets()
diff --git a/Engine/Plugins/Animation/RigLogic/Source/RigLogicModule/Public/SkelMeshDNAUtils.h b/Engine/Plugins/Animation/RigLogic/Source/RigLogicModule/Public/SkelMeshDNAUtils.h
--- a/Engine/Plugins/Animation/RigLogic/Source/RigLogicModule/Public/SkelMeshDNAUtils.h
+++ b/Engine/Plugins/Animation/RigLogic/Source/RigLogicModule/Public/SkelMeshDNAUtils.h
@@ -30,55 +30,57 @@ class RIGLOGICMODULE_API USkelMeshDNAUtils: public UObject
 public:
 
 	/** Prepare context object that will allow mapping of DNA structures to SkelMesh ones for updating **/
-	static FDNAToSkelMeshMap* CreateMapForUpdatingNeutralMesh(IDNAReader* DNAReader, USkeletalMesh* SkelMesh);
+	static FDNAToSkelMeshMap* CreateMapForUpdatingNeutralMesh(IDNAReader* InDNAReader, USkeletalMesh* InSkelMesh);
+	/** Prepare context object that will allow mapping of DNA structures extracted from DNAAsset to SkelMesh ones for updating **/
+	static FDNAToSkelMeshMap* CreateMapForUpdatingNeutralMesh(USkeletalMesh* InSkelMesh);
 
 #if WITH_EDITORONLY_DATA
 	/** Updates the positions, orientation and scale in the joint hierarchy using the data from DNA file **/
-	static void UpdateJoints(USkeletalMesh* SkelMesh, IDNAReader* DNAReader, FDNAToSkelMeshMap* DNAToSkelMeshMap);
-	/** Updates the base mesh vertex positions for all mesh sections of all LODs, using the data from DNA file 
+	static void UpdateJoints(USkeletalMesh* InSkelMesh, IDNAReader* InDNAReader, FDNAToSkelMeshMap* InDNAToSkelMeshMap);
+	/** Updates the base mesh vertex positions for all mesh sections of all LODs, using the data from DNA file
 	  * NOTE: Not calling RebuildRenderData automatically, it needs to be called explicitly after the first update
 	  *       As the topology doesn't change, for subsequent updates it can be ommited to gain performance **/
-	static void UpdateBaseMesh(USkeletalMesh* SkelMesh, IDNAReader* DNAReader, FDNAToSkelMeshMap* DNAToSkelMeshMap, ELodUpdateOption UpdateOption = ELodUpdateOption::LOD0Only);
+	static void UpdateBaseMesh(USkeletalMesh* InSkelMesh, IDNAReader* InDNAReader, FDNAToSkelMeshMap* InDNAToSkelMeshMap, ELodUpdateOption InUpdateOption = ELodUpdateOption::LOD0Only);
 	/** Updates the morph targets for all mesh sections of LODs, using the data from DNA file **/
-	static void UpdateMorphTargets(USkeletalMesh* SkelMesh, IDNAReader* DNAReader, FDNAToSkelMeshMap* DNAToSkelMeshMap, ELodUpdateOption UpdateOption = ELodUpdateOption::LOD0Only);
+	static void UpdateMorphTargets(USkeletalMesh* InSkelMesh, IDNAReader* InDNAReader, FDNAToSkelMeshMap* InDNAToSkelMeshMap, ELodUpdateOption InUpdateOption = ELodUpdateOption::LOD0Only);
 	/** Updates the skin weights for all LODs using the data from DNA file **/
-	static void UpdateSkinWeights(USkeletalMesh* SkelMesh, IDNAReader* DNAReader, FDNAToSkelMeshMap* DNAToSkelMeshMap, ELodUpdateOption UpdateOption = ELodUpdateOption::LOD0Only);
+	static void UpdateSkinWeights(USkeletalMesh* InSkelMesh, IDNAReader* InDNAReader, FDNAToSkelMeshMap* InDNAToSkelMeshMap, ELodUpdateOption InUpdateOption = ELodUpdateOption::LOD0Only);
 	/** Rechunks the mesh after the update **/
-	static void RebuildRenderData(USkeletalMesh* SkelMesh);
+	static void RebuildRenderData(USkeletalMesh* InSkelMesh);
 	/** Re-initialize vertex positions for rendering after the update **/
-	static void RebuildRenderData_VertexPosition(USkeletalMesh* SkelMesh);
+	static void RebuildRenderData_VertexPosition(USkeletalMesh* InSkelMesh);
 	/** Update joint behavior **/
-	static void UpdateJointBehavior(USkeletalMeshComponent* SkelMeshComponent);
+	/*  NOTE: DNAAsset->SetBehaviorReader needs to be called before this */
+	static void UpdateJointBehavior(USkeletalMeshComponent* InSkelMeshComponent);
 #endif // WITH_EDITORONLY_DATA
 	/** Converts DNA vertex coordinates to UE4 coordinate system **/
-	inline static FVector ConvertDNAVertexToUE4CoordSystem(FVector VertexPositionInDNA)
+	inline static FVector ConvertDNAVertexToUE4CoordSystem(FVector InVertexPositionInDNA)
 	{
-		return FVector{-VertexPositionInDNA.X, VertexPositionInDNA.Y, -VertexPositionInDNA.Z};
+		return FVector{ -InVertexPositionInDNA.X, InVertexPositionInDNA.Y, -InVertexPositionInDNA.Z };
 	}
 
 	/** Converts UE4 coordinate system to DNA vertex coordinates **/
-	inline static FVector ConvertUE4CoordSystemToDNAVertex(FVector VertexPositionInUE4)
+	inline static FVector ConvertUE4CoordSystemToDNAVertex(FVector InVertexPositionInUE4)
 	{
-		return FVector{-VertexPositionInUE4.X, VertexPositionInUE4.Y, -VertexPositionInUE4.Z};
+		return FVector{ -InVertexPositionInUE4.X, InVertexPositionInUE4.Y, -InVertexPositionInUE4.Z };
 	}
 
 	/** Updates source skeleton data for the purpose of character cooking and export.  */
-	static void UpdateSourceData(USkeletalMesh* SkelMesh);
-	static void UpdateSourceData(USkeletalMesh* SkelMesh, IDNAReader* DNAReader, FDNAToSkelMeshMap* DNAToSkelMeshMap);
+	static void UpdateSourceData(USkeletalMesh* InSkelMesh);
 	
 private:
 
-	inline static void GetLODRange(ELodUpdateOption UpdateOption, const int32& LODNum, int32& LODStart, int32& LODRangeSize)
+	inline static void GetLODRange(ELodUpdateOption InUpdateOption, const int32& InLODNum, int32& OutLODStart, int32& OutLODRangeSize)
 	{
-		LODStart = 0;
-		LODRangeSize = LODNum;
-		if (UpdateOption == ELodUpdateOption::LOD1AndHigher)
+		OutLODStart = 0;
+		OutLODRangeSize = InLODNum;
+		if (InUpdateOption == ELodUpdateOption::LOD1AndHigher)
 		{
-			LODStart = 1;
+			OutLODStart = 1;
 		}
-		else if (UpdateOption == ELodUpdateOption::LOD0Only && LODRangeSize > 0)
+		else if (InUpdateOption == ELodUpdateOption::LOD0Only && OutLODRangeSize > 0)
 		{
-			LODRangeSize = 1;
+			OutLODRangeSize = 1;
 		}
 	}
 };
