Fix merge error
#rnx

diff --git a/Engine/Plugins/Animation/RigLogic/Source/RigLogicModule/Private/SkelMeshDNAUtils.cpp b/Engine/Plugins/Animation/RigLogic/Source/RigLogicModule/Private/SkelMeshDNAUtils.cpp
--- a/Engine/Plugins/Animation/RigLogic/Source/RigLogicModule/Private/SkelMeshDNAUtils.cpp
+++ b/Engine/Plugins/Animation/RigLogic/Source/RigLogicModule/Private/SkelMeshDNAUtils.cpp
@@ -322,7 +322,7 @@ void USkelMeshDNAUtils::UpdateSkinWeights(USkeletalMesh* InSkelMesh, IDNAReader*
 				for (int32 Index = 0; Index < FMath::Min(SkinJointNum, MAX_TOTAL_INFLUENCES); ++Index)
 				{
 					// Find Engine bone for corresponding DNAJoint for the same influence.
-					const int32 UpdatedBoneId = InDNAToSkelMeshMap->GetUEBoneIndex(NormalizedJoints[InfluenceIndex]);
+					const int32 UpdatedBoneId = InDNAToSkelMeshMap->GetUEBoneIndex(DNASkinJoints[InfluenceIndex]);
 					// BoneMap holds subset of bones belonging to current section.
 					const int32 BoneMapIndex = Section.BoneMap.Find(UpdatedBoneId);
 
