Enable -25 Android AutoSDK to use NDK 25b

diff --git a/Engine/Extras/Android/SetupAndroid.bat b/Engine/Extras/Android/SetupAndroid.bat
--- a/Engine/Extras/Android/SetupAndroid.bat
+++ b/Engine/Extras/Android/SetupAndroid.bat
@@ -17,10 +17,10 @@ SET CMAKE_VERSION=%3
 SET NDK_VERSION=%4
 
 rem hardcoded versions for compatibility with non-Turnkey manual running
-if "%PLATFORMS_VERSION%" == "" SET PLATFORMS_VERSION=android-30
+if "%PLATFORMS_VERSION%" == "" SET PLATFORMS_VERSION=android-32
 if "%BUILDTOOLS_VERSION%" == "" SET BUILDTOOLS_VERSION=30.0.3
 if "%CMAKE_VERSION%" == "" SET CMAKE_VERSION=3.10.2.4988404
-if "%NDK_VERSION%" == "" SET NDK_VERSION=21.4.7075529
+if "%NDK_VERSION%" == "" SET NDK_VERSION=25.1.8937393
 
 
 FOR /F "tokens=2*" %%A IN ('REG.exe query "%KEY_NAME%" /v "%VALUE_NAME%"') DO (set STUDIO_PATH=%%B)
diff --git a/Engine/Extras/Android/SetupAndroid.command b/Engine/Extras/Android/SetupAndroid.command
--- a/Engine/Extras/Android/SetupAndroid.command
+++ b/Engine/Extras/Android/SetupAndroid.command
@@ -51,7 +51,7 @@ fi
 
 # hardcoded versions for compatibility with non-Turnkey manual running
 if [[ -z "${PLATFORMS_VERSION}" ]]; then
-    PLATFORMS_VERSION="android-30"
+    PLATFORMS_VERSION="android-32"
 fi
 if [[ -z "${BUILDTOOLS_VERSION}" ]]; then
     BUILDTOOLS_VERSION="30.0.3"
@@ -60,7 +60,7 @@ if [[ -z "${CMAKE_VERSION}" ]]; then
     CMAKE_VERSION="3.10.2.4988404"
 fi
 if [[ -z "${NDK_VERSION}" ]]; then
-    NDK_VERSION="21.4.7075529"
+    NDK_VERSION="25.1.8937393"
 fi
 
 
diff --git a/Engine/Extras/Android/SetupAndroid.sh b/Engine/Extras/Android/SetupAndroid.sh
--- a/Engine/Extras/Android/SetupAndroid.sh
+++ b/Engine/Extras/Android/SetupAndroid.sh
@@ -27,7 +27,7 @@ fi
 
 # hardcoded versions for compatibility with non-Turnkey manual running
 if [[ -z "${PLATFORMS_VERSION}" ]]; then
-    PLATFORMS_VERSION="android-30"
+    PLATFORMS_VERSION="android-32"
 fi
 if [[ -z "${BUILDTOOLS_VERSION}" ]]; then
     BUILDTOOLS_VERSION="30.0.3"
@@ -36,7 +36,7 @@ if [[ -z "${CMAKE_VERSION}" ]]; then
     CMAKE_VERSION="3.10.2.4988404"
 fi
 if [[ -z "${NDK_VERSION}" ]]; then
-    NDK_VERSION="21.4.7075529"
+    NDK_VERSION="25.1.8937393"
 fi
 
 
diff --git a/Engine/Extras/Android/UpdateLinker.bat b/Engine/Extras/Android/UpdateLinker.bat
deleted file mode 100755
diff --git a/Engine/Extras/Android/UpdateLinker.command b/Engine/Extras/Android/UpdateLinker.command
deleted file mode 100644
diff --git a/Engine/Extras/Android/UpdateLinkerLumin.command b/Engine/Extras/Android/UpdateLinkerLumin.command
deleted file mode 100644
diff --git a/Engine/Plugins/Animation/RigLogic/Source/RigLogicLib/Private/pma/resources/AlignedMemoryResource.cpp b/Engine/Plugins/Animation/RigLogic/Source/RigLogicLib/Private/pma/resources/AlignedMemoryResource.cpp
--- a/Engine/Plugins/Animation/RigLogic/Source/RigLogicLib/Private/pma/resources/AlignedMemoryResource.cpp
+++ b/Engine/Plugins/Animation/RigLogic/Source/RigLogicLib/Private/pma/resources/AlignedMemoryResource.cpp
@@ -13,7 +13,10 @@
     #include <malloc.h>
 #endif
 
-#if defined(__cplusplus) && (__cplusplus >= 201703L) && (defined(_GLIBCXX_HAVE_ALIGNED_ALLOC) || \
+#if defined(__ANDROID__)
+    #define ALIGNED_ALLOC(ptr, alignment, size) ptr = ::memalign(alignment, size)
+    #define ALIGNED_FREE(ptr) ::free(ptr)
+#elif defined(__cplusplus) && (__cplusplus >= 201703L) && (defined(_GLIBCXX_HAVE_ALIGNED_ALLOC) || \
     defined(_LIBCPP_HAS_ALIGNED_ALLOC) || defined(_LIBCPP_HAS_C11_FEATURES))
     #define ALIGNED_ALLOC(ptr, alignment, size) ptr = std::aligned_alloc(alignment, size)
     #define ALIGNED_FREE(ptr) std::free(ptr)
@@ -24,9 +27,6 @@
 #elif defined(_MSC_VER)
     #define ALIGNED_ALLOC(ptr, alignment, size) ptr = _aligned_malloc(size, alignment)
     #define ALIGNED_FREE(ptr) _aligned_free(ptr)
-#elif defined(__ANDROID__)
-    #define ALIGNED_ALLOC(ptr, alignment, size) ptr = ::memalign(alignment, size)
-    #define ALIGNED_FREE(ptr) ::free(ptr)
 #else
     inline std::uintptr_t alignAddress(std::uintptr_t address, std::size_t alignment) {
         const std::size_t mask = alignment - 1ul;
diff --git a/Engine/Source/Programs/UnrealBuildTool/Platform/Android/AndroidPlatformSDK.Versions.cs b/Engine/Source/Programs/UnrealBuildTool/Platform/Android/AndroidPlatformSDK.Versions.cs
--- a/Engine/Source/Programs/UnrealBuildTool/Platform/Android/AndroidPlatformSDK.Versions.cs
+++ b/Engine/Source/Programs/UnrealBuildTool/Platform/Android/AndroidPlatformSDK.Versions.cs
@@ -8,18 +8,18 @@ namespace UnrealBuildTool
 	{
 		public override string GetMainVersion()
 		{
-			return "r21e";
+			return "r25b";
 		}
 		
 		public override string GetAutoSDKDirectoryForMainVersion()
 		{
-			return "-24";
+			return "-25";
 		}
 
 		protected override void GetValidVersionRange(out string MinVersion, out string MaxVersion)
 		{
 			MinVersion = "r21a";
-			MaxVersion = "r25";
+			MaxVersion = "r25b";
 		}
 
 		protected override void GetValidSoftwareVersionRange(out string? MinVersion, out string? MaxVersion)
@@ -31,10 +31,10 @@ namespace UnrealBuildTool
 		{
 			switch (VersionType.ToLower())
 			{
-				case "platforms": return "android-28";
-				case "build-tools": return "28.0.3";
+				case "platforms": return "android-32";
+				case "build-tools": return "30.0.3";
 				case "cmake": return "3.10.2.4988404";
-				case "ndk": return "21.4.7075529";
+				case "ndk": return "25.1.8937393";
 			}
 
 			return "";
