Add missing function to SkelMeshDNAUtils in RigLogic

diff --git a/Engine/Plugins/Animation/RigLogic/Source/RigLogicModule/Private/SkelMeshDNAUtils.cpp b/Engine/Plugins/Animation/RigLogic/Source/RigLogicModule/Private/SkelMeshDNAUtils.cpp
--- a/Engine/Plugins/Animation/RigLogic/Source/RigLogicModule/Private/SkelMeshDNAUtils.cpp
+++ b/Engine/Plugins/Animation/RigLogic/Source/RigLogicModule/Private/SkelMeshDNAUtils.cpp
@@ -463,4 +463,24 @@ void USkelMeshDNAUtils::UpdateSourceData(USkeletalMesh* InSkelMesh)
 	MeshUtilities.CreateImportDataFromLODModel(InSkelMesh);
 }
 
+UDNAAsset* USkelMeshDNAUtils::GetMeshDNA(USkeletalMesh* InSkelMesh)
+{
+	if (InSkelMesh && InSkelMesh->GetAssetUserDataArray())
+	{
+		for (UAssetUserData* UserData : *InSkelMesh->GetAssetUserDataArray())
+		{
+			if (UserData && UserData->IsA<UDNAAsset>())
+			{
+				UDNAAsset* DNAAsset = Cast<UDNAAsset>(UserData);
+				if (DNAAsset->GetBehaviorReader() && DNAAsset->GetGeometryReader())
+				{
+					return DNAAsset;
+				}
+			}
+		}
+	}
+
+	return nullptr;
+}
+
 #endif // WITH_EDITORONLY_DATA
diff --git a/Engine/Plugins/Animation/RigLogic/Source/RigLogicModule/Public/SkelMeshDNAUtils.h b/Engine/Plugins/Animation/RigLogic/Source/RigLogicModule/Public/SkelMeshDNAUtils.h
--- a/Engine/Plugins/Animation/RigLogic/Source/RigLogicModule/Public/SkelMeshDNAUtils.h
+++ b/Engine/Plugins/Animation/RigLogic/Source/RigLogicModule/Public/SkelMeshDNAUtils.h
@@ -4,6 +4,7 @@
 
 #include "DNAToSkelMeshMap.h"
 #include "Engine/SkeletalMesh.h"
+#include "DNAAsset.h"
 #include "SkelMeshDNAUtils.generated.h"
 
 DECLARE_LOG_CATEGORY_EXTERN(LogDNAUtils, Log, All);
@@ -52,6 +53,8 @@ public:
 	/** Update joint behavior **/
 	/*  NOTE: DNAAsset->SetBehaviorReader needs to be called before this */
 	static void UpdateJointBehavior(USkeletalMeshComponent* InSkelMeshComponent);
+	/** Gets the DNA asset embedded in the mesh */
+	static UDNAAsset* GetMeshDNA(USkeletalMesh* InSkelMesh);
 #endif // WITH_EDITORONLY_DATA
 	/** Converts DNA vertex coordinates to UE4 coordinate system **/
 	inline static FVector ConvertDNAVertexToUE4CoordSystem(FVector InVertexPositionInDNA)
