Fix DNA hash calculation #preflight 637239528f4cb2e4dc8f4098 #rb violeta.vukobrat filip.micic

diff --git a/Engine/Plugins/Animation/RigLogic/Source/RigLogicModule/Private/Hasher.h b/Engine/Plugins/Animation/RigLogic/Source/RigLogicModule/Private/Hasher.h
--- a/Engine/Plugins/Animation/RigLogic/Source/RigLogicModule/Private/Hasher.h
+++ b/Engine/Plugins/Animation/RigLogic/Source/RigLogicModule/Private/Hasher.h
@@ -11,6 +11,11 @@ public:
 	{
 		return HashedValue;
 	}
+	void Update(const FString& Value)
+	{
+		const auto& Data = Value.GetCharArray();
+		Update(reinterpret_cast<const char*>(Data.GetData()), Data.Num() * Data.GetTypeSize());
+	}
 	template <typename T>
 	void Update(T Value)
 	{
@@ -19,7 +24,7 @@ public:
 	template <typename T>
 	void Update(TArrayView<T> Values)
 	{
-		Update(reinterpret_cast<const char*>(Values.Data()), Values.Num());
+		Update(reinterpret_cast<const char*>(Values.Data()), Values.Num() * Values.GetTypeSize());
 	}
 private:
 	void Update(const char* Bytes, size_t Length)
