Enable import of DNA file with the same name as Skeletal Mesh

diff --git a/Engine/Plugins/Animation/RigLogic/Source/RigLogicEditor/Private/DNAAssetImportFactory.cpp b/Engine/Plugins/Animation/RigLogic/Source/RigLogicEditor/Private/DNAAssetImportFactory.cpp
--- a/Engine/Plugins/Animation/RigLogic/Source/RigLogicEditor/Private/DNAAssetImportFactory.cpp
+++ b/Engine/Plugins/Animation/RigLogic/Source/RigLogicEditor/Private/DNAAssetImportFactory.cpp
@@ -53,13 +53,47 @@ void UDNAAssetImportFactory::PostInitProperties()
 	ImportUI = NewObject<UDNAAssetImportUI>(this, NAME_None, RF_NoFlags);
 }
 
+bool UDNAAssetImportFactory::CanReimport(UObject* Obj, TArray<FString>& OutFilenames)
+{
+	TArray<FString> FactoryExtensions;
+	GetSupportedFileExtensions(FactoryExtensions);
+
+	return FactoryExtensions.Contains(FPaths::GetExtension(PreferredReimportPath));
+}
+
+EReimportResult::Type UDNAAssetImportFactory::Reimport(UObject* Obj, int32 SourceFileIndex)
+{
+	// This is done here in oreder to enable importing of the dna file with the same name as 
+	// skeletal mesh. This is treated as an normal import but engine is recognizing it as an skeletal mesh reimport.
+
+	UPackage* Pkg = Obj->GetPackage();
+	FString Name = FPaths::GetBaseFilename(PreferredReimportPath);
+	bool bImportWasCancelled = false;
+	UObject* Result = ImportObject(SupportedClass, Pkg, FName(*Name), RF_Public | RF_Standalone | RF_Transactional, PreferredReimportPath, nullptr, bImportWasCancelled);
+
+	// Clean up and remove the factories we created from the root set
+	CleanUp();
+
+	if (Result != nullptr)
+	{
+		return EReimportResult::Succeeded;
+	}
+
+	return EReimportResult::Failed;
+}
+
+int32 UDNAAssetImportFactory::GetPriority() const
+{
+	return ImportPriority;
+}
+
 UObject* UDNAAssetImportFactory::FactoryCreateFile(UClass* InClass, UObject* InParent, FName InName, EObjectFlags Flags, const FString& Filename, const TCHAR* Parms, FFeedbackContext* Warn, bool& bOutOperationCanceled)
 {
 	// this one prints all messages that are stored in FFbxImporter
 	FDNAImporter* DNAImporter = FDNAImporter::GetInstance();
 	FDNAAssetImportOptions* ImportOptions = DNAImporter->GetImportOptions();
 	//Clean up the options
-	FDNAAssetImportOptions::ResetOptions(ImportOptions);	
+	FDNAAssetImportOptions::ResetOptions(ImportOptions);
 
 	struct FRestoreImportUI
 	{
@@ -117,7 +151,7 @@ UObject* UDNAAssetImportFactory::FactoryCreateFile(UClass* InClass, UObject* InP
 			DNAAssetFullName.Append(DNAAssetFileName);
 
 			GEditor->GetEditorSubsystem<UImportSubsystem>()->BroadcastAssetPreImport(this, InClass, InParent, *DNAAssetFullName, TEXT("dna"));
-			
+
 			UDNAAsset* DNAAsset = NewObject< UDNAAsset >(SkelMesh, FName(*DNAAssetFileName)); //SkelMesh has to be its outer, otherwise DNAAsset won't be saved			
 			if (DNAAsset->Init(*Filename))
 			{
@@ -139,7 +173,7 @@ UObject* UDNAAssetImportFactory::FactoryCreateFile(UClass* InClass, UObject* InP
 	}
 
 	// Failed to load file or create DNA stream. Clean and return nullptr.
-	DNAImporter->PartialCleanUp();	
+	DNAImporter->PartialCleanUp();
 	GEditor->GetEditorSubsystem<UImportSubsystem>()->BroadcastAssetPostImport(this, nullptr);
 	return nullptr;
 }
@@ -156,10 +190,6 @@ bool UDNAAssetImportFactory::FactoryCanImport(const FString& Filename)
 	return false;
 }
 
-void UDNAAssetImportFactory::CleanUp()
-{
-}
-
 // DNA Asset Import UI Implementation.
 UDNAAssetImportUI::UDNAAssetImportUI(const FObjectInitializer& ObjectInitializer)
 	: Super(ObjectInitializer)
diff --git a/Engine/Plugins/Animation/RigLogic/Source/RigLogicEditor/Public/DNAAssetImportFactory.h b/Engine/Plugins/Animation/RigLogic/Source/RigLogicEditor/Public/DNAAssetImportFactory.h
--- a/Engine/Plugins/Animation/RigLogic/Source/RigLogicEditor/Public/DNAAssetImportFactory.h
+++ b/Engine/Plugins/Animation/RigLogic/Source/RigLogicEditor/Public/DNAAssetImportFactory.h
@@ -12,10 +12,11 @@
 DECLARE_LOG_CATEGORY_EXTERN(LogDNAImportFactory, Log, All);
 
 /** Factory responsible for importing DNA file and attaching DNA data into SkeletalMesh
+*	Also extends ReimportHandler for importing DNA file with the same name as SkeletalMesh
  */
 UCLASS(transient)
-class UDNAAssetImportFactory: public UFactory
-{ 
+class UDNAAssetImportFactory : public UFactory, public FReimportHandler
+{
 	GENERATED_UCLASS_BODY()
 
 public:
@@ -29,8 +30,25 @@ public:
 	/** UObject properties */
 	virtual void PostInitProperties() override;
 
+	//~ Begin FReimportHandler Interface
+	virtual bool CanReimport(UObject* Obj, TArray<FString>& OutFilenames) override;
+	virtual void SetReimportPaths(UObject* Obj, const TArray<FString>& NewReimportPaths) override
+	{
+		return SetReimportPaths(Obj, NewReimportPaths[0], 0);
+	}
+	virtual void SetReimportPaths(UObject* Obj, const FString& NewReimportPath, const int32 SourceIndex) override {};
+
+	virtual EReimportResult::Type Reimport(UObject* Obj) override
+	{
+		return Reimport(Obj, INDEX_NONE);
+	}
+	virtual EReimportResult::Type Reimport(UObject* Obj, int32 SourceFileIndex);
+	virtual int32 GetPriority() const override;
+	virtual void PostImportCleanUp() { CleanUp(); }
+	//~ End FReimportHandler Interface
+
 	/** UFactory interface */
 	virtual UObject* FactoryCreateFile(UClass* InClass, UObject* InParent, FName InName, EObjectFlags Flags, const FString& Filename, const TCHAR* Parms, FFeedbackContext* Warn, bool& bOutOperationCanceled) override;
 	virtual bool FactoryCanImport(const FString& Filename) override;
-	virtual void CleanUp() override;
+	virtual void CleanUp() override {};
 };
diff --git a/Engine/Source/Editor/UnrealEd/Private/Factories/EditorFactories.cpp b/Engine/Source/Editor/UnrealEd/Private/Factories/EditorFactories.cpp
--- a/Engine/Source/Editor/UnrealEd/Private/Factories/EditorFactories.cpp
+++ b/Engine/Source/Editor/UnrealEd/Private/Factories/EditorFactories.cpp
@@ -6314,6 +6314,9 @@ bool UReimportFbxSkeletalMeshFactory::FactoryCanImport(const FString& Filename)
 bool UReimportFbxSkeletalMeshFactory::CanReimport( UObject* Obj, TArray<FString>& OutFilenames )
 {
 	USkeletalMesh* SkeletalMesh = Cast<USkeletalMesh>(Obj);
+	TArray<FString> FactoryExtensions;
+	GetSupportedFileExtensions(FactoryExtensions);
+
 	if (SkeletalMesh && !SkeletalMesh->HasCustomActorReimportFactory())
 	{
 		if (UAssetImportData* AssetImportData = SkeletalMesh->GetAssetImportData())
@@ -6330,6 +6333,11 @@ bool UReimportFbxSkeletalMeshFactory::CanReimport( UObject* Obj, TArray<FString>
 				//Fbx factory only support fbx and obj files
 				return false;
 			}
+			else if (!FactoryExtensions.Contains(FPaths::GetExtension(this->PreferredReimportPath)))
+			{
+				// Unsupported extensions shouldnt be considered for reimport
+				return false;
+			}
 			AssetImportData->ExtractFilenames(OutFilenames);
 		}
 		else
