Implemented a non destructive method of updating SkeletalMesh source data in order to preserve imported mesh

diff --git a/Engine/Plugins/Animation/RigLogic/Source/RigLogicModule/Private/SkelMeshDNAUtils.cpp b/Engine/Plugins/Animation/RigLogic/Source/RigLogicModule/Private/SkelMeshDNAUtils.cpp
--- a/Engine/Plugins/Animation/RigLogic/Source/RigLogicModule/Private/SkelMeshDNAUtils.cpp
+++ b/Engine/Plugins/Animation/RigLogic/Source/RigLogicModule/Private/SkelMeshDNAUtils.cpp
@@ -19,6 +19,7 @@
 #include "Rendering/SkeletalMeshLODRenderData.h"
 #if WITH_EDITORONLY_DATA
 #include "MeshUtilities.h"
+#include "LODUtilities.h"
 #endif // WITH_EDITORONLY_DATA
 #include "Modules/ModuleManager.h"
 
@@ -465,6 +466,126 @@ void USkelMeshDNAUtils::UpdateSourceData(USkeletalMesh* InSkelMesh)
 	MeshUtilities.CreateImportDataFromLODModel(InSkelMesh);
 }
 
+void USkelMeshDNAUtils::UpdateSourceData(USkeletalMesh* InSkelMesh, class IDNAReader* InDNAReader, class FDNAToSkelMeshMap* InDNAToSkelMeshMap)
+{
+	FSkeletalMeshModel* ImportedModel = InSkelMesh->GetImportedModel();
+	const int32 LODCount = ImportedModel->LODModels.Num();
+	const TArray<FTransform>& RawBonePose = InSkelMesh->GetRefSkeleton().GetRawRefBonePose();
+	for (int32 LODIndex = 0; LODIndex < LODCount; LODIndex++)
+	{
+		// Update vertices.	
+		const FSkeletalMeshLODModel& LODModel = ImportedModel->LODModels[LODIndex];
+
+		FSkeletalMeshImportData ImportData;
+		InSkelMesh->LoadLODImportedData(LODIndex, ImportData);
+
+		const int32 LODMeshVtxCount = LODModel.MeshToImportVertexMap.Num();
+		TArray<FSoftSkinVertex> LODVertices;
+		LODModel.GetVertices(LODVertices);
+
+		TArray<SkeletalMeshImportData::FRawBoneInfluence> NewInfluences;
+		TArray<bool> HasOverlappingVertices;
+		HasOverlappingVertices.AddZeroed(LODMeshVtxCount);
+		for (int32 LODMeshVtxIndex = 0; LODMeshVtxIndex < LODMeshVtxCount; LODMeshVtxIndex++)
+		{
+			// Update points.
+			int32 FbxVertexIndex = LODModel.MeshToImportVertexMap[LODMeshVtxIndex];
+			if (!HasOverlappingVertices[FbxVertexIndex])
+			{
+				HasOverlappingVertices[FbxVertexIndex] = true;
+				if (FbxVertexIndex <= LODModel.MaxImportVertex)
+				{
+					ImportData.Points[FbxVertexIndex] = LODVertices[LODMeshVtxIndex].Position;
+				}
+
+				// Update influences.
+				int32 SectionIdx;
+				int32 VertexIdx;
+				LODModel.GetSectionFromVertexIndex(LODMeshVtxIndex, SectionIdx, VertexIdx);
+				if (LODModel.Sections[SectionIdx].SoftVertices[VertexIdx].Color.B != 0)
+				{
+					int32 DNAMeshIndex = InDNAToSkelMeshMap->ImportVtxToDNAMeshIndex[LODIndex][LODMeshVtxIndex];
+					int32 DNAVertexIndex = InDNAToSkelMeshMap->ImportVtxToDNAVtxIndex[LODIndex][LODMeshVtxIndex];
+
+					if (DNAVertexIndex >= 0)
+					{
+						TArrayView<const float> DNASkinWeights = InDNAReader->GetSkinWeightsValues(DNAMeshIndex, DNAVertexIndex);
+						TArrayView<const uint16> DNASkinJoints = InDNAReader->GetSkinWeightsJointIndices(DNAMeshIndex, DNAVertexIndex);
+						uint16 SkinJointNum = DNASkinJoints.Num();
+						for (uint16 InfluenceIndex = 0; InfluenceIndex < SkinJointNum; ++InfluenceIndex)
+						{
+							float InfluenceWeight = DNASkinWeights[InfluenceIndex];
+							int32 UpdatedBoneId = InDNAToSkelMeshMap->GetUEBoneIndex(DNASkinJoints[InfluenceIndex]);
+
+							SkeletalMeshImportData::FRawBoneInfluence Influence;
+							Influence.VertexIndex = FbxVertexIndex;
+							Influence.BoneIndex = UpdatedBoneId;
+							Influence.Weight = InfluenceWeight;
+							NewInfluences.Add(Influence);
+						}
+						ImportData.Influences.RemoveAll([FbxVertexIndex](const SkeletalMeshImportData::FRawBoneInfluence& BoneInfluence) { return FbxVertexIndex == BoneInfluence.VertexIndex; });
+					}
+				}
+			}
+		}
+		ImportData.Influences.Append(NewInfluences);
+		// Sort influences by vertex index.
+		FLODUtilities::ProcessImportMeshInfluences(ImportData.Wedges.Num(), ImportData.Influences, InSkelMesh->GetPathName());
+
+		// Update reference pose.
+		const int32 JointCount = LODModel.RequiredBones.Num();
+		if (ImportData.RefBonesBinary.Num() == JointCount)
+		{
+			for (int32 JointIndex = 0; JointIndex < JointCount; JointIndex++)
+			{
+				const int32 OriginalBoneIndex = LODModel.RequiredBones[JointIndex];
+				const FTransform& UpdatedTransform = FTransform(RawBonePose[OriginalBoneIndex]);
+				ImportData.RefBonesBinary[OriginalBoneIndex].BonePos.Transform = FTransform3f(UpdatedTransform);
+			}
+		}
+
+		// Update morph targets. 
+		const int32 MorphTargetCount = InSkelMesh->GetMorphTargets().Num();
+		ImportData.MorphTargetModifiedPoints.Empty(MorphTargetCount);
+		ImportData.MorphTargetNames.Empty(MorphTargetCount);
+		ImportData.MorphTargets.Empty(MorphTargetCount);
+		if (LODIndex == 0)
+		{
+			// Blend shapes are used only in LOD0.
+			for (int32 MorphIndex = 0; MorphIndex < MorphTargetCount; MorphIndex++)
+			{
+				UMorphTarget* MorphTarget = InSkelMesh->GetMorphTargets()[MorphIndex];
+				// Add Morph target name.
+				ImportData.MorphTargetNames.Add(MorphTarget->GetName());
+				FSkeletalMeshImportData MorphTargetImportDeltas;
+				MorphTargetImportDeltas.bDiffPose = ImportData.bDiffPose;
+				MorphTargetImportDeltas.bUseT0AsRefPose = ImportData.bUseT0AsRefPose;
+
+				FMorphTargetLODModel& MorphLODModel = MorphTarget->GetMorphLODModels()[LODIndex];
+
+				// Init deltas and vertices for the current morph target.
+				int32 NumDeltas = MorphLODModel.Vertices.Num();
+				MorphTargetImportDeltas.Points.Reserve(NumDeltas);
+				TSet<uint32> MorphTargetImportVertices;
+				MorphTargetImportVertices.Reserve(NumDeltas);
+
+				FMorphTargetDelta* Deltas = MorphLODModel.Vertices.GetData();
+				for (int32 DeltaIndex = 0; DeltaIndex < NumDeltas; DeltaIndex++)
+				{
+					const uint32 SourceIndex = LODModel.MeshToImportVertexMap[Deltas[DeltaIndex].SourceIdx];
+					MorphTargetImportDeltas.Points.Add(ImportData.Points[SourceIndex] + Deltas[DeltaIndex].PositionDelta);
+					MorphTargetImportVertices.Add(SourceIndex);
+				}
+				ImportData.MorphTargetModifiedPoints.Add(MorphTargetImportVertices);
+				ImportData.MorphTargets.Add(MorphTargetImportDeltas);
+			}
+
+		}
+
+		InSkelMesh->SaveLODImportedData(LODIndex, ImportData);
+	}
+}
+
 UDNAAsset* USkelMeshDNAUtils::GetMeshDNA(USkeletalMesh* InSkelMesh)
 {
 	if (InSkelMesh && InSkelMesh->GetAssetUserDataArray())
diff --git a/Engine/Plugins/Animation/RigLogic/Source/RigLogicModule/Public/SkelMeshDNAUtils.h b/Engine/Plugins/Animation/RigLogic/Source/RigLogicModule/Public/SkelMeshDNAUtils.h
--- a/Engine/Plugins/Animation/RigLogic/Source/RigLogicModule/Public/SkelMeshDNAUtils.h
+++ b/Engine/Plugins/Animation/RigLogic/Source/RigLogicModule/Public/SkelMeshDNAUtils.h
@@ -55,6 +55,11 @@ public:
 	static void UpdateJointBehavior(USkeletalMeshComponent* InSkelMeshComponent);
 	/** Gets the DNA asset embedded in the mesh */
 	static UDNAAsset* GetMeshDNA(USkeletalMesh* InSkelMesh);
+
+	/** Updates source skeleton data for the purpose of character cooking and export.  */
+	static void UpdateSourceData(USkeletalMesh* InSkelMesh);
+	/** Updates ImportData structs for SkeletalMesh in non destructive way to preserve the original mesh */
+	static void UpdateSourceData(USkeletalMesh* InSkelMesh, class IDNAReader* InDNAReader, class FDNAToSkelMeshMap* InDNAToSkelMeshMap);
 #endif // WITH_EDITORONLY_DATA
 	/** Converts DNA vertex coordinates to UE4 coordinate system **/
 	inline static FVector ConvertDNAVertexToUE4CoordSystem(FVector InVertexPositionInDNA)
@@ -67,9 +72,6 @@ public:
 	{
 		return FVector{ -InVertexPositionInUE4.X, InVertexPositionInUE4.Y, -InVertexPositionInUE4.Z };
 	}
-
-	/** Updates source skeleton data for the purpose of character cooking and export.  */
-	static void UpdateSourceData(USkeletalMesh* InSkelMesh);
 	
 private:
 
