Move DNAIndexMapping onto DNAAsset to avoid lifetime issues

diff --git a/Engine/Plugins/Animation/RigLogic/Source/RigLogicModule/Private/AnimNode_RigLogic.cpp b/Engine/Plugins/Animation/RigLogic/Source/RigLogicModule/Private/AnimNode_RigLogic.cpp
--- a/Engine/Plugins/Animation/RigLogic/Source/RigLogicModule/Private/AnimNode_RigLogic.cpp
+++ b/Engine/Plugins/Animation/RigLogic/Source/RigLogicModule/Private/AnimNode_RigLogic.cpp
@@ -34,37 +34,6 @@ FAnimNode_RigLogic::~FAnimNode_RigLogic()
 	}
 }
 
-bool FAnimNode_RigLogic::NeedsOnInitializeAnimInstance() const
-{
-	return true;
-}
-
-void FAnimNode_RigLogic::OnInitializeAnimInstance(const FAnimInstanceProxy* InProxy, const UAnimInstance* InAnimInstance)
-{
-	DECLARE_SCOPE_HIERARCHICAL_COUNTER_FUNC()
-	QUICK_SCOPE_CYCLE_COUNTER(STAT_AnimNode_RigLogic_OnInitializeAnimInstance);
-	LLM_SCOPE_BYNAME(TEXT("Animation/RigLogic"));
-
-	USkeletalMeshComponent* SkeletalMeshComponent = InAnimInstance->GetSkelMeshComponent();
-	if (SkeletalMeshComponent == nullptr)
-	{
-		return;
-	}
-
-	USkeletalMesh* SkeletalMesh = SkeletalMeshComponent->GetSkeletalMeshAsset();
-	if (SkeletalMesh == nullptr)
-	{
-		return;
-	}
-
-	UDNAIndexMapping* DNAIndexMappingContainer = Cast<UDNAIndexMapping>(SkeletalMesh->GetAssetUserDataOfClass(UDNAIndexMapping::StaticClass()));
-	if (DNAIndexMappingContainer == nullptr)
-	{
-		DNAIndexMappingContainer = NewObject<UDNAIndexMapping>(SkeletalMesh);
-		SkeletalMesh->AddAssetUserData(DNAIndexMappingContainer);
-	}
-}
-
 void FAnimNode_RigLogic::Initialize_AnyThread(const FAnimationInitializeContext& Context)
 {
 	DECLARE_SCOPE_HIERARCHICAL_COUNTER_FUNC()
@@ -82,7 +51,6 @@ void FAnimNode_RigLogic::CacheBones_AnyThread(const FAnimationCacheBonesContext&
 
 	AnimSequence.CacheBones(Context);
 
-	// Initialize things that depend on the skeleton of Anim BP
 	USkeletalMeshComponent* SkeletalMeshComponent = Context.AnimInstanceProxy->GetSkelMeshComponent();
 	if (SkeletalMeshComponent == nullptr)
 	{
@@ -101,12 +69,6 @@ void FAnimNode_RigLogic::CacheBones_AnyThread(const FAnimationCacheBonesContext&
 		return;
 	}
 
-	UDNAIndexMapping* DNAIndexMappingContainer = Cast<UDNAIndexMapping>(SkeletalMesh->GetAssetUserDataOfClass(UDNAIndexMapping::StaticClass()));
-	if (DNAIndexMappingContainer == nullptr)
-	{
-		return;
-	}
-
 	UDNAAsset* DNAAsset = Cast<UDNAAsset>(SkeletalMesh->GetAssetUserDataOfClass(UDNAAsset::StaticClass()));
 	if (DNAAsset == nullptr)
 	{
@@ -129,7 +91,7 @@ void FAnimNode_RigLogic::CacheBones_AnyThread(const FAnimationCacheBonesContext&
 		RigInstance = new FRigInstance(LocalRigRuntimeContext->RigLogic.Get());
 	}
 
-	LocalDNAIndexMapping = DNAIndexMappingContainer->GetCachedMapping(LocalRigRuntimeContext->BehaviorReader.Get(), Skeleton, SkeletalMesh, SkeletalMeshComponent);
+	LocalDNAIndexMapping = DNAAsset->GetDNAIndexMapping(Skeleton, SkeletalMesh, SkeletalMeshComponent);
 	// CacheBones is called on LOD switches as well, in which case compact pose bone indices must be remapped
 	const FBoneContainer& RequiredBones = Context.AnimInstanceProxy->GetRequiredBones();
 	if (RequiredBones.IsValid())
diff --git a/Engine/Plugins/Animation/RigLogic/Source/RigLogicModule/Private/DNAAsset.cpp b/Engine/Plugins/Animation/RigLogic/Source/RigLogicModule/Private/DNAAsset.cpp
--- a/Engine/Plugins/Animation/RigLogic/Source/RigLogicModule/Private/DNAAsset.cpp
+++ b/Engine/Plugins/Animation/RigLogic/Source/RigLogicModule/Private/DNAAsset.cpp
@@ -6,6 +6,7 @@
 #include "ArchiveMemoryStream.h"
 #include "DNAAssetCustomVersion.h"
 #include "DNAReaderAdapter.h"
+#include "DNAIndexMapping.h"
 #include "FMemoryResource.h"
 #include "RigLogicMemoryStream.h"
 #include "SharedRigRuntimeContext.h"
@@ -121,13 +122,16 @@ void UDNAAsset::SetGeometryReader(TSharedPtr<IDNAReader> SourceDNAReader)
 void UDNAAsset::InvalidateRigRuntimeContext()
 {
 	FWriteScopeLock ScopeLock{RigRuntimeContextUpdateLock};
+	FWriteScopeLock MappingScopeLock(DNAIndexMappingUpdateLock);
 	RigRuntimeContext = nullptr;
+	DNAIndexMappingContainer.Empty(1);
 }
 
 TSharedPtr<FSharedRigRuntimeContext> UDNAAsset::GetRigRuntimeContext(EDNARetentionPolicy Policy)
 {
 	LLM_SCOPE_BYNAME(TEXT("Animation/RigLogic"));
 
+	FScopeLock DNAScopeLock{&DNAUpdateSection};
 	FRWScopeLock ContextScopeLock(RigRuntimeContextUpdateLock, SLT_ReadOnly);
 	if (!RigRuntimeContext.IsValid())
 	{
@@ -135,7 +139,6 @@ TSharedPtr<FSharedRigRuntimeContext> UDNAAsset::GetRigRuntimeContext(EDNARetenti
 		if (!RigRuntimeContext.IsValid())
 		{
 			TSharedPtr<FSharedRigRuntimeContext> NewContext = MakeShared<FSharedRigRuntimeContext>();
-			FScopeLock DNAScopeLock{&DNAUpdateSection};
 			if (BehaviorReader.IsValid() && (BehaviorReader->GetJointCount() != 0))
 			{
 				NewContext->BehaviorReader = BehaviorReader;
@@ -155,6 +158,48 @@ TSharedPtr<FSharedRigRuntimeContext> UDNAAsset::GetRigRuntimeContext(EDNARetenti
 	return RigRuntimeContext;
 }
 
+TSharedPtr<FDNAIndexMapping> UDNAAsset::GetDNAIndexMapping(const USkeleton* Skeleton,
+														   const USkeletalMesh* SkeletalMesh,
+														   const USkeletalMeshComponent* SkeletalMeshComponent)
+{
+	LLM_SCOPE_BYNAME(TEXT("Animation/RigLogic"));
+
+	FScopeLock DNAScopeLock{&DNAUpdateSection};
+	FWriteScopeLock MappingScopeLock(DNAIndexMappingUpdateLock);
+
+	// Find currently needed mapping and also clean stale objects along the way (requires only one iteration over the map)
+	TSharedPtr<FDNAIndexMapping> DNAIndexMapping;
+	for (auto Iterator = DNAIndexMappingContainer.CreateIterator(); Iterator; ++Iterator)
+	{
+		if (Iterator->Key.IsValid())
+		{
+			if (Iterator->Key == SkeletalMesh)
+			{
+				DNAIndexMapping = Iterator->Value;
+			}
+		}
+		else
+		{
+			Iterator.RemoveCurrent();
+		}
+	}
+
+	// Check if currently needed mapping exists, and if not, create it now
+	const FGuid SkeletonGuid = Skeleton->GetGuid();
+	if (!DNAIndexMapping.IsValid() || (SkeletonGuid != DNAIndexMapping->SkeletonGuid))
+	{
+		DNAIndexMapping = MakeShared<FDNAIndexMapping>();
+		DNAIndexMapping->SkeletonGuid = SkeletonGuid;
+		DNAIndexMapping->MapControlCurves(BehaviorReader.Get(), Skeleton);
+		DNAIndexMapping->MapJoints(BehaviorReader.Get(), SkeletalMeshComponent);
+		DNAIndexMapping->MapMorphTargets(BehaviorReader.Get(), Skeleton, SkeletalMesh);
+		DNAIndexMapping->MapMaskMultipliers(BehaviorReader.Get(), Skeleton);
+		DNAIndexMappingContainer.Add(SkeletalMesh, DNAIndexMapping);
+	}
+
+	return DNAIndexMapping;
+}
+
 bool UDNAAsset::Init(const FString& DNAFilename)
 {
 	LLM_SCOPE_BYNAME(TEXT("Animation/RigLogic"));
diff --git a/Engine/Plugins/Animation/RigLogic/Source/RigLogicModule/Private/DNAIndexMapping.cpp b/Engine/Plugins/Animation/RigLogic/Source/RigLogicModule/Private/DNAIndexMapping.cpp
--- a/Engine/Plugins/Animation/RigLogic/Source/RigLogicModule/Private/DNAIndexMapping.cpp
+++ b/Engine/Plugins/Animation/RigLogic/Source/RigLogicModule/Private/DNAIndexMapping.cpp
@@ -4,30 +4,11 @@
 
 #include "HAL/LowLevelMemTracker.h"
 #include "DNAReader.h"
-#include "Hasher.h"
 #include "Animation/Skeleton.h"
 #include "Animation/SmartName.h"
 #include "Components/SkeletalMeshComponent.h"
 #include "Engine/SkeletalMesh.h"
 
-#include UE_INLINE_GENERATED_CPP_BY_NAME(DNAIndexMapping)
-
-
-static uint32 HashDNA(const IDefinitionReader* Reader)
-{
-	Hasher DNAHasher;
-	DNAHasher.Update(Reader->GetTranslationUnit());
-	DNAHasher.Update(Reader->GetRotationUnit());
-	DNAHasher.Update(Reader->GetCoordinateSystem());
-	DNAHasher.Update(Reader->GetLODCount());
-	DNAHasher.Update(Reader->GetDBMaxLOD());
-	DNAHasher.Update(Reader->GetDBComplexity());
-	DNAHasher.Update(Reader->GetDBName());
-	DNAHasher.Update(Reader->GetJointCount());
-	DNAHasher.Update(Reader->GetBlendShapeChannelCount());
-	DNAHasher.Update(Reader->GetAnimatedMapCount());
-	return DNAHasher.GetHash();
-}
 
 /** Constructs curve name from nameToSplit using formatString of form x<obj>y<attr>z **/
 static FString CreateCurveName(const FString& NameToSplit, const FString& FormatString)
@@ -159,37 +140,3 @@ void FDNAIndexMapping::MapMaskMultipliers(const IBehaviorReader* DNABehavior, co
 		}
 	}
 }
-
-UDNAIndexMapping::UDNAIndexMapping() : Cached{nullptr}
-{
-}
-
-TSharedPtr<FDNAIndexMapping> UDNAIndexMapping::GetCachedMapping(const IBehaviorReader* DNABehavior,
-																const USkeleton* Skeleton,
-																const USkeletalMesh* SkeletalMesh,
-																const USkeletalMeshComponent* SkeletalMeshComponent)
-{
-	LLM_SCOPE_BYNAME(TEXT("Animation/RigLogic"));
-
-	const FGuid SkeletonGuid = Skeleton->GetGuid();
-	const uint32 DNAHash = HashDNA(DNABehavior);
-	// If GUID or DNA hash is different, all skeleton <-> DNA mappings should be recreated
-	FRWScopeLock ScopeLock(Lock, SLT_ReadOnly);
-	if (!Cached.IsValid() || (Cached->SkeletonGuid != SkeletonGuid) || (Cached->DNAHash != DNAHash))
-	{
-		ScopeLock.ReleaseReadOnlyLockAndAcquireWriteLock_USE_WITH_CAUTION();
-		if (!Cached.IsValid() || (Cached->SkeletonGuid != SkeletonGuid) || (Cached->DNAHash != DNAHash))
-		{
-			TSharedPtr<FDNAIndexMapping> NewMapping = MakeShared<FDNAIndexMapping>();
-			NewMapping->SkeletonGuid = SkeletonGuid;
-			NewMapping->DNAHash = DNAHash;
-			NewMapping->MapControlCurves(DNABehavior, Skeleton);
-			NewMapping->MapJoints(DNABehavior, SkeletalMeshComponent);
-			NewMapping->MapMorphTargets(DNABehavior, Skeleton, SkeletalMesh);
-			NewMapping->MapMaskMultipliers(DNABehavior, Skeleton);
-			Cached = NewMapping;
-		}
-	}
-	return Cached;
-}
-
diff --git a/Engine/Plugins/Animation/RigLogic/Source/RigLogicModule/Private/DNAIndexMapping.h b/Engine/Plugins/Animation/RigLogic/Source/RigLogicModule/Private/DNAIndexMapping.h
--- a/Engine/Plugins/Animation/RigLogic/Source/RigLogicModule/Private/DNAIndexMapping.h
+++ b/Engine/Plugins/Animation/RigLogic/Source/RigLogicModule/Private/DNAIndexMapping.h
@@ -2,19 +2,13 @@
 
 #pragma once
 
-#include "DNAReader.h"
-
-#include "Engine/AssetUserData.h"
+#include "CoreMinimal.h"
 #include "BoneIndices.h"
 
-#include "DNAIndexMapping.generated.h"
-
-
 class IBehaviorReader;
 class USkeleton;
 class USkeletalMesh;
 
-
 struct FDNAIndexMapping
 {
 	template <typename T>
@@ -23,7 +17,6 @@ struct FDNAIndexMapping
 		TArray<T> Values;
 	};
 
-	uint32 DNAHash;
 	FGuid SkeletonGuid;
 	// For maps, we use int32 instead of SmartName::UID_Ttype directly to allow storing INDEX_NONE for missing elements
 	// if value is valid, it is cast to SmartName::UID_Type
@@ -40,20 +33,3 @@ struct FDNAIndexMapping
 	void MapMaskMultipliers(const IBehaviorReader* DNABehavior, const USkeleton* Skeleton);
 
 };
-
-UCLASS(NotBlueprintable, hidecategories = (Object))
-class UDNAIndexMapping : public UAssetUserData
-{
-	GENERATED_BODY()
-
-public:
-	UDNAIndexMapping();
-	TSharedPtr<FDNAIndexMapping> GetCachedMapping(const IBehaviorReader* DNABehavior,
-												  const USkeleton* Skeleton,
-												  const USkeletalMesh* SkeletalMesh,
-												  const USkeletalMeshComponent* SkeletalMeshComponent);
-
-private:
-	FRWLock Lock;
-	TSharedPtr<FDNAIndexMapping> Cached;
-};
diff --git a/Engine/Plugins/Animation/RigLogic/Source/RigLogicModule/Public/AnimNode_RigLogic.h b/Engine/Plugins/Animation/RigLogic/Source/RigLogicModule/Public/AnimNode_RigLogic.h
--- a/Engine/Plugins/Animation/RigLogic/Source/RigLogicModule/Public/AnimNode_RigLogic.h
+++ b/Engine/Plugins/Animation/RigLogic/Source/RigLogicModule/Public/AnimNode_RigLogic.h
@@ -25,8 +25,6 @@ public:
 	FAnimNode_RigLogic();
 	~FAnimNode_RigLogic();
 
-	bool NeedsOnInitializeAnimInstance() const override;
-	void OnInitializeAnimInstance(const FAnimInstanceProxy* InProxy, const UAnimInstance* InAnimInstance) override;
 	void Initialize_AnyThread(const FAnimationInitializeContext& Context) override;
 	void CacheBones_AnyThread(const FAnimationCacheBonesContext& Context) override;
 	void Update_AnyThread(const FAnimationUpdateContext& Context) override;
diff --git a/Engine/Plugins/Animation/RigLogic/Source/RigLogicModule/Public/DNAAsset.h b/Engine/Plugins/Animation/RigLogic/Source/RigLogicModule/Public/DNAAsset.h
--- a/Engine/Plugins/Animation/RigLogic/Source/RigLogicModule/Public/DNAAsset.h
+++ b/Engine/Plugins/Animation/RigLogic/Source/RigLogicModule/Public/DNAAsset.h
@@ -20,6 +20,7 @@ class IBehaviorReader;
 class IGeometryReader;
 class FRigLogicMemoryStream;
 class UAssetUserData;
+struct FDNAIndexMapping;
 struct FSharedRigRuntimeContext;
 enum class EDNADataLayer: uint8;
 
@@ -71,6 +72,9 @@ private:
 
 	TSharedPtr<FSharedRigRuntimeContext> GetRigRuntimeContext(EDNARetentionPolicy Policy);
 	void InvalidateRigRuntimeContext();
+	TSharedPtr<FDNAIndexMapping> GetDNAIndexMapping(const USkeleton* Skeleton,
+													const USkeletalMesh* SkeletalMesh,
+													const USkeletalMeshComponent* SkeletalMeshComponent);
 
 private:
 	// Synchronize DNA updates
@@ -79,6 +83,9 @@ private:
 	// Synchronize Rig Runtime Context updates
 	FRWLock RigRuntimeContextUpdateLock;
 
+	// Synchronize DNA Index Mapping updates
+	FRWLock DNAIndexMappingUpdateLock;
+
 	/** Part of the .dna file needed for run-time execution of RigLogic;
 	 **/
 	TSharedPtr<IBehaviorReader> BehaviorReader;
@@ -92,4 +99,9 @@ private:
 	**/
 	TSharedPtr<FSharedRigRuntimeContext> RigRuntimeContext;
 
+	/** Container for Skeleton <-> DNAAsset index mappings
+	  * The mapping object owners will be the SkeletalMeshes, and periodic cleanups will
+	  * ensure that dead objects are deleted from the map.
+	 **/
+	TMap<TWeakObjectPtr<const USkeletalMesh>, TSharedPtr<FDNAIndexMapping>> DNAIndexMappingContainer;
 };
