Version bumps: Riglogic 8.1.0, GS 9.1.0, DNA 7.1.0, all dependencies updated to required versions.
GS: Added GenePoolSerialization.
DNA: StreamReader/Writer is now abstract class and it is required of user to change it to BinaryStreamReader/Writer.

diff --git a/Engine/Plugins/Animation/RigLogic/RigLogic.uplugin b/Engine/Plugins/Animation/RigLogic/RigLogic.uplugin
--- a/Engine/Plugins/Animation/RigLogic/RigLogic.uplugin
+++ b/Engine/Plugins/Animation/RigLogic/RigLogic.uplugin
@@ -1,9 +1,9 @@
 {
 	"FileVersion": 3,
-	"Version": 7,
-	"VersionName": "7.0.5",
-	"FriendlyName": "RigLogic Plugin v7.0.5",
-	"Description": "3Lateral RigLogic Plugin for Facial Animation v7.0.5",
+	"Version": 8.1,
+	"VersionName": "8.1.0",
+	"FriendlyName": "RigLogic Plugin v8.1.0",
+	"Description": "3Lateral RigLogic Plugin for Facial Animation v8.1.0",
 	"Category": "Animation",
 	"CreatedBy": "3Lateral",
 	"CreatedByURL": "https://www.3lateral.com/",
diff --git a/Engine/Plugins/Animation/RigLogic/Source/RigLogicLib/Private/dna/DNA.h b/Engine/Plugins/Animation/RigLogic/Source/RigLogicLib/Private/dna/DNA.h
--- a/Engine/Plugins/Animation/RigLogic/Source/RigLogicLib/Private/dna/DNA.h
+++ b/Engine/Plugins/Animation/RigLogic/Source/RigLogicLib/Private/dna/DNA.h
@@ -6,7 +6,7 @@
 #include "dna/SurjectiveMapping.h"
 #include "dna/TypeDefs.h"
 
-#include <terse/utils/ArchiveOffset.h>
+#include <terse/types/ArchiveOffset.h>
 
 #ifdef _MSC_VER
     #pragma warning(push)
@@ -32,7 +32,10 @@ struct RawSurjectiveMapping : public SurjectiveMapping<TFrom, TTo> {
 
     template<class Archive>
     void serialize(Archive& archive) {
-        archive(this->from, this->to);
+        archive.label("from");
+        archive(this->from);
+        archive.label("to");
+        archive(this->to);
     }
 
 };
@@ -47,11 +50,13 @@ struct ExpectedValue {
 
     template<class Archive>
     void load(Archive& archive) {
+        archive.label("value");
         archive(got);
     }
 
     template<class Archive>
     void save(Archive& archive) {
+        archive.label("value");
         archive(expected);
     }
 
@@ -72,6 +77,7 @@ struct Signature {
 
     template<class Archive>
     void serialize(Archive& archive) {
+        archive.label("data");
         archive(value);
     }
 
@@ -92,7 +98,10 @@ struct Version {
 
     template<class Archive>
     void serialize(Archive& archive) {
-        archive(generation, version);
+        archive.label("generation");
+        archive(generation);
+        archive.label("version");
+        archive(version);
     }
 
     bool matches() const {
@@ -113,14 +122,22 @@ struct SectionLookupTable {
 
     template<class Archive>
     void serialize(Archive& archive) {
-        archive(descriptor,
-                definition,
-                behavior,
-                controls,
-                joints,
-                blendShapeChannels,
-                animatedMaps,
-                geometry);
+        archive.label("descriptor");
+        archive(descriptor);
+        archive.label("definition");
+        archive(definition);
+        archive.label("behavior");
+        archive(behavior);
+        archive.label("controls");
+        archive(controls);
+        archive.label("joints");
+        archive(joints);
+        archive.label("blendShapeChannels");
+        archive(blendShapeChannels);
+        archive.label("animatedMaps");
+        archive(animatedMaps);
+        archive.label("geometry");
+        archive(geometry);
     }
 
 };
@@ -132,7 +149,12 @@ struct RawCoordinateSystem {
 
     template<class Archive>
     void serialize(Archive& archive) {
-        archive(xAxis, yAxis, zAxis);
+        archive.label("xAxis");
+        archive(xAxis);
+        archive.label("yAxis");
+        archive(yAxis);
+        archive.label("zAxis");
+        archive(zAxis);
     }
 
 };
@@ -142,7 +164,10 @@ struct RawLODMapping : public LODMapping {
 
     template<class Archive>
     void serialize(Archive& archive) {
-        archive(lods, indices);
+        archive.label("lods");
+        archive(lods);
+        archive.label("indices");
+        archive(indices);
     }
 
 };
@@ -182,19 +207,31 @@ struct RawDescriptor {
 
     template<class Archive>
     void serialize(Archive& archive) {
-        archive(marker,
-                name,
-                archetype,
-                gender,
-                age,
-                metadata,
-                translationUnit,
-                rotationUnit,
-                coordinateSystem,
-                lodCount,
-                maxLOD,
-                complexity,
-                dbName);
+        archive(marker);
+        archive.label("name");
+        archive(name);
+        archive.label("archetype");
+        archive(archetype);
+        archive.label("gender");
+        archive(gender);
+        archive.label("age");
+        archive(age);
+        archive.label("metadata");
+        archive(metadata);
+        archive.label("translationUnit");
+        archive(translationUnit);
+        archive.label("rotationUnit");
+        archive(rotationUnit);
+        archive.label("coordinateSystem");
+        archive(coordinateSystem);
+        archive.label("lodCount");
+        archive(lodCount);
+        archive.label("maxLOD");
+        archive(maxLOD);
+        archive.label("complexity");
+        archive(complexity);
+        archive.label("dbName");
+        archive(dbName);
     }
 
 };
@@ -224,7 +261,12 @@ struct RawVector3Vector {
 
     template<class Archive>
     void serialize(Archive& archive) {
-        archive(xs, ys, zs);
+        archive.label("xs");
+        archive(xs);
+        archive.label("ys");
+        archive(ys);
+        archive.label("zs");
+        archive(zs);
     }
 
     std::size_t size() const {
@@ -306,21 +348,35 @@ struct RawDefinition {
 
     template<class Archive>
     void serialize(Archive& archive) {
-        archive(marker,
-                lodJointMapping,
-                lodBlendShapeMapping,
-                lodAnimatedMapMapping,
-                lodMeshMapping,
-                guiControlNames,
-                rawControlNames,
-                jointNames,
-                blendShapeChannelNames,
-                animatedMapNames,
-                meshNames,
-                meshBlendShapeChannelMapping,
-                jointHierarchy,
-                neutralJointTranslations,
-                neutralJointRotations);
+        archive(marker);
+        archive.label("lodJointMapping");
+        archive(lodJointMapping);
+        archive.label("lodBlendShapeMapping");
+        archive(lodBlendShapeMapping);
+        archive.label("lodAnimatedMapMapping");
+        archive(lodAnimatedMapMapping);
+        archive.label("lodMeshMapping");
+        archive(lodMeshMapping);
+        archive.label("guiControlNames");
+        archive(guiControlNames);
+        archive.label("rawControlNames");
+        archive(rawControlNames);
+        archive.label("jointNames");
+        archive(jointNames);
+        archive.label("blendShapeChannelNames");
+        archive(blendShapeChannelNames);
+        archive.label("animatedMapNames");
+        archive(animatedMapNames);
+        archive.label("meshNames");
+        archive(meshNames);
+        archive.label("meshBlendShapeChannelMapping");
+        archive(meshBlendShapeChannelMapping);
+        archive.label("jointHierarchy");
+        archive(jointHierarchy);
+        archive.label("neutralJointTranslations");
+        archive(neutralJointTranslations);
+        archive.label("neutralJointRotations");
+        archive(neutralJointRotations);
     }
 
 };
@@ -344,12 +400,18 @@ struct RawConditionalTable {
 
     template<class Archive>
     void serialize(Archive& archive) {
-        archive(inputIndices,
-                outputIndices,
-                fromValues,
-                toValues,
-                slopeValues,
-                cutValues);
+        archive.label("inputIndices");
+        archive(inputIndices);
+        archive.label("outputIndices");
+        archive(outputIndices);
+        archive.label("fromValues");
+        archive(fromValues);
+        archive.label("toValues");
+        archive(toValues);
+        archive.label("slopeValues");
+        archive(slopeValues);
+        archive.label("cutValues");
+        archive(cutValues);
     }
 
 };
@@ -367,7 +429,12 @@ struct RawPSDMatrix {
 
     template<class Archive>
     void serialize(Archive& archive) {
-        archive(rows, columns, values);
+        archive.label("rows");
+        archive(rows);
+        archive.label("columns");
+        archive(columns);
+        archive.label("values");
+        archive(values);
     }
 
 };
@@ -385,7 +452,12 @@ struct RawControls {
 
     template<class Archive>
     void serialize(Archive& archive) {
-        archive(psdCount, conditionals, psds);
+        archive.label("psdCount");
+        archive(psdCount);
+        archive.label("conditionals");
+        archive(conditionals);
+        archive.label("psds");
+        archive(psds);
     }
 
 };
@@ -416,7 +488,16 @@ struct RawJointGroup {
 
     template<class Archive>
     void serialize(Archive& archive) {
-        archive(lods, inputIndices, outputIndices, values, jointIndices);
+        archive.label("lods");
+        archive(lods);
+        archive.label("inputIndices");
+        archive(inputIndices);
+        archive.label("outputIndices");
+        archive(outputIndices);
+        archive.label("values");
+        archive(values);
+        archive.label("jointIndices");
+        archive(jointIndices);
     }
 
 };
@@ -434,7 +515,12 @@ struct RawJoints {
 
     template<class Archive>
     void serialize(Archive& archive) {
-        archive(rowCount, colCount, jointGroups);
+        archive.label("rowCount");
+        archive(rowCount);
+        archive.label("colCount");
+        archive(colCount);
+        archive.label("jointGroups");
+        archive(jointGroups);
     }
 
 };
@@ -452,7 +538,12 @@ struct RawBlendShapeChannels {
 
     template<class Archive>
     void serialize(Archive& archive) {
-        archive(lods, inputIndices, outputIndices);
+        archive.label("lods");
+        archive(lods);
+        archive.label("inputIndices");
+        archive(inputIndices);
+        archive.label("outputIndices");
+        archive(outputIndices);
     }
 
 };
@@ -468,7 +559,10 @@ struct RawAnimatedMaps {
 
     template<class Archive>
     void serialize(Archive& archive) {
-        archive(lods, conditionals);
+        archive.label("lods");
+        archive(lods);
+        archive.label("conditionals");
+        archive(conditionals);
     }
 
 };
@@ -507,15 +601,18 @@ struct RawBehavior {
 
     template<class Archive>
     void serialize(Archive& archive) {
-        archive(marker,
-                controlsMarker,
-                controls,
-                jointsMarker,
-                joints,
-                blendShapeChannelsMarker,
-                blendShapeChannels,
-                animatedMapsMarker,
-                animatedMaps);
+        archive(marker, controlsMarker);
+        archive.label("controls");
+        archive(controls);
+        archive(jointsMarker);
+        archive.label("joints");
+        archive(joints);
+        archive(blendShapeChannelsMarker);
+        archive.label("blendShapeChannels");
+        archive(blendShapeChannels);
+        archive(animatedMapsMarker);
+        archive.label("animatedMaps");
+        archive(animatedMaps);
     }
 
 };
@@ -531,7 +628,10 @@ struct RawTextureCoordinateVector {
 
     template<class Archive>
     void serialize(Archive& archive) {
-        archive(us, vs);
+        archive.label("us");
+        archive(us);
+        archive.label("vs");
+        archive(vs);
     }
 
     std::size_t size() const {
@@ -559,7 +659,12 @@ struct RawVertexLayoutVector {
 
     template<class Archive>
     void serialize(Archive& archive) {
-        archive(positions, textureCoordinates, normals);
+        archive.label("positions");
+        archive(positions);
+        archive.label("textureCoordinates");
+        archive(textureCoordinates);
+        archive.label("normals");
+        archive(normals);
     }
 
     std::size_t size() const {
@@ -584,6 +689,7 @@ struct RawFace {
 
     template<class Archive>
     void serialize(Archive& archive) {
+        archive.label("layoutIndices");
         archive(layoutIndices);
     }
 
@@ -600,7 +706,10 @@ struct RawVertexSkinWeights {
 
     template<class Archive>
     void serialize(Archive& archive) {
-        archive(weights, jointIndices);
+        archive.label("weights");
+        archive(weights);
+        archive.label("jointIndices");
+        archive(jointIndices);
     }
 
 };
@@ -618,7 +727,12 @@ struct RawBlendShapeTarget {
 
     template<class Archive>
     void serialize(Archive& archive) {
-        archive(deltas, vertexIndices, blendShapeChannelIndex);
+        archive.label("deltas");
+        archive(deltas);
+        archive.label("vertexIndices");
+        archive(vertexIndices);
+        archive.label("blendShapeChannelIndex");
+        archive(blendShapeChannelIndex);
     }
 
 };
@@ -650,16 +764,25 @@ struct RawMesh {
 
     template<class Archive>
     void serialize(Archive& archive) {
-        archive(offset,
-                positions,
-                textureCoordinates,
-                normals,
-                layouts,
-                faces,
-                maximumInfluencePerVertex,
-                skinWeights,
-                blendShapeTargets,
-                marker);
+        archive.label("offset");
+        archive(offset);
+        archive.label("positions");
+        archive(positions);
+        archive.label("textureCoordinates");
+        archive(textureCoordinates);
+        archive.label("normals");
+        archive(normals);
+        archive.label("layouts");
+        archive(layouts);
+        archive.label("faces");
+        archive(faces);
+        archive.label("maximumInfluencePerVertex");
+        archive(maximumInfluencePerVertex);
+        archive.label("skinWeights");
+        archive(skinWeights);
+        archive.label("blendShapeTargets");
+        archive(blendShapeTargets);
+        archive(marker);
     }
 
 };
@@ -675,13 +798,15 @@ struct RawGeometry {
 
     template<class Archive>
     void serialize(Archive& archive) {
-        archive(marker, meshes);
+        archive(marker);
+        archive.label("meshes");
+        archive(meshes);
     }
 
 };
 
 struct DNA {
-	MemoryResource* memRes;
+    MemoryResource* memRes;
     Signature<3> signature{{'D', 'N', 'A'}};
     Version version{2, 1};
     SectionLookupTable sections;
@@ -707,32 +832,60 @@ struct DNA {
 
     template<class Archive>
     void load(Archive& archive) {
-        archive(signature, version);
+        archive.label("signature");
+        archive(signature);
+        archive.label("version");
+        archive(version);
         if (signature.matches() && version.matches()) {
-            archive(sections, descriptor, definition, behavior, geometry, eof);
+            archive.label("sections");
+            archive(sections);
+            archive.label("descriptor");
+            archive(descriptor);
+            archive.label("definition");
+            archive(definition);
+            archive.label("behavior");
+            archive(behavior);
+            archive.label("geometry");
+            archive(geometry);
+            archive.label("eof");
+            archive(eof);
             assert(eof.matches());
         }
     }
 
     template<class Archive>
     void save(Archive& archive) {
-        archive(signature, version, sections, descriptor, definition, behavior, geometry, eof);
+        archive.label("signature");
+        archive(signature);
+        archive.label("version");
+        archive(version);
+        archive.label("sections");
+        archive(sections);
+        archive.label("descriptor");
+        archive(descriptor);
+        archive.label("definition");
+        archive(definition);
+        archive.label("behavior");
+        archive(behavior);
+        archive.label("geometry");
+        archive(geometry);
+        archive.label("eof");
+        archive(eof);
     }
 
-	void unloadDefinition()
-	{
-		definition = RawDefinition{sections.definition, memRes};
-	}
+    void unloadDefinition() {
+        definition = RawDefinition{sections.definition, memRes};
+    }
 
-	void unloadBehavior()
-	{
-		behavior = RawBehavior{sections.behavior, sections.controls, sections.joints, sections.blendShapeChannels, sections.animatedMaps, memRes};
-	}
+    void unloadBehavior() {
+        behavior =
+            RawBehavior{sections.behavior, sections.controls, sections.joints, sections.blendShapeChannels, sections.animatedMaps,
+                        memRes};
+    }
 
-	void unloadGeometry()
-	{
-		geometry = RawGeometry{sections.geometry, memRes};
-	}
+    void unloadGeometry() {
+        geometry = RawGeometry{sections.geometry, memRes};
+    }
 
 };
 
diff --git a/Engine/Plugins/Animation/RigLogic/Source/RigLogicLib/Private/dna/LODConstraint.cpp b/Engine/Plugins/Animation/RigLogic/Source/RigLogicLib/Private/dna/LODConstraint.cpp
--- a/Engine/Plugins/Animation/RigLogic/Source/RigLogicLib/Private/dna/LODConstraint.cpp
+++ b/Engine/Plugins/Animation/RigLogic/Source/RigLogicLib/Private/dna/LODConstraint.cpp
@@ -29,7 +29,11 @@ LODConstraint::LODConstraint(ConstArrayView<std::uint16_t> lods_, MemoryResource
 }
 
 bool LODConstraint::hasImpactOn(std::uint16_t lodCount) const {
-    return (getMaxLOD() > 0u) || (getMinLOD() < static_cast<std::uint16_t>(lodCount - 1u));
+    std::uint16_t lod = {};
+    for (auto it = lods.begin(); (it != lods.end()) && (lod < lodCount); ++it) {
+        lod = static_cast<std::uint16_t>(lod + static_cast<std::uint16_t>(lod == *it));
+    }
+    return (lod != lodCount);
 }
 
 std::uint16_t LODConstraint::getMaxLOD() const {
@@ -40,6 +44,10 @@ std::uint16_t LODConstraint::getMinLOD() const {
     return (lods.empty() ? std::uint16_t{} : lods.back());
 }
 
+std::uint16_t LODConstraint::getLODCount() const {
+    return static_cast<std::uint16_t>(lods.size());
+}
+
 void LODConstraint::clampTo(std::uint16_t lodCount) {
     extd::filter(lods, [lodCount](std::uint16_t lod, std::size_t  /*unused*/) {
             return lod < lodCount;
diff --git a/Engine/Plugins/Animation/RigLogic/Source/RigLogicLib/Private/dna/LODConstraint.h b/Engine/Plugins/Animation/RigLogic/Source/RigLogicLib/Private/dna/LODConstraint.h
--- a/Engine/Plugins/Animation/RigLogic/Source/RigLogicLib/Private/dna/LODConstraint.h
+++ b/Engine/Plugins/Animation/RigLogic/Source/RigLogicLib/Private/dna/LODConstraint.h
@@ -16,6 +16,7 @@ class LODConstraint {
         bool hasImpactOn(std::uint16_t lodCount) const;
         std::uint16_t getMaxLOD() const;
         std::uint16_t getMinLOD() const;
+        std::uint16_t getLODCount() const;
         void clampTo(std::uint16_t lodCount);
         void applyTo(Vector<std::uint16_t>& unconstrainedLODs) const;
         void applyTo(DynArray<std::uint16_t>& unconstrainedLODs) const;
diff --git a/Engine/Plugins/Animation/RigLogic/Source/RigLogicLib/Private/dna/LODMapping.cpp b/Engine/Plugins/Animation/RigLogic/Source/RigLogicLib/Private/dna/LODMapping.cpp
--- a/Engine/Plugins/Animation/RigLogic/Source/RigLogicLib/Private/dna/LODMapping.cpp
+++ b/Engine/Plugins/Animation/RigLogic/Source/RigLogicLib/Private/dna/LODMapping.cpp
@@ -126,10 +126,12 @@ void LODMapping::associateLODWithIndices(std::uint16_t lod, std::uint16_t index)
     lods[lod] = index;
 }
 
-void LODMapping::mergeIndicesInto(UnorderedSet<std::uint16_t>& destination) const {
+UnorderedSet<std::uint16_t> LODMapping::getCombinedDistinctIndices(MemoryResource* memRes) const {
+    UnorderedSet<std::uint16_t> distinctIndices{memRes};
     for (const auto& row : indices) {
-        destination.insert(row.begin(), row.end());
+        distinctIndices.insert(row.begin(), row.end());
     }
+    return distinctIndices;
 }
 
 }  // namespace dna
diff --git a/Engine/Plugins/Animation/RigLogic/Source/RigLogicLib/Private/dna/LODMapping.h b/Engine/Plugins/Animation/RigLogic/Source/RigLogicLib/Private/dna/LODMapping.h
--- a/Engine/Plugins/Animation/RigLogic/Source/RigLogicLib/Private/dna/LODMapping.h
+++ b/Engine/Plugins/Animation/RigLogic/Source/RigLogicLib/Private/dna/LODMapping.h
@@ -35,7 +35,7 @@ class LODMapping {
         void associateLODWithIndices(std::uint16_t lod, std::uint16_t index);
         void mapIndices(std::function<std::uint16_t(std::uint16_t)> mapper);
         void filterIndices(std::function<bool(std::uint16_t)> filterer);
-        void mergeIndicesInto(UnorderedSet<std::uint16_t>& destination) const;
+        UnorderedSet<std::uint16_t> getCombinedDistinctIndices(MemoryResource* memRes) const;
 
     private:
         void cleanupIndices();
diff --git a/Engine/Plugins/Animation/RigLogic/Source/RigLogicLib/Private/dna/Writer.cpp b/Engine/Plugins/Animation/RigLogic/Source/RigLogicLib/Private/dna/Writer.cpp
--- a/Engine/Plugins/Animation/RigLogic/Source/RigLogicLib/Private/dna/Writer.cpp
+++ b/Engine/Plugins/Animation/RigLogic/Source/RigLogicLib/Private/dna/Writer.cpp
@@ -47,6 +47,7 @@ static void copyDescriptor(const DescriptorReader* source, DescriptorWriter* des
     destination->setArchetype(source->getArchetype());
     destination->setGender(source->getGender());
     destination->setAge(source->getAge());
+    destination->clearMetaData();
     for (std::uint32_t i = 0u; i < source->getMetaDataCount(); ++i) {
         const auto key = source->getMetaDataKey(i);
         const auto value = source->getMetaDataValue(key);
@@ -102,6 +103,22 @@ static void copyNameIndices(IndicesGetter getIndices,
 }
 
 static void copyDefinition(const DefinitionReader* source, DefinitionWriter* destination, MemoryResource* memRes) {
+    destination->clearGUIControlNames();
+    destination->clearRawControlNames();
+    destination->clearJointNames();
+    destination->clearBlendShapeChannelNames();
+    destination->clearAnimatedMapNames();
+    destination->clearMeshNames();
+    destination->clearJointIndices();
+    destination->clearLODJointMappings();
+    destination->clearBlendShapeChannelIndices();
+    destination->clearLODBlendShapeChannelMappings();
+    destination->clearAnimatedMapIndices();
+    destination->clearLODAnimatedMapMappings();
+    destination->clearMeshIndices();
+    destination->clearLODMeshMappings();
+    destination->clearMeshBlendShapeChannelMappings();
+
     const auto lodCount = source->getLODCount();
 
     for (std::uint16_t i = source->getGUIControlCount(); i > 0u; --i) {
@@ -181,6 +198,8 @@ static void copyDefinition(const DefinitionReader* source, DefinitionWriter* des
 }
 
 static void copyBehavior(const BehaviorReader* source, BehaviorWriter* destination, MemoryResource*  /*unused*/) {
+    destination->clearJointGroups();
+
     auto guiToRawInputIndices = source->getGUIToRawInputIndices();
     destination->setGUIToRawInputIndices(guiToRawInputIndices.data(), static_cast<std::uint16_t>(guiToRawInputIndices.size()));
 
@@ -274,7 +293,33 @@ static void copyBehavior(const BehaviorReader* source, BehaviorWriter* destinati
     destination->setAnimatedMapCutValues(animatedMapCutValues.data(), static_cast<std::uint16_t>(animatedMapCutValues.size()));
 }
 
+static bool hasGeometry(const GeometryReader* source) {
+    // Heuristic for determining whether source DNA actually has any geometry data, or mesh count is non-zero only
+    // because of mesh names stored in definition layer
+    std::uint32_t totalVertexCount = {};
+    std::uint32_t totalVertexNormalCount = {};
+    std::uint32_t totalTextureCoordCount = {};
+    std::uint32_t totalVertexLayoutCount = {};
+    std::uint32_t totalSkinWeightCount = {};
+    for (std::uint16_t meshIndex = {}; meshIndex < source->getMeshCount(); ++meshIndex) {
+        totalVertexCount += source->getVertexPositionCount(meshIndex);
+        totalVertexNormalCount += source->getVertexNormalCount(meshIndex);
+        totalTextureCoordCount += source->getVertexTextureCoordinateCount(meshIndex);
+        totalVertexLayoutCount += source->getVertexLayoutCount(meshIndex);
+        totalSkinWeightCount += source->getSkinWeightsCount(meshIndex);
+    }
+    return ((totalVertexCount != 0u) || (totalVertexNormalCount != 0u) || (totalTextureCoordCount != 0u) ||
+            (totalVertexLayoutCount != 0u) || (totalSkinWeightCount != 0u));
+}
+
 static void copyGeometry(const GeometryReader* source, GeometryWriter* destination, MemoryResource* memRes) {
+    destination->clearMeshes();
+
+    if (!hasGeometry(source)) {
+        // Source DNA was loaded without geometry layer
+        return;
+    }
+
     for (std::uint16_t meshIndexPlusOne = source->getMeshCount(); meshIndexPlusOne > 0u; --meshIndexPlusOne) {
         const auto meshIndex = static_cast<std::uint16_t>(meshIndexPlusOne - 1u);
         auto vertexCount = source->getVertexPositionCount(meshIndex);
@@ -326,7 +371,22 @@ static void copyGeometry(const GeometryReader* source, GeometryWriter* destinati
     }
 }
 
+static bool hasBlendShapeTargets(const GeometryReader* source) {
+    // Heuristic for determining whether source DNA actually has any blend shape target data, or mesh count is non-zero only
+    // because of mesh names stored in definition layer
+    std::uint32_t totalBlendShapeTargetCount = {};
+    for (std::uint16_t meshIndex = {}; meshIndex < source->getMeshCount(); ++meshIndex) {
+        totalBlendShapeTargetCount += source->getBlendShapeTargetCount(meshIndex);
+    }
+    return (totalBlendShapeTargetCount != 0u);
+}
+
 static void copyBlendShapeTargets(const GeometryReader* source, GeometryWriter* destination, MemoryResource* memRes) {
+    if (!hasBlendShapeTargets(source)) {
+        // Source DNA was loaded without blend shape targets
+        return;
+    }
+
     for (std::uint16_t meshIndexPlusOne = source->getMeshCount(); meshIndexPlusOne > 0u; --meshIndexPlusOne) {
         const auto meshIndex = static_cast<std::uint16_t>(meshIndexPlusOne - 1u);
         for (std::uint16_t blendShapeTargetIndexPlusOne = source->getBlendShapeTargetCount(meshIndex);
diff --git a/Engine/Plugins/Animation/RigLogic/Source/RigLogicLib/Private/dna/filters/AnimatedMapFilter.cpp b/Engine/Plugins/Animation/RigLogic/Source/RigLogicLib/Private/dna/filters/AnimatedMapFilter.cpp
--- a/Engine/Plugins/Animation/RigLogic/Source/RigLogicLib/Private/dna/filters/AnimatedMapFilter.cpp
+++ b/Engine/Plugins/Animation/RigLogic/Source/RigLogicLib/Private/dna/filters/AnimatedMapFilter.cpp
@@ -15,11 +15,13 @@ AnimatedMapFilter::AnimatedMapFilter(MemoryResource* memRes_) :
     remappedIndices{memRes} {
 }
 
-void AnimatedMapFilter::apply(RawDefinition& dest) {
-    // Collect all distinct element position indices that are referenced by the present LODs
-    dest.lodAnimatedMapMapping.mergeIndicesInto(passingIndices);
+void AnimatedMapFilter::configure(std::uint16_t animatedMapCount, UnorderedSet<std::uint16_t> allowedAnimatedMapIndices) {
+    passingIndices = std::move(allowedAnimatedMapIndices);
     // Fill the structure that maps indices prior to deletion to indices after deletion
-    remap(static_cast<std::uint16_t>(dest.animatedMapNames.size()), passingIndices, remappedIndices);
+    remap(animatedMapCount, passingIndices, remappedIndices);
+}
+
+void AnimatedMapFilter::apply(RawDefinition& dest) {
     // Fix indices so they match the same elements as earlier (but their
     // actual position changed with the deletion of the unneeded entries)
     dest.lodAnimatedMapMapping.mapIndices([this](std::uint16_t value) {
diff --git a/Engine/Plugins/Animation/RigLogic/Source/RigLogicLib/Private/dna/filters/AnimatedMapFilter.h b/Engine/Plugins/Animation/RigLogic/Source/RigLogicLib/Private/dna/filters/AnimatedMapFilter.h
--- a/Engine/Plugins/Animation/RigLogic/Source/RigLogicLib/Private/dna/filters/AnimatedMapFilter.h
+++ b/Engine/Plugins/Animation/RigLogic/Source/RigLogicLib/Private/dna/filters/AnimatedMapFilter.h
@@ -13,6 +13,7 @@ struct RawDefinition;
 class AnimatedMapFilter {
     public:
         explicit AnimatedMapFilter(MemoryResource* memRes_);
+        void configure(std::uint16_t animatedMapCount, UnorderedSet<std::uint16_t> allowedAnimatedMapIndices);
         void apply(RawDefinition& dest);
         bool passes(std::uint16_t index) const;
 
diff --git a/Engine/Plugins/Animation/RigLogic/Source/RigLogicLib/Private/dna/filters/BlendShapeFilter.cpp b/Engine/Plugins/Animation/RigLogic/Source/RigLogicLib/Private/dna/filters/BlendShapeFilter.cpp
--- a/Engine/Plugins/Animation/RigLogic/Source/RigLogicLib/Private/dna/filters/BlendShapeFilter.cpp
+++ b/Engine/Plugins/Animation/RigLogic/Source/RigLogicLib/Private/dna/filters/BlendShapeFilter.cpp
@@ -15,11 +15,13 @@ BlendShapeFilter::BlendShapeFilter(MemoryResource* memRes_) :
     remappedIndices{memRes} {
 }
 
-void BlendShapeFilter::apply(RawDefinition& dest) {
-    // Collect all distinct element position indices that are referenced by the present LODs
-    dest.lodBlendShapeMapping.mergeIndicesInto(passingIndices);
+void BlendShapeFilter::configure(std::uint16_t blendShapeCount, UnorderedSet<std::uint16_t> allowedBlendShapeIndices) {
+    passingIndices = std::move(allowedBlendShapeIndices);
     // Fill the structure that maps indices prior to deletion to indices after deletion
-    remap(static_cast<std::uint16_t>(dest.blendShapeChannelNames.size()), passingIndices, remappedIndices);
+    remap(blendShapeCount, passingIndices, remappedIndices);
+}
+
+void BlendShapeFilter::apply(RawDefinition& dest) {
     // Fix indices so they match the same elements as earlier (but their
     // actual position changed with the deletion of the unneeded entries)
     dest.lodBlendShapeMapping.mapIndices([this](std::uint16_t value) {
diff --git a/Engine/Plugins/Animation/RigLogic/Source/RigLogicLib/Private/dna/filters/BlendShapeFilter.h b/Engine/Plugins/Animation/RigLogic/Source/RigLogicLib/Private/dna/filters/BlendShapeFilter.h
--- a/Engine/Plugins/Animation/RigLogic/Source/RigLogicLib/Private/dna/filters/BlendShapeFilter.h
+++ b/Engine/Plugins/Animation/RigLogic/Source/RigLogicLib/Private/dna/filters/BlendShapeFilter.h
@@ -13,6 +13,7 @@ struct RawDefinition;
 class BlendShapeFilter {
     public:
         explicit BlendShapeFilter(MemoryResource* memRes_);
+        void configure(std::uint16_t blendShapeCount, UnorderedSet<std::uint16_t> allowedBlendShapeIndices);
         void apply(RawDefinition& dest);
         bool passes(std::uint16_t index) const;
 
diff --git a/Engine/Plugins/Animation/RigLogic/Source/RigLogicLib/Private/dna/filters/JointFilter.cpp b/Engine/Plugins/Animation/RigLogic/Source/RigLogicLib/Private/dna/filters/JointFilter.cpp
--- a/Engine/Plugins/Animation/RigLogic/Source/RigLogicLib/Private/dna/filters/JointFilter.cpp
+++ b/Engine/Plugins/Animation/RigLogic/Source/RigLogicLib/Private/dna/filters/JointFilter.cpp
@@ -11,15 +11,23 @@ namespace dna {
 
 JointFilter::JointFilter(MemoryResource* memRes_) :
     memRes{memRes_},
-    passingIndices{memRes},
-    remappedIndices{memRes} {
+    passingIndices{memRes_},
+    remappedIndices{memRes_},
+    option{Option::All},
+    rootJointIndex{} {
 }
 
-void JointFilter::apply(RawDefinition& dest) {
-    // Collect all distinct element position indices that are referenced by the present LODs
-    dest.lodJointMapping.mergeIndicesInto(passingIndices);
+void JointFilter::configure(std::uint16_t jointCount, UnorderedSet<std::uint16_t> allowedJointIndices, Option option_) {
+    option = option_;
+    passingIndices = std::move(allowedJointIndices);
     // Fill the structure that maps indices prior to deletion to indices after deletion
-    remap(static_cast<std::uint16_t>(dest.jointNames.size()), passingIndices, remappedIndices);
+    remap(jointCount, passingIndices, remappedIndices);
+}
+
+void JointFilter::apply(RawDefinition& dest) {
+    if (option != Option::All) {
+        return;
+    }
     // Fix indices so they match the same elements as earlier (but their
     // actual position changed with the deletion of the unneeded entries)
     dest.lodJointMapping.mapIndices([this](std::uint16_t value) {
@@ -32,6 +40,13 @@ void JointFilter::apply(RawDefinition& dest) {
     for (auto& jntIdx : dest.jointHierarchy) {
         jntIdx = remappedIndices[jntIdx];
     }
+    // Find root joint index
+    for (std::uint16_t jointIdx = 0u; jointIdx < dest.jointHierarchy.size(); ++jointIdx) {
+        if (dest.jointHierarchy[jointIdx] == jointIdx) {
+            rootJointIndex = jointIdx;
+            break;
+        }
+    }
     // Delete entries from other mappings that reference any of the deleted elements
     extd::filter(dest.neutralJointTranslations.xs, extd::byPosition(passingIndices));
     extd::filter(dest.neutralJointTranslations.ys, extd::byPosition(passingIndices));
@@ -41,6 +56,116 @@ void JointFilter::apply(RawDefinition& dest) {
     extd::filter(dest.neutralJointRotations.zs, extd::byPosition(passingIndices));
 }
 
+void JointFilter::apply(RawBehavior& dest) {
+    static constexpr std::uint16_t jointAttributeCount = 9u;
+
+    for (auto& jointGroup : dest.joints.jointGroups) {
+        if (option == Option::All) {
+            // Remove joint index from joint group and remap joint indices
+            extd::filter(jointGroup.jointIndices, [this](std::uint16_t jntIdx, std::size_t  /*unused*/) {
+                    return passes(jntIdx);
+                });
+            for (auto& jntIdx : jointGroup.jointIndices) {
+                jntIdx = remapped(jntIdx);
+            }
+        }
+        // Collect row indices of removed output indices to be used for joint delta removal
+        Set<std::size_t> rowsToDelete{memRes};
+        // Remove output indices belonging to the deletable joint
+        extd::filter(jointGroup.outputIndices, [this, &rowsToDelete](std::uint16_t outputIndex, std::size_t rowIndex) {
+                const auto jointIndex = static_cast<std::uint16_t>(outputIndex / jointAttributeCount);
+                if (!passes(jointIndex)) {
+                    rowsToDelete.insert(rowIndex);
+                    return false;
+                }
+                return true;
+            });
+
+        if (option == Option::All) {
+            // Remap the rest of output indices
+            for (auto& attrIdx : jointGroup.outputIndices) {
+                const auto jntIdx = static_cast<std::uint16_t>(attrIdx / jointAttributeCount);
+                const auto relAttrIdx = attrIdx - (jntIdx * jointAttributeCount);
+                attrIdx = static_cast<std::uint16_t>(remapped(jntIdx) * jointAttributeCount + relAttrIdx);
+            }
+        }
+
+        // If no animation data remains, there's no point in keeping input indices
+        const auto jointGroupColumnCount = static_cast<std::uint16_t>(jointGroup.inputIndices.size());
+        if (jointGroup.outputIndices.empty()) {
+            jointGroup.inputIndices.clear();
+        }
+
+        // Remove joint deltas associated with the removed output indices
+        extd::filter(jointGroup.values, [&rowsToDelete, jointGroupColumnCount](float  /*unused*/, std::size_t index) {
+                const std::uint16_t rowIndex = static_cast<std::uint16_t>(index / jointGroupColumnCount);
+                return (rowsToDelete.find(rowIndex) == rowsToDelete.end());
+            });
+        // Recompute LODs
+        for (auto& lod : jointGroup.lods) {
+            std::uint16_t decrementBy = 0u;
+            for (const auto rowIndex : rowsToDelete) {
+                if (rowIndex < lod) {
+                    ++decrementBy;
+                }
+            }
+            lod = static_cast<std::uint16_t>(lod - decrementBy);
+        }
+    }
+}
+
+void JointFilter::apply(RawVertexSkinWeights& dest) {
+    if (option != Option::All) {
+        return;
+    }
+
+    auto itWeightSrc = dest.weights.begin();
+    auto itWeightDst = itWeightSrc;
+    auto itJointSrc = dest.jointIndices.begin();
+    auto itJointDst = itJointSrc;
+    float discardedWeights = 0.0f;
+    while (itJointSrc != dest.jointIndices.end()) {
+        if (passes(*itJointSrc)) {
+            *itJointDst = *itJointSrc;
+            ++itJointSrc;
+            ++itJointDst;
+
+            *itWeightDst = *itWeightSrc;
+            ++itWeightSrc;
+            ++itWeightDst;
+        } else {
+            discardedWeights += *itWeightSrc;
+            ++itJointSrc;
+            ++itWeightSrc;
+        }
+    }
+    dest.jointIndices.resize(static_cast<std::size_t>(std::distance(dest.jointIndices.begin(), itJointDst)));
+    dest.weights.resize(static_cast<std::size_t>(std::distance(dest.weights.begin(), itWeightDst)));
+    assert(dest.jointIndices.size() == dest.weights.size());
+
+    if (passingIndices.empty()) {
+        return;
+    }
+
+    if (dest.jointIndices.empty()) {
+        // Reassign complete influence to root joint
+        dest.jointIndices.resize_uninitialized(1ul);
+        dest.jointIndices[0ul] = rootJointIndex;
+        dest.weights.resize_uninitialized(1ul);
+        dest.weights[0ul] = 1.0f;
+    } else {
+        // Normalize weights
+        for (auto& jntIdx : dest.jointIndices) {
+            jntIdx = remapped(jntIdx);
+        }
+
+        const float normalizationRatio = 1.0f / (1.0f - discardedWeights);
+        for (auto& weight : dest.weights) {
+            weight *= normalizationRatio;
+        }
+    }
+}
+
 bool JointFilter::passes(std::uint16_t index) const {
     return extd::contains(passingIndices, index);
 }
diff --git a/Engine/Plugins/Animation/RigLogic/Source/RigLogicLib/Private/dna/filters/JointFilter.h b/Engine/Plugins/Animation/RigLogic/Source/RigLogicLib/Private/dna/filters/JointFilter.h
--- a/Engine/Plugins/Animation/RigLogic/Source/RigLogicLib/Private/dna/filters/JointFilter.h
+++ b/Engine/Plugins/Animation/RigLogic/Source/RigLogicLib/Private/dna/filters/JointFilter.h
@@ -8,12 +8,24 @@
 
 namespace dna {
 
+struct RawBehavior;
 struct RawDefinition;
+struct RawLODMapping;
+struct RawVertexSkinWeights;
 
 class JointFilter {
+    public:
+        enum class Option {
+            All,
+            AnimationOnly
+        };
+
     public:
         explicit JointFilter(MemoryResource* memRes_);
+        void configure(std::uint16_t jointCount, UnorderedSet<std::uint16_t> allowedJointIndices, Option option_ = Option::All);
         void apply(RawDefinition& dest);
+        void apply(RawBehavior& dest);
+        void apply(RawVertexSkinWeights& dest);
         bool passes(std::uint16_t index) const;
         std::uint16_t remapped(std::uint16_t oldIndex) const;
         std::uint16_t maxRemappedIndex() const;
@@ -22,6 +34,8 @@ class JointFilter {
         MemoryResource* memRes;
         UnorderedSet<std::uint16_t> passingIndices;
         UnorderedMap<std::uint16_t, std::uint16_t> remappedIndices;
+        Option option;
+        std::uint16_t rootJointIndex;
 
 };
 
diff --git a/Engine/Plugins/Animation/RigLogic/Source/RigLogicLib/Private/dna/filters/MeshFilter.cpp b/Engine/Plugins/Animation/RigLogic/Source/RigLogicLib/Private/dna/filters/MeshFilter.cpp
--- a/Engine/Plugins/Animation/RigLogic/Source/RigLogicLib/Private/dna/filters/MeshFilter.cpp
+++ b/Engine/Plugins/Animation/RigLogic/Source/RigLogicLib/Private/dna/filters/MeshFilter.cpp
@@ -15,11 +15,13 @@ MeshFilter::MeshFilter(MemoryResource* memRes_) :
     remappedIndices{memRes} {
 }
 
-void MeshFilter::apply(RawDefinition& dest) {
-    // Collect all distinct element position indices that are referenced by the present LODs
-    dest.lodMeshMapping.mergeIndicesInto(passingIndices);
+void MeshFilter::configure(std::uint16_t meshCount, UnorderedSet<std::uint16_t> allowedMeshIndices) {
+    passingIndices = std::move(allowedMeshIndices);
     // Fill the structure that maps indices prior to deletion to indices after deletion
-    remap(static_cast<std::uint16_t>(dest.meshNames.size()), passingIndices, remappedIndices);
+    remap(meshCount, passingIndices, remappedIndices);
+}
+
+void MeshFilter::apply(RawDefinition& dest) {
     // Fix indices so they match the same elements as earlier (but their
     // actual position changed with the deletion of the unneeded entries)
     dest.lodMeshMapping.mapIndices([this](std::uint16_t value) {
diff --git a/Engine/Plugins/Animation/RigLogic/Source/RigLogicLib/Private/dna/filters/MeshFilter.h b/Engine/Plugins/Animation/RigLogic/Source/RigLogicLib/Private/dna/filters/MeshFilter.h
--- a/Engine/Plugins/Animation/RigLogic/Source/RigLogicLib/Private/dna/filters/MeshFilter.h
+++ b/Engine/Plugins/Animation/RigLogic/Source/RigLogicLib/Private/dna/filters/MeshFilter.h
@@ -13,6 +13,7 @@ struct RawDefinition;
 class MeshFilter {
     public:
         explicit MeshFilter(MemoryResource* memRes_);
+        void configure(std::uint16_t meshCount, UnorderedSet<std::uint16_t> allowedMeshIndices);
         void apply(RawDefinition& dest);
         bool passes(std::uint16_t index) const;
 
diff --git a/Engine/Plugins/Animation/RigLogic/Source/RigLogicLib/Private/dna/stream/BinaryStreamReaderImpl.cpp b/Engine/Plugins/Animation/RigLogic/Source/RigLogicLib/Private/dna/stream/BinaryStreamReaderImpl.cpp
new file mode 100644
--- /dev/null
+++ b/Engine/Plugins/Animation/RigLogic/Source/RigLogicLib/Private/dna/stream/BinaryStreamReaderImpl.cpp
@@ -0,0 +1,140 @@
+// Copyright Epic Games, Inc. All Rights Reserved.
+
+#include "dna/stream/BinaryStreamReaderImpl.h"
+
+#include "dna/TypeDefs.h"
+#include "dna/types/Limits.h"
+
+#include <status/Provider.h>
+#include <trio/utils/StreamScope.h>
+
+#ifdef _MSC_VER
+    #pragma warning(push)
+    #pragma warning(disable : 4365 4987)
+#endif
+#include <cstddef>
+#include <limits>
+#include <tuple>
+#ifdef _MSC_VER
+    #pragma warning(pop)
+#endif
+
+namespace dna {
+
+#ifdef __clang__
+    #pragma clang diagnostic push
+    #pragma clang diagnostic ignored "-Wglobal-constructors"
+#endif
+sc::StatusProvider BinaryStreamReaderImpl::status{SignatureMismatchError, VersionMismatchError, InvalidDataError};
+#ifdef __clang__
+    #pragma clang diagnostic pop
+#endif
+
+BinaryStreamReader::~BinaryStreamReader() = default;
+
+BinaryStreamReader* BinaryStreamReader::create(BoundedIOStream* stream,
+                                               DataLayer layer,
+                                               std::uint16_t maxLOD,
+                                               MemoryResource* memRes) {
+    PolyAllocator<BinaryStreamReaderImpl> alloc{memRes};
+    return alloc.newObject(stream, layer, maxLOD, LODLimits::min(), memRes);
+}
+
+BinaryStreamReader* BinaryStreamReader::create(BoundedIOStream* stream,
+                                               DataLayer layer,
+                                               std::uint16_t maxLOD,
+                                               std::uint16_t minLOD,
+                                               MemoryResource* memRes) {
+    PolyAllocator<BinaryStreamReaderImpl> alloc{memRes};
+    return alloc.newObject(stream, layer, maxLOD, minLOD, memRes);
+}
+
+BinaryStreamReader* BinaryStreamReader::create(BoundedIOStream* stream,
+                                               DataLayer layer,
+                                               std::uint16_t* lods,
+                                               std::uint16_t lodCount,
+                                               MemoryResource* memRes) {
+    PolyAllocator<BinaryStreamReaderImpl> alloc{memRes};
+    return alloc.newObject(stream, layer, ConstArrayView<std::uint16_t>{lods, lodCount}, memRes);
+}
+
+void BinaryStreamReader::destroy(BinaryStreamReader* instance) {
+    // NOLINTNEXTLINE(cppcoreguidelines-pro-type-static-cast-downcast)
+    auto reader = static_cast<BinaryStreamReaderImpl*>(instance);
+    PolyAllocator<BinaryStreamReaderImpl> alloc{reader->getMemoryResource()};
+    alloc.deleteObject(reader);
+}
+
+BinaryStreamReaderImpl::BinaryStreamReaderImpl(BoundedIOStream* stream_,
+                                               DataLayer layer_,
+                                               std::uint16_t maxLOD_,
+                                               std::uint16_t minLOD_,
+                                               MemoryResource* memRes_) :
+    BaseImpl{memRes_},
+    ReaderImpl{memRes_},
+    stream{stream_},
+    archive{stream_, layer_, maxLOD_, minLOD_, memRes_},
+    lodConstrained{(maxLOD_ != LODLimits::max()) || (minLOD_ != LODLimits::min())} {
+}
+
+BinaryStreamReaderImpl::BinaryStreamReaderImpl(BoundedIOStream* stream_,
+                                               DataLayer layer_,
+                                               ConstArrayView<std::uint16_t> lods_,
+                                               MemoryResource* memRes_) :
+    BaseImpl{memRes_},
+    ReaderImpl{memRes_},
+    stream{stream_},
+    archive{stream_, layer_, lods_, memRes_},
+    lodConstrained{true} {
+}
+
+bool BinaryStreamReaderImpl::isLODConstrained() const {
+    return lodConstrained;
+}
+
+void BinaryStreamReaderImpl::unload(DataLayer layer) {
+    if ((layer == DataLayer::All) ||
+        (layer == DataLayer::AllWithoutBlendShapes) ||
+        (layer == DataLayer::Descriptor)) {
+        dna = DNA{memRes};
+    } else if ((layer == DataLayer::Geometry) || (layer == DataLayer::GeometryWithoutBlendShapes)) {
+        dna.unloadGeometry();
+    } else if (layer == DataLayer::Behavior) {
+        dna.unloadBehavior();
+    } else if (layer == DataLayer::Definition) {
+        dna.unloadGeometry();
+        dna.unloadBehavior();
+        dna.unloadDefinition();
+    }
+}
+
+void BinaryStreamReaderImpl::read() {
+    // Due to possible usage of custom stream implementations, the status actually must be cleared at this point
+    // as external streams do not have access to the status reset API
+    status.reset();
+
+    trio::StreamScope scope{stream};
+    if (!sc::Status::isOk()) {
+        return;
+    }
+
+    archive >> dna;
+    if (!sc::Status::isOk()) {
+        return;
+    }
+
+    if (!dna.signature.matches()) {
+        status.set(SignatureMismatchError, dna.signature.value.expected.data(), dna.signature.value.got.data());
+        return;
+    }
+    if (!dna.version.matches()) {
+        status.set(VersionMismatchError,
+                   dna.version.generation.expected,
+                   dna.version.version.expected,
+                   dna.version.generation.got,
+                   dna.version.version.got);
+        return;
+    }
+}
+
+}  // namespace dna
diff --git a/Engine/Plugins/Animation/RigLogic/Source/RigLogicLib/Private/dna/stream/BinaryStreamReaderImpl.h b/Engine/Plugins/Animation/RigLogic/Source/RigLogicLib/Private/dna/stream/BinaryStreamReaderImpl.h
new file mode 100644
--- /dev/null
+++ b/Engine/Plugins/Animation/RigLogic/Source/RigLogicLib/Private/dna/stream/BinaryStreamReaderImpl.h
@@ -0,0 +1,38 @@
+// Copyright Epic Games, Inc. All Rights Reserved.
+
+#pragma once
+
+#include "dna/BinaryStreamReader.h"
+#include "dna/ReaderImpl.h"
+#include "dna/TypeDefs.h"
+#include "dna/stream/FilteredInputArchive.h"
+
+#include <status/Provider.h>
+
+namespace dna {
+
+class BinaryStreamReaderImpl : public ReaderImpl<BinaryStreamReader> {
+    public:
+        BinaryStreamReaderImpl(BoundedIOStream* stream_,
+                               DataLayer layer_,
+                               std::uint16_t maxLOD_,
+                               std::uint16_t minLOD_,
+                               MemoryResource* memRes_);
+        BinaryStreamReaderImpl(BoundedIOStream* stream_,
+                               DataLayer layer_,
+                               ConstArrayView<std::uint16_t> lods,
+                               MemoryResource* memRes_);
+
+        void unload(DataLayer layer) override;
+        void read() override;
+        bool isLODConstrained() const;
+
+    private:
+        static sc::StatusProvider status;
+
+        BoundedIOStream* stream;
+        FilteredInputArchive archive;
+        bool lodConstrained;
+};
+
+}  // namespace dna
diff --git a/Engine/Plugins/Animation/RigLogic/Source/RigLogicLib/Private/dna/stream/FilteredInputArchive.cpp b/Engine/Plugins/Animation/RigLogic/Source/RigLogicLib/Private/dna/stream/FilteredInputArchive.cpp
--- a/Engine/Plugins/Animation/RigLogic/Source/RigLogicLib/Private/dna/stream/FilteredInputArchive.cpp
+++ b/Engine/Plugins/Animation/RigLogic/Source/RigLogicLib/Private/dna/stream/FilteredInputArchive.cpp
@@ -88,7 +88,7 @@ void FilteredInputArchive::process(RawDescriptor& dest) {
     lodConstraint.clampTo(dest.lodCount);
     unconstrainedLODCount = dest.lodCount;
     dest.maxLOD = static_cast<std::uint16_t>(dest.maxLOD + lodConstraint.getMaxLOD());
-    dest.lodCount = static_cast<std::uint16_t>(lodConstraint.getMinLOD() - lodConstraint.getMaxLOD() + 1);
+    dest.lodCount = lodConstraint.getLODCount();
 }
 
 void FilteredInputArchive::process(RawDefinition& dest) {
@@ -102,14 +102,34 @@ void FilteredInputArchive::process(RawDefinition& dest) {
         return;
     }
 
+    // To find joints that are not in any LOD, find the joints that are not in LOD 0 (the current max LOD, at index 0), as it
+    // contains joints from all lower LODs.
+    Vector<std::uint16_t> jointsNotInLOD0{memRes};
+    const auto jointIndicesForLOD0 = dest.lodJointMapping.getIndices(0);
+    for (std::uint16_t idx = 0; idx < dest.jointNames.size(); ++idx) {
+        if (std::find(jointIndicesForLOD0.begin(), jointIndicesForLOD0.end(), idx) == jointIndicesForLOD0.end()) {
+            jointsNotInLOD0.push_back(idx);
+        }
+    }
+
     // Discard LOD data that is not relevant for the selected MaxLOD and MinLOD constraints
     dest.lodMeshMapping.discardLODs(lodConstraint);
     dest.lodJointMapping.discardLODs(lodConstraint);
     dest.lodBlendShapeMapping.discardLODs(lodConstraint);
     dest.lodAnimatedMapMapping.discardLODs(lodConstraint);
+    MeshFilter::configure(static_cast<std::uint16_t>(dest.meshNames.size()),
+                          dest.lodMeshMapping.getCombinedDistinctIndices(memRes));
     MeshFilter::apply(dest);
+    auto allowedJointIndices = dest.lodJointMapping.getCombinedDistinctIndices(memRes);
+    // In order to keep joints that are not in any LOD, add them all to the list of joints to keep when filtering.
+    allowedJointIndices.insert(jointsNotInLOD0.begin(), jointsNotInLOD0.end());
+    JointFilter::configure(static_cast<std::uint16_t>(dest.jointNames.size()), allowedJointIndices);
     JointFilter::apply(dest);
+    BlendShapeFilter::configure(static_cast<std::uint16_t>(dest.blendShapeChannelNames.size()),
+                                dest.lodBlendShapeMapping.getCombinedDistinctIndices(memRes));
     BlendShapeFilter::apply(dest);
+    AnimatedMapFilter::configure(static_cast<std::uint16_t>(dest.animatedMapNames.size()),
+                                 dest.lodAnimatedMapMapping.getCombinedDistinctIndices(memRes));
     AnimatedMapFilter::apply(dest);
 }
 
@@ -271,31 +291,7 @@ void FilteredInputArchive::process(RawVertexSkinWeights& dest) {
 
     if (lodConstraint.hasImpactOn(unconstrainedLODCount)) {
         assert(dest.weights.size() == dest.jointIndices.size());
-
-        auto itWeightSrc = dest.weights.begin();
-        auto itWeightDst = itWeightSrc;
-        auto itJointSrc = dest.jointIndices.begin();
-        auto itJointDst = itJointSrc;
-        while (itJointSrc != dest.jointIndices.end()) {
-            if (JointFilter::passes(*itJointSrc)) {
-                *itJointDst = *itJointSrc;
-                ++itJointSrc;
-                ++itJointDst;
-
-                *itWeightDst = *itWeightSrc;
-                ++itWeightSrc;
-                ++itWeightDst;
-            } else {
-                ++itJointSrc;
-                ++itWeightSrc;
-            }
-        }
-        dest.jointIndices.resize(static_cast<std::size_t>(std::distance(dest.jointIndices.begin(), itJointDst)));
-        dest.weights.resize(static_cast<std::size_t>(std::distance(dest.weights.begin(), itWeightDst)));
-
-        for (auto& jntIdx : dest.jointIndices) {
-            jntIdx = JointFilter::remapped(jntIdx);
-        }
+        JointFilter::apply(dest);
     }
 }
 
diff --git a/Engine/Plugins/Animation/RigLogic/Source/RigLogicLib/Private/dna/stream/FilteredInputArchive.h b/Engine/Plugins/Animation/RigLogic/Source/RigLogicLib/Private/dna/stream/FilteredInputArchive.h
--- a/Engine/Plugins/Animation/RigLogic/Source/RigLogicLib/Private/dna/stream/FilteredInputArchive.h
+++ b/Engine/Plugins/Animation/RigLogic/Source/RigLogicLib/Private/dna/stream/FilteredInputArchive.h
@@ -32,12 +32,14 @@ class FilteredInputArchive final : public AnimatedMapFilter, public BlendShapeFi
     public terse::ExtendableBinaryInputArchive<FilteredInputArchive,
                                                BoundedIOStream,
                                                std::uint32_t,
-                                               std::uint32_t> {
+                                               std::uint32_t,
+                                               terse::Endianness::Network> {
     private:
         using BaseArchive = terse::ExtendableBinaryInputArchive<FilteredInputArchive,
                                                                 BoundedIOStream,
                                                                 std::uint32_t,
-                                                                std::uint32_t>;
+                                                                std::uint32_t,
+                                                                terse::Endianness::Network>;
         friend Archive<FilteredInputArchive>;
 
     public:
diff --git a/Engine/Plugins/Animation/RigLogic/Source/RigLogicLib/Private/dna/utils/Extd.h b/Engine/Plugins/Animation/RigLogic/Source/RigLogicLib/Private/dna/utils/Extd.h
--- a/Engine/Plugins/Animation/RigLogic/Source/RigLogicLib/Private/dna/utils/Extd.h
+++ b/Engine/Plugins/Animation/RigLogic/Source/RigLogicLib/Private/dna/utils/Extd.h
@@ -131,7 +131,7 @@ TIterator advanced(TIterator source, TDistance distance) {
 }
 
 template<class TIterator, class Predicate>
-typename TIterator::difference_type advanceWhile(TIterator& it, const TIterator& end, Predicate pred) {
+typename std::iterator_traits<TIterator>::difference_type advanceWhile(TIterator& it, const TIterator& end, Predicate pred) {
     const auto start = it;
     while (it != end && pred(*it)) {
         ++it;
diff --git a/Engine/Plugins/Animation/RigLogic/Source/RigLogicLib/Private/riglogic/animatedmaps/AnimatedMapsFactory.cpp b/Engine/Plugins/Animation/RigLogic/Source/RigLogicLib/Private/riglogic/animatedmaps/AnimatedMapsFactory.cpp
--- a/Engine/Plugins/Animation/RigLogic/Source/RigLogicLib/Private/riglogic/animatedmaps/AnimatedMapsFactory.cpp
+++ b/Engine/Plugins/Animation/RigLogic/Source/RigLogicLib/Private/riglogic/animatedmaps/AnimatedMapsFactory.cpp
@@ -50,7 +50,8 @@ AnimatedMapsFactory::AnimatedMapsPtr AnimatedMapsFactory::create(const dna::Beha
                                   std::move(slopeValues),
                                   std::move(cutValues),
                                   inputCount,
-                                  outputCount};
+                                  outputCount,
+                                  memRes};
     return ManagedAnimatedMaps::with(memRes).create(std::move(lods), std::move(conditionals));
 }
 
diff --git a/Engine/Plugins/Animation/RigLogic/Source/RigLogicLib/Private/riglogic/conditionaltable/ConditionalTable.cpp b/Engine/Plugins/Animation/RigLogic/Source/RigLogicLib/Private/riglogic/conditionaltable/ConditionalTable.cpp
--- a/Engine/Plugins/Animation/RigLogic/Source/RigLogicLib/Private/riglogic/conditionaltable/ConditionalTable.cpp
+++ b/Engine/Plugins/Animation/RigLogic/Source/RigLogicLib/Private/riglogic/conditionaltable/ConditionalTable.cpp
@@ -24,9 +24,31 @@ namespace {
 const float clampMin = 0.0f;
 const float clampMax = 1.0f;
 
+Vector<std::uint16_t> buildIntervalSkipMap(const Vector<std::uint16_t>& inputIndices,
+                                           const Vector<std::uint16_t>& outputIndices,
+                                           MemoryResource* memRes) {
+    assert(inputIndices.size() == outputIndices.size());
+    Vector<std::uint16_t> intervalsRemaining{inputIndices.size(), {}, memRes};
+    for (std::size_t i = {}; i < inputIndices.size();) {
+        std::uint16_t intervalCount = 1u;
+        const std::uint16_t currentInputIndex = inputIndices[i];
+        const std::uint16_t currentOutputIndex = outputIndices[i];
+        for (std::size_t j = i + 1ul;
+             (j < inputIndices.size()) && (currentInputIndex == inputIndices[j]) && (currentOutputIndex == outputIndices[j]);
+             ++j) {
+            intervalsRemaining[j] = intervalCount++;
+        }
+        auto start = intervalsRemaining.data() + i;
+        std::reverse(start, start + intervalCount);
+        i += intervalCount;
+    }
+    return intervalsRemaining;
+}
+
 }  // namespace
 
 ConditionalTable::ConditionalTable(MemoryResource* memRes) :
+    intervalsRemaining{memRes},
     inputIndices{memRes},
     outputIndices{memRes},
     fromValues{memRes},
@@ -44,7 +66,9 @@ ConditionalTable::ConditionalTable(Vector<std::uint16_t>&& inputIndices_,
                                    Vector<float>&& slopeValues_,
                                    Vector<float>&& cutValues_,
                                    std::uint16_t inputCount_,
-                                   std::uint16_t outputCount_) :
+                                   std::uint16_t outputCount_,
+                                   MemoryResource* memRes) :
+    intervalsRemaining{buildIntervalSkipMap(inputIndices_, outputIndices_, memRes)},
     inputIndices{std::move(inputIndices_)},
     outputIndices{std::move(outputIndices_)},
     fromValues{std::move(fromValues_)},
@@ -66,15 +90,16 @@ std::uint16_t ConditionalTable::getOutputCount() const {
 void ConditionalTable::calculate(const float* inputs, float* outputs, std::uint16_t chunkSize) const {
     std::fill_n(outputs, outputCount, 0.0f);
 
-    for (std::uint16_t i = 0u; i < chunkSize; ++i) {
+    for (std::uint16_t i = {}; i < chunkSize; ++i) {
         const float inValue = inputs[inputIndices[i]];
         const float from = fromValues[i];
         const float to = toValues[i];
-        if ((from < inValue) && (inValue <= to)) {
+        if ((from <= inValue) && (inValue <= to)) {
             const std::uint16_t outIndex = outputIndices[i];
             const float slope = slopeValues[i];
             const float cut = cutValues[i];
             outputs[outIndex] += (slope * inValue + cut);
+            i = static_cast<std::uint16_t>(i + intervalsRemaining[i]);
         }
     }
 
diff --git a/Engine/Plugins/Animation/RigLogic/Source/RigLogicLib/Private/riglogic/conditionaltable/ConditionalTable.h b/Engine/Plugins/Animation/RigLogic/Source/RigLogicLib/Private/riglogic/conditionaltable/ConditionalTable.h
--- a/Engine/Plugins/Animation/RigLogic/Source/RigLogicLib/Private/riglogic/conditionaltable/ConditionalTable.h
+++ b/Engine/Plugins/Animation/RigLogic/Source/RigLogicLib/Private/riglogic/conditionaltable/ConditionalTable.h
@@ -18,7 +18,8 @@ class ConditionalTable {
                          Vector<float>&& slopeValues_,
                          Vector<float>&& cutValues_,
                          std::uint16_t inputCount_,
-                         std::uint16_t outputCount_);
+                         std::uint16_t outputCount_,
+                         MemoryResource* memRes);
 
         std::uint16_t getInputCount() const;
         std::uint16_t getOutputCount() const;
@@ -27,7 +28,8 @@ class ConditionalTable {
 
         template<class Archive>
         void serialize(Archive& archive) {
-            archive(inputIndices,
+            archive(intervalsRemaining,
+                    inputIndices,
                     outputIndices,
                     fromValues,
                     toValues,
@@ -38,6 +40,7 @@ class ConditionalTable {
         }
 
     private:
+        Vector<std::uint16_t> intervalsRemaining;
         Vector<std::uint16_t> inputIndices;
         Vector<std::uint16_t> outputIndices;
         Vector<float> fromValues;
diff --git a/Engine/Plugins/Animation/RigLogic/Source/RigLogicLib/Private/riglogic/controls/ControlsFactory.cpp b/Engine/Plugins/Animation/RigLogic/Source/RigLogicLib/Private/riglogic/controls/ControlsFactory.cpp
--- a/Engine/Plugins/Animation/RigLogic/Source/RigLogicLib/Private/riglogic/controls/ControlsFactory.cpp
+++ b/Engine/Plugins/Animation/RigLogic/Source/RigLogicLib/Private/riglogic/controls/ControlsFactory.cpp
@@ -39,6 +39,7 @@ ControlsFactory::ControlsPtr ControlsFactory::create(const dna::BehaviorReader*
             std::swap(fromValues[i], toValues[i]);
         }
     }
+
     ConditionalTable conditionals{std::move(inputIndices),
                                   std::move(outputIndices),
                                   std::move(fromValues),
@@ -46,7 +47,8 @@ ControlsFactory::ControlsPtr ControlsFactory::create(const dna::BehaviorReader*
                                   std::move(slopeValues),
                                   std::move(cutValues),
                                   reader->getGUIControlCount(),
-                                  reader->getRawControlCount()};
+                                  reader->getRawControlCount(),
+                                  memRes};
 
     Vector<std::uint16_t> psdRows{memRes};
     Vector<std::uint16_t> psdColumns{memRes};
diff --git a/Engine/Plugins/Animation/RigLogic/Source/RigLogicLib/Private/riglogic/joints/JointsBuilder.h b/Engine/Plugins/Animation/RigLogic/Source/RigLogicLib/Private/riglogic/joints/JointsBuilder.h
--- a/Engine/Plugins/Animation/RigLogic/Source/RigLogicLib/Private/riglogic/joints/JointsBuilder.h
+++ b/Engine/Plugins/Animation/RigLogic/Source/RigLogicLib/Private/riglogic/joints/JointsBuilder.h
@@ -2,7 +2,6 @@
 
 #pragma once
 
-#include "HAL/Platform.h"
 #include "riglogic/joints/JointsEvaluator.h"
 #include "riglogic/riglogic/Configuration.h"
 #include "riglogic/types/Aliases.h"
diff --git a/Engine/Plugins/Animation/RigLogic/Source/RigLogicLib/Private/riglogic/riglogic/RigInstanceImpl.cpp b/Engine/Plugins/Animation/RigLogic/Source/RigLogicLib/Private/riglogic/riglogic/RigInstanceImpl.cpp
--- a/Engine/Plugins/Animation/RigLogic/Source/RigLogicLib/Private/riglogic/riglogic/RigInstanceImpl.cpp
+++ b/Engine/Plugins/Animation/RigLogic/Source/RigLogicLib/Private/riglogic/riglogic/RigInstanceImpl.cpp
@@ -94,6 +94,11 @@ std::uint16_t RigInstanceImpl::getLOD() const {
 }
 
 void RigInstanceImpl::setLOD(std::uint16_t level) {
+    if (level != lodLevel) {
+        std::fill(jointOutputs.begin(), jointOutputs.end(), 0.0f);
+        std::fill(blendShapeOutputs.begin(), blendShapeOutputs.end(), 0.0f);
+        std::fill(animatedMapOutputs.begin(), animatedMapOutputs.end(), 0.0f);
+    }
     lodLevel = extd::clamp(level, static_cast<std::uint16_t>(0), lodMaxLevel);
 }
 
diff --git a/Engine/Plugins/Animation/RigLogic/Source/RigLogicLib/Private/riglogic/utils/Extd.h b/Engine/Plugins/Animation/RigLogic/Source/RigLogicLib/Private/riglogic/utils/Extd.h
--- a/Engine/Plugins/Animation/RigLogic/Source/RigLogicLib/Private/riglogic/utils/Extd.h
+++ b/Engine/Plugins/Animation/RigLogic/Source/RigLogicLib/Private/riglogic/utils/Extd.h
@@ -131,7 +131,7 @@ TIterator advanced(TIterator source, TDistance distance) {
 }
 
 template<class TIterator, class Predicate>
-typename TIterator::difference_type advanceWhile(TIterator& it, const TIterator& end, Predicate pred) {
+typename std::iterator_traits<TIterator>::difference_type advanceWhile(TIterator& it, const TIterator& end, Predicate pred) {
     const auto start = it;
     while (it != end && pred(*it)) {
         ++it;
diff --git a/Engine/Plugins/Animation/RigLogic/Source/RigLogicLib/Private/status/Storage.cpp b/Engine/Plugins/Animation/RigLogic/Source/RigLogicLib/Private/status/Storage.cpp
--- a/Engine/Plugins/Animation/RigLogic/Source/RigLogicLib/Private/status/Storage.cpp
+++ b/Engine/Plugins/Animation/RigLogic/Source/RigLogicLib/Private/status/Storage.cpp
@@ -19,11 +19,11 @@ struct StatusCodeStorage {
     char message[512];
 };
 
-void strcopy(char* destination, const char* source, std::size_t size) {
+void strcopy(char* destination, const char* source, std::size_t bufferSize) {
     #ifdef _MSC_VER
-        strncpy_s(destination, size, source, size);
+        strncpy_s(destination, bufferSize, source, bufferSize - 1ul);
     #else
-        std::strncpy(destination, source, size);
+        std::strncpy(destination, source, bufferSize - 1ul);
     #endif
 }
 
diff --git a/Engine/Plugins/Animation/RigLogic/Source/RigLogicLib/Public/dna/StreamReader.h b/Engine/Plugins/Animation/RigLogic/Source/RigLogicLib/Public/dna/StreamReader.h
--- a/Engine/Plugins/Animation/RigLogic/Source/RigLogicLib/Public/dna/StreamReader.h
+++ b/Engine/Plugins/Animation/RigLogic/Source/RigLogicLib/Public/dna/StreamReader.h
@@ -16,128 +16,11 @@ class DNAAPI StreamReader : public Reader {
         static const sc::StatusCode InvalidDataError;
 
     public:
-        /**
-            @brief Factory method for creation of Reader
-            @param stream
-                Source stream from which data is going to be read.
-            @param layer
-                Specify the layer up to which the data needs to be loaded.
-            @note
-                The Definition data layer depends on and thus implicitly loads the Descriptor layer.
-                The Behavior data layer depends on and thus implicitly loads the Definition layer.
-                The Geometry data layer depends on and thus also implicitly loads the Definition layer.
-            @param maxLOD
-                The maximum level of details to be loaded.
-            @note
-                A value of zero indicates to load all LODs.
-            @warning
-                The maxLOD value must be less than the value returned by getLODCount.
-            @see getLODCount
-            @param memRes
-                Memory resource to be used for allocations.
-            @note
-                If a memory resource is not given, a default allocation mechanism will be used.
-            @warning
-                User is responsible for releasing the returned pointer by calling destroy.
-            @see destroy
-        */
-        static StreamReader* create(BoundedIOStream* stream,
-                                    DataLayer layer = DataLayer::All,
-                                    std::uint16_t maxLOD = 0u,
-                                    MemoryResource* memRes = nullptr);
-        /**
-            @brief Factory method for creation of Reader
-            @param stream
-                Source stream from which data is going to be read.
-            @param layer
-                Specify the layer up to which the data needs to be loaded.
-            @note
-                The Definition data layer depends on and thus implicitly loads the Descriptor layer.
-                The Behavior data layer depends on and thus implicitly loads the Definition layer.
-                The Geometry data layer depends on and thus also implicitly loads the Definition layer.
-            @param maxLOD
-                The maximum level of details to be loaded.
-            @param minLOD
-                The minimum level of details to be loaded.
-            @note
-                A range of [0, LOD count - 1] for maxLOD / minLOD respectively indicates to load all LODs.
-            @warning
-                Both maxLOD and minLOD values must be less than the value returned by getLODCount.
-            @see getLODCount
-            @param memRes
-                Memory resource to be used for allocations.
-            @note
-                If a memory resource is not given, a default allocation mechanism will be used.
-            @warning
-                User is responsible for releasing the returned pointer by calling destroy.
-            @see destroy
-        */
-        static StreamReader* create(BoundedIOStream* stream,
-                                    DataLayer layer,
-                                    std::uint16_t maxLOD,
-                                    std::uint16_t minLOD,
-                                    MemoryResource* memRes = nullptr);
-        /**
-            @brief Factory method for creation of Reader
-            @param stream
-                Source stream from which data is going to be read.
-            @param layer
-                Specify the layer up to which the data needs to be loaded.
-            @note
-                The Definition data layer depends on and thus implicitly loads the Descriptor layer.
-                The Behavior data layer depends on and thus implicitly loads the Definition layer.
-                The Geometry data layer depends on and thus also implicitly loads the Definition layer.
-            @param lods
-                An array specifying which exact lods to load.
-            @warning
-                All values in the array must be less than the value returned by getLODCount.
-            @see getLODCount
-            @param lodCount
-                The number of elements in the lods array.
-            @warning
-                There cannot be more elements in the array than the value returned by getLODCount.
-            @see getLODCount
-            @param memRes
-                Memory resource to be used for allocations.
-            @note
-                If a memory resource is not given, a default allocation mechanism will be used.
-            @warning
-                User is responsible for releasing the returned pointer by calling destroy.
-            @see destroy
-        */
-        static StreamReader* create(BoundedIOStream* stream,
-                                    DataLayer layer,
-                                    std::uint16_t* lods,
-                                    std::uint16_t lodCount,
-                                    MemoryResource* memRes = nullptr);
-        /**
-            @brief Method for freeing a Reader instance.
-            @param instance
-                Instance of Reader to be freed.
-            @see create
-        */
-        static void destroy(StreamReader* instance);
-
         ~StreamReader() override;
         /**
            @brief read data from stream into internal structures.
         */
         virtual void read() = 0;
-
 };
 
 }  // namespace dna
-
-namespace pma {
-
-template<>
-struct DefaultInstanceCreator<dna::StreamReader> {
-    using type = pma::FactoryCreate<dna::StreamReader>;
-};
-
-template<>
-struct DefaultInstanceDestroyer<dna::StreamReader> {
-    using type = pma::FactoryDestroy<dna::StreamReader>;
-};
-
-}  // namespace pma
diff --git a/Engine/Plugins/Animation/RigLogic/Source/RigLogicLib/Public/dna/StreamWriter.h b/Engine/Plugins/Animation/RigLogic/Source/RigLogicLib/Public/dna/StreamWriter.h
--- a/Engine/Plugins/Animation/RigLogic/Source/RigLogicLib/Public/dna/StreamWriter.h
+++ b/Engine/Plugins/Animation/RigLogic/Source/RigLogicLib/Public/dna/StreamWriter.h
@@ -10,27 +10,6 @@ namespace dna {
 
 class DNAAPI StreamWriter : public Writer {
     public:
-        /**
-            @brief Factory method for creation of Writer
-            @param stream
-                Stream into which the data is going to be written.
-            @param memRes
-                Memory resource to be used for allocations.
-            @note
-                If a memory resource is not given, a default allocation mechanism will be used.
-            @warning
-                User is responsible for releasing the returned pointer by calling destroy.
-            @see destroy
-        */
-        static StreamWriter* create(BoundedIOStream* stream, MemoryResource* memRes = nullptr);
-        /**
-            @brief Method for freeing a Writer instance.
-            @param instance
-                Instance of Writer to be freed.
-            @see create
-        */
-        static void destroy(StreamWriter* instance);
-
         ~StreamWriter() override;
         /**
             @brief Write data to stream from internal structures.
@@ -39,17 +18,3 @@ class DNAAPI StreamWriter : public Writer {
 };
 
 }  // namespace dna
-
-namespace pma {
-
-template<>
-struct DefaultInstanceCreator<dna::StreamWriter> {
-    using type = pma::FactoryCreate<dna::StreamWriter>;
-};
-
-template<>
-struct DefaultInstanceDestroyer<dna::StreamWriter> {
-    using type = pma::FactoryDestroy<dna::StreamWriter>;
-};
-
-}  // namespace pma
diff --git a/Engine/Plugins/Animation/RigLogic/Source/RigLogicLib/Public/dna/version/Version.h b/Engine/Plugins/Animation/RigLogic/Source/RigLogicLib/Public/dna/version/Version.h
--- a/Engine/Plugins/Animation/RigLogic/Source/RigLogicLib/Public/dna/version/Version.h
+++ b/Engine/Plugins/Animation/RigLogic/Source/RigLogicLib/Public/dna/version/Version.h
@@ -2,7 +2,7 @@
 
 #pragma once
 
-#define DNA_MAJOR_VERSION 6
-#define DNA_MINOR_VERSION 0
+#define DNA_MAJOR_VERSION 7
+#define DNA_MINOR_VERSION 1
 #define DNA_PATCH_VERSION 0
-#define DNA_VERSION_STRING "6.0.0"
+#define DNA_VERSION_STRING "7.1.0"
diff --git a/Engine/Plugins/Animation/RigLogic/Source/RigLogicLib/Public/pma/version/Version.h b/Engine/Plugins/Animation/RigLogic/Source/RigLogicLib/Public/pma/version/Version.h
--- a/Engine/Plugins/Animation/RigLogic/Source/RigLogicLib/Public/pma/version/Version.h
+++ b/Engine/Plugins/Animation/RigLogic/Source/RigLogicLib/Public/pma/version/Version.h
@@ -4,5 +4,5 @@
 
 #define PMA_MAJOR_VERSION 1
 #define PMA_MINOR_VERSION 3
-#define PMA_PATCH_VERSION 0
-#define PMA_VERSION_STRING "1.3.0"
+#define PMA_PATCH_VERSION 3
+#define PMA_VERSION_STRING "1.3.3"
diff --git a/Engine/Plugins/Animation/RigLogic/Source/RigLogicLib/Public/riglogic/types/Aliases.h b/Engine/Plugins/Animation/RigLogic/Source/RigLogicLib/Public/riglogic/types/Aliases.h
--- a/Engine/Plugins/Animation/RigLogic/Source/RigLogicLib/Public/riglogic/types/Aliases.h
+++ b/Engine/Plugins/Animation/RigLogic/Source/RigLogicLib/Public/riglogic/types/Aliases.h
@@ -3,8 +3,10 @@
 #pragma once
 
 #include <dna/DataLayer.h>
-#include <dna/StreamReader.h>
-#include <dna/StreamWriter.h>
+#include <dna/BinaryStreamReader.h>
+#include <dna/BinaryStreamWriter.h>
+#include <dna/JSONStreamReader.h>
+#include <dna/JSONStreamWriter.h>
 #include <dna/types/Aliases.h>
 #include <dna/types/Vector3.h>
 #include <pma/MemoryResource.h>
@@ -26,8 +28,10 @@ using trio::FileStream;
 using trio::MemoryMappedFileStream;
 using trio::MemoryStream;
 using dna::DataLayer;
-using dna::StreamReader;
-using dna::StreamWriter;
+using dna::BinaryStreamReader;
+using dna::BinaryStreamWriter;
+using dna::JSONStreamReader;
+using dna::JSONStreamWriter;
 using dna::StringView;
 using dna::Vector3;
 
diff --git a/Engine/Plugins/Animation/RigLogic/Source/RigLogicLib/Public/riglogic/version/Version.h b/Engine/Plugins/Animation/RigLogic/Source/RigLogicLib/Public/riglogic/version/Version.h
--- a/Engine/Plugins/Animation/RigLogic/Source/RigLogicLib/Public/riglogic/version/Version.h
+++ b/Engine/Plugins/Animation/RigLogic/Source/RigLogicLib/Public/riglogic/version/Version.h
@@ -3,6 +3,6 @@
 #pragma once
 
 #define RL_MAJOR_VERSION 8
-#define RL_MINOR_VERSION 0
+#define RL_MINOR_VERSION 1
 #define RL_PATCH_VERSION 0
-#define RL_VERSION_STRING "8.0.0"
+#define RL_VERSION_STRING "8.1.0"
diff --git a/Engine/Plugins/Animation/RigLogic/Source/RigLogicLib/Public/status/version/Version.h b/Engine/Plugins/Animation/RigLogic/Source/RigLogicLib/Public/status/version/Version.h
--- a/Engine/Plugins/Animation/RigLogic/Source/RigLogicLib/Public/status/version/Version.h
+++ b/Engine/Plugins/Animation/RigLogic/Source/RigLogicLib/Public/status/version/Version.h
@@ -4,5 +4,5 @@
 
 #define SC_MAJOR_VERSION 1
 #define SC_MINOR_VERSION 1
-#define SC_PATCH_VERSION 1
-#define SC_VERSION_STRING "1.1.1"
+#define SC_PATCH_VERSION 4
+#define SC_VERSION_STRING "1.1.4"
diff --git a/Engine/Plugins/Animation/RigLogic/Source/RigLogicLib/Public/terse/Archive.h b/Engine/Plugins/Animation/RigLogic/Source/RigLogicLib/Public/terse/Archive.h
--- a/Engine/Plugins/Animation/RigLogic/Source/RigLogicLib/Public/terse/Archive.h
+++ b/Engine/Plugins/Animation/RigLogic/Source/RigLogicLib/Public/terse/Archive.h
@@ -16,7 +16,19 @@ namespace terse {
 template<class ArchiveImpl>
 class Archive {
     public:
-        explicit Archive(ArchiveImpl* impl_) : impl{impl_} {
+        explicit Archive(ArchiveImpl* impl_) : impl{impl_}, userData{nullptr} {
+        }
+
+        bool isOk() {
+            return impl->isOk();
+        }
+
+        void sync() {
+            impl->sync();
+        }
+
+        void label(const char* value) {
+            impl->label(value);
         }
 
         template<typename ... Args>
@@ -36,6 +48,14 @@ class Archive {
             return *impl;
         }
 
+        void* getUserData() const {
+            return userData;
+        }
+
+        void setUserData(void* data) {
+            userData = data;
+        }
+
     protected:
         template<typename Head>
         void dispatch(Head&& head) {
@@ -50,6 +70,7 @@ class Archive {
 
     private:
         ArchiveImpl* impl;
+        void* userData;
 };
 
 }  // namespace terse
diff --git a/Engine/Plugins/Animation/RigLogic/Source/RigLogicLib/Public/terse/archives/Traits.h b/Engine/Plugins/Animation/RigLogic/Source/RigLogicLib/Public/terse/archives/Traits.h
new file mode 100644
--- /dev/null
+++ b/Engine/Plugins/Animation/RigLogic/Source/RigLogicLib/Public/terse/archives/Traits.h
@@ -0,0 +1,198 @@
+// Copyright Epic Games, Inc. All Rights Reserved.
+
+#pragma once
+
+#ifdef _MSC_VER
+    #pragma warning(push)
+    #pragma warning(disable : 4365 4987)
+#endif
+#include <cstdint>
+#include <type_traits>
+#include <utility>
+#ifdef _MSC_VER
+    #pragma warning(pop)
+#endif
+
+static_assert(sizeof(char) == 1ul, "Unsupported platform, char is not 8-bits wide.");
+
+namespace terse {
+
+namespace traits {
+
+template<typename ...>
+struct sink {
+    using type = void;
+};
+
+template<typename T, typename = void>
+struct needs_allocator : std::false_type {};
+
+template<typename T>
+struct needs_allocator<T, typename sink<typename T::allocator_type,
+                                        decltype(std::declval<T>().get_allocator())>::type> : std::true_type {};
+
+template<class>
+struct true_sink : std::true_type {};
+
+template<class TContainer>
+static auto test_reserve_member(std::int32_t)->true_sink<decltype(std::declval<TContainer>().reserve(0u))>;
+
+template<class>
+static auto test_reserve_member(std::uint32_t)->std::false_type;
+
+template<class T>
+struct has_reserve_member : decltype(test_reserve_member<T>(0)) {};
+
+template<class TContainer>
+static auto test_push_back_member(std::int32_t)->true_sink<decltype(std::declval<TContainer>().push_back(
+                                                                        std::declval<typename TContainer::value_type>()))>;
+
+template<class>
+static auto test_push_back_member(std::uint32_t)->std::false_type;
+
+template<class T>
+struct has_push_back_member : decltype(test_push_back_member<T>(0)) {};
+
+// Serializer member functions
+
+template<class T>
+static auto test_serialize_member(std::int32_t)->true_sink<decltype(std::declval<T>().serialize(std::declval<T&>()))>;
+
+template<class>
+static auto test_serialize_member(std::uint32_t)->std::false_type;
+
+template<class T>
+struct has_serialize_member : decltype(test_serialize_member<T>(0)) {};
+
+template<class T>
+static auto test_load_member(std::int32_t)->true_sink<decltype(std::declval<T>().load(std::declval<T&>()))>;
+
+template<class>
+static auto test_load_member(std::uint32_t)->std::false_type;
+
+template<class T>
+struct has_load_member : decltype(test_load_member<T>(0)) {};
+
+template<class T>
+static auto test_save_member(std::int32_t)->true_sink<decltype(std::declval<T>().save(std::declval<T&>()))>;
+
+template<class>
+static auto test_save_member(std::uint32_t)->std::false_type;
+
+template<class T>
+struct has_save_member : decltype(test_save_member<T>(0)) {};
+
+// Serializer free functions
+
+template<class T>
+static auto test_serialize_function(std::int32_t)->true_sink<decltype(serialize(std::declval<T&>(), std::declval<T&>()))>;
+
+template<class>
+static auto test_serialize_function(std::uint32_t)->std::false_type;
+
+template<class T>
+struct has_serialize_function : decltype(test_serialize_function<T>(0)) {};
+
+template<class T>
+static auto test_load_function(std::int32_t)->true_sink<decltype(load(std::declval<T&>(), std::declval<T&>()))>;
+
+template<class>
+static auto test_load_function(std::uint32_t)->std::false_type;
+
+template<class T>
+struct has_load_function : decltype(test_load_function<T>(0)) {};
+
+template<class T>
+static auto test_save_function(std::int32_t)->true_sink<decltype(save(std::declval<T&>(), std::declval<T&>()))>;
+
+template<class>
+static auto test_save_function(std::uint32_t)->std::false_type;
+
+template<class T>
+struct has_save_function : decltype(test_save_function<T>(0)) {};
+
+// Versioned serializer member functions
+
+template<class T, class V>
+static auto test_versioned_serialize_member(std::int32_t)->true_sink <
+decltype(std::declval<T>().serialize(std::declval<T&>(), V{})) >;
+
+template<class, class>
+static auto test_versioned_serialize_member(std::uint32_t)->std::false_type;
+
+template<class T, class V>
+struct has_versioned_serialize_member : decltype(test_versioned_serialize_member<T, V>(0)) {};
+
+template<class T, class V>
+static auto test_versioned_load_member(std::int32_t)->true_sink < decltype(std::declval<T>().load(std::declval<T&>(), V{})) >;
+
+template<class, class>
+static auto test_versioned_load_member(std::uint32_t)->std::false_type;
+
+template<class T, class V>
+struct has_versioned_load_member : decltype(test_versioned_load_member<T, V>(0)) {};
+
+template<class T, class V>
+static auto test_versioned_save_member(std::int32_t)->true_sink < decltype(std::declval<T>().save(std::declval<T&>(), V{})) >;
+
+template<class, class>
+static auto test_versioned_save_member(std::uint32_t)->std::false_type;
+
+template<class T, class V>
+struct has_versioned_save_member : decltype(test_versioned_save_member<T, V>(0)) {};
+
+// Versioned serializer free functions
+
+template<class T, class V>
+static auto test_versioned_serialize_function(std::int32_t)->true_sink < decltype(serialize(std::declval<T&>(), V{},
+                                                                                            std::declval<T&>())) >;
+
+template<class, class>
+static auto test_versioned_serialize_function(std::uint32_t)->std::false_type;
+
+template<class T, class V>
+struct has_versioned_serialize_function : decltype(test_versioned_serialize_function<T, V>(0)) {};
+
+template<class T, class V>
+static auto test_versioned_load_function(std::int32_t)->true_sink < decltype(load(std::declval<T&>(), V{}, std::declval<T&>())) >;
+
+template<class, class>
+static auto test_versioned_load_function(std::uint32_t)->std::false_type;
+
+template<class T, class V>
+struct has_versioned_load_function : decltype(test_versioned_load_function<T, V>(0)) {};
+
+template<class T, class V>
+static auto test_versioned_save_function(std::int32_t)->true_sink < decltype(save(std::declval<T&>(), V{}, std::declval<T&>())) >;
+
+template<class, class>
+static auto test_versioned_save_function(std::uint32_t)->std::false_type;
+
+template<class T, class V>
+struct has_versioned_save_function : decltype(test_versioned_save_function<T, V>(0)) {};
+
+// Heuristic traits
+
+template<typename TContainer>
+using is_batchable = std::is_scalar<typename TContainer::value_type>;
+
+template<typename TContainer>
+struct has_wide_elements {
+    static constexpr bool value = (sizeof(typename TContainer::value_type) > 1ul);
+};
+
+template<typename T>
+struct is_pair : public std::false_type {};
+
+template<typename K, typename V>
+struct is_pair<std::pair<K, V> > : public std::true_type {};
+
+template<typename T>
+struct is_tuple : public std::false_type {};
+
+template<typename K, typename V>
+struct is_tuple<std::tuple<K, V> > : public std::true_type {};
+
+}  // namespace traits
+
+}  // namespace terse
diff --git a/Engine/Plugins/Animation/RigLogic/Source/RigLogicLib/Public/terse/archives/binary/InputArchive.h b/Engine/Plugins/Animation/RigLogic/Source/RigLogicLib/Public/terse/archives/binary/InputArchive.h
--- a/Engine/Plugins/Animation/RigLogic/Source/RigLogicLib/Public/terse/archives/binary/InputArchive.h
+++ b/Engine/Plugins/Animation/RigLogic/Source/RigLogicLib/Public/terse/archives/binary/InputArchive.h
@@ -3,9 +3,15 @@
 #pragma once
 
 #include "terse/Archive.h"
-#include "terse/archives/binary/Traits.h"
+#include "terse/archives/Common.h"
+#include "terse/archives/Traits.h"
+#include "terse/types/Anchor.h"
+#include "terse/types/ArchiveOffset.h"
+#include "terse/types/ArchiveSize.h"
+#include "terse/types/Blob.h"
 #include "terse/types/DynArray.h"
-#include "terse/utils/ArchiveOffset.h"
+#include "terse/types/Transparent.h"
+#include "terse/types/Versioned.h"
 #include "terse/utils/ByteSwap.h"
 
 #ifdef _MSC_VER
@@ -27,68 +33,19 @@
     #pragma warning(pop)
 #endif
 
-namespace pma {
-
-class MemoryResource;
-
-}  // namespace pma
-
 namespace terse {
 
-namespace impl {
-
-template<typename T>
-struct ValueFactory {
-    static constexpr bool NeedsAllocator = traits::needs_allocator<T>::value;
-    static constexpr bool NeedsMemoryResource = std::is_constructible<T, pma::MemoryResource*>::value;
-    static constexpr bool IsPair = traits::is_pair<T>::value;
-    static constexpr bool IsTuple = traits::is_tuple<T>::value;
-    static constexpr bool IsPrimitive = (!NeedsAllocator && !NeedsMemoryResource & !IsPair & !IsTuple);
-
-    template<class ParentAllocator, bool IsPrimitive = IsPrimitive>
-    static typename std::enable_if<IsPrimitive, T>::type create(const ParentAllocator&  /*unused*/) {
-        return T{};
-    }
-
-    template<class ParentAllocator, bool IsPrimitive = IsPrimitive>
-    static typename std::enable_if<!IsPrimitive && NeedsMemoryResource, T>::type create(const ParentAllocator& alloc) {
-        return T{alloc.getMemoryResource()};
-    }
-
-    template<class ParentAllocator, bool IsPrimitive = IsPrimitive>
-    static typename std::enable_if<!IsPrimitive && !NeedsMemoryResource && NeedsAllocator, T>::type create(
-        const ParentAllocator& alloc) {
-        using TAllocator = typename std::allocator_traits<ParentAllocator>::template rebind_alloc<typename T::value_type>;
-        return T{TAllocator{alloc}};
-    }
-
-    template<class ParentAllocator, bool IsPrimitive = IsPrimitive>
-    static typename std::enable_if<!IsPrimitive && !NeedsMemoryResource && !NeedsAllocator && IsPair, T>::type create(
-        const ParentAllocator& alloc) {
-        using K = typename T::first_type;
-        using V = typename T::second_type;
-        return T{ValueFactory<K>::create(alloc), ValueFactory<V>::create(alloc)};
-    }
-
-    template<class ParentAllocator, bool IsPrimitive = IsPrimitive>
-    static typename std::enable_if<!IsPrimitive && !NeedsMemoryResource && !NeedsAllocator && !IsPair && IsTuple, T>::type create(
-        const ParentAllocator& alloc) {
-        using K = typename std::tuple_element<0, T>::type;
-        using V = typename std::tuple_element<0, T>::type;
-        return T{ValueFactory<K>::create(alloc), ValueFactory<V>::create(alloc)};
-    }
-
-};
-
-}  // namespace impl
-
-template<class TExtender, class TStream, typename TSize, typename TOffset>
+template<class TExtender, class TStream, typename TSize, typename TOffset, Endianness EByteOrder>
 class ExtendableBinaryInputArchive : public Archive<TExtender> {
     public:
         // Given the possibility of both 32 and 64bit platforms, use a fixed width type during serialization
         using SizeType = TSize;
         using OffsetType = TOffset;
 
+        static constexpr Endianness endianness() {
+            return EByteOrder;
+        }
+
     private:
         using BaseArchive = Archive<TExtender>;
 
@@ -96,22 +53,30 @@ class ExtendableBinaryInputArchive : public Archive<TExtender> {
         ExtendableBinaryInputArchive(TExtender* extender, TStream* stream_) : BaseArchive{extender}, stream{stream_} {
         }
 
+        bool isOk() {
+            return true;
+        }
+
+        void sync() {
+        }
+
+        void label(const char*  /*unused*/) {
+        }
+
     protected:
+        template<typename T>
+        void process(Transparent<T>&& dest) {
+            process(dest.data);
+        }
+
+        void process(Anchor<OffsetType>& dest) {
+            dest.value = static_cast<OffsetType>(stream->tell());
+        }
+
         void process(ArchiveOffset<OffsetType>& dest) {
-            // Store the position of the offset itself, so it can be seeked to when writing the stream
-            #if !defined(__clang__) && defined(__GNUC__)
-                #pragma GCC diagnostic push
-                #pragma GCC diagnostic ignored "-Wuseless-cast"
-            #endif
-            dest.position = static_cast<decltype(dest.position)>(stream->tell());
-            #if !defined(__clang__) && defined(__GNUC__)
-                #pragma GCC diagnostic pop
-            #endif
             // Load the offset value itself (this points forward within the stream to the position of
             // the data with which the offset is associated)
             process(dest.value);
-            // Sanity check for making sure there is an associated proxy with the offset
-            assert(dest.proxy != nullptr);
         }
 
         void process(typename ArchiveOffset<OffsetType>::Proxy& dest) {
@@ -119,6 +84,60 @@ class ExtendableBinaryInputArchive : public Archive<TExtender> {
             stream->seek(dest.target->value);
         }
 
+        void process(typename ArchiveOffset<OffsetType>::Proxy&& dest) {
+            process(dest);
+        }
+
+        void process(ArchiveSize<SizeType, OffsetType>& dest) {
+            // Load the size value itself (this value can be combined with an offset to get the end of a data region)
+            process(dest.value);
+        }
+
+        void process(typename ArchiveSize<SizeType, OffsetType>::Proxy& dest) {
+            // Calculate the offset + size value stored in the associated `ArchiveSize` and `Anchor` and seek to it
+            assert(dest.base != nullptr);
+            stream->seek(dest.base->value + dest.target->value);
+        }
+
+        void process(typename ArchiveSize<SizeType, OffsetType>::Proxy&& dest) {
+            process(dest);
+        }
+
+        template<typename T, typename V>
+        typename std::enable_if<traits::has_versioned_load_member<T, V>::value,
+                                void>::type process(Versioned<T, V>&& dest) {
+            dest.data.load(*static_cast<TExtender*>(this), V{});
+        }
+
+        template<typename T, typename V>
+        typename std::enable_if<traits::has_versioned_serialize_member<T, V>::value,
+                                void>::type process(Versioned<T, V>&& dest) {
+            dest.data.serialize(*static_cast<TExtender*>(this), V{});
+        }
+
+        template<typename T, typename V>
+        typename std::enable_if<traits::has_versioned_load_function<T, V>::value,
+                                void>::type process(Versioned<T, V>&& dest) {
+            load(*static_cast<TExtender*>(this), V{}, dest.data);
+        }
+
+        template<typename T, typename V>
+        typename std::enable_if<traits::has_versioned_serialize_function<T, V>::value,
+                                void>::type process(Versioned<T, V>&& dest) {
+            serialize(*static_cast<TExtender*>(this), V{}, dest.data);
+        }
+
+        template<typename T, typename V>
+        typename std::enable_if<!traits::has_versioned_load_member<T, V>::value &&
+                                !traits::has_versioned_serialize_member<T, V>::value &&
+                                !traits::has_versioned_load_function<T, V>::value &&
+                                !traits::has_versioned_serialize_function<T, V>::value,
+                                void>::type process(Versioned<T, V>&& dest) {
+            // If no versioned serializer is found, but version information was passed along, fall back to unversioned
+            // functions if available.
+            BaseArchive::dispatch(dest.data);
+        }
+
         template<typename T>
         typename std::enable_if<traits::has_load_member<T>::value,
                                 void>::type process(T& dest) {
@@ -149,7 +168,17 @@ class ExtendableBinaryInputArchive : public Archive<TExtender> {
                                 void>::type process(T& dest) {
             // NOLINTNEXTLINE(cppcoreguidelines-pro-type-reinterpret-cast)
             stream->read(reinterpret_cast<char*>(&dest), sizeof(T));
-            networkToHost(dest);
+            SwapFrom<EByteOrder>::swap(dest);
+        }
+
+        template<typename T, typename ... Args>
+        void process(Blob<T, Args...>& dest) {
+            using ValueType = typename Blob<T, Args...>::value_type;
+            // Blob relies on a predetermined (user specified) size of how much data it should consume
+            if (dest.size() != 0ul) {
+                // NOLINTNEXTLINE(cppcoreguidelines-pro-type-reinterpret-cast)
+                stream->read(reinterpret_cast<char*>(dest.data()), dest.size() * sizeof(ValueType));
+            }
         }
 
         template<typename T, std::size_t N>
@@ -196,9 +225,13 @@ class ExtendableBinaryInputArchive : public Archive<TExtender> {
         }
 
         template<class TContainer>
-        typename std::enable_if<!traits::is_batchable<TContainer>::value>::type
+        typename std::enable_if<!traits::is_batchable<TContainer>::value
+                                && traits::has_reserve_member<TContainer>::value
+                                && traits::has_push_back_member<TContainer>::value
+                                >::type
         processElements(TContainer& dest, std::size_t size) {
             using ValueType = typename TContainer::value_type;
+            dest.clear();
             dest.reserve(size);
             for (std::size_t i = 0ul; i < size; ++i) {
                 dest.push_back(impl::ValueFactory<ValueType>::create(dest.get_allocator()));
@@ -206,6 +239,17 @@ class ExtendableBinaryInputArchive : public Archive<TExtender> {
             }
         }
 
+        template<class TContainer>
+        typename std::enable_if<!traits::is_batchable<TContainer>::value
+                                && !traits::has_reserve_member<TContainer>::value
+                                && !traits::has_push_back_member<TContainer>::value>::type
+        processElements(TContainer& dest, std::size_t size) {
+            resize(dest, size);
+            for (auto& element : dest) {
+                BaseArchive::dispatch(element);
+            }
+        }
+
         template<class TContainer>
         typename std::enable_if<traits::is_batchable<TContainer>::value && traits::has_wide_elements<TContainer>::value>::type
         processElements(TContainer& dest, std::size_t size) {
@@ -217,11 +261,11 @@ class ExtendableBinaryInputArchive : public Archive<TExtender> {
                 const std::size_t blockWidth = 16ul / sizeof(ValueType);
                 const std::size_t alignedSize = size - (size % blockWidth);
                 for (std::size_t i = 0ul; i < alignedSize; i += blockWidth) {
-                    networkToHost128(&dest[i]);
+                    SwapFrom<EByteOrder>::swap(&dest[i]);
                 }
 
                 for (std::size_t i = alignedSize; i < size; ++i) {
-                    networkToHost(dest[i]);
+                    SwapFrom<EByteOrder>::swap(dest[i]);
                 }
             }
         }
@@ -252,10 +296,14 @@ class ExtendableBinaryInputArchive : public Archive<TExtender> {
         TStream* stream;
 };
 
-template<class TStream, typename TSize = std::uint32_t, typename TOffset = TSize>
-class BinaryInputArchive : public ExtendableBinaryInputArchive<BinaryInputArchive<TStream>, TStream, TSize, TOffset> {
+template<class TStream, typename TSize = std::uint32_t, typename TOffset = TSize, Endianness EByteOrder = Endianness::Network>
+class BinaryInputArchive : public ExtendableBinaryInputArchive<BinaryInputArchive<TStream, TSize, TOffset, EByteOrder>,
+                                                               TStream,
+                                                               TSize,
+                                                               TOffset,
+                                                               EByteOrder> {
     private:
-        using BaseArchive = ExtendableBinaryInputArchive<BinaryInputArchive, TStream, TSize, TOffset>;
+        using BaseArchive = ExtendableBinaryInputArchive<BinaryInputArchive, TStream, TSize, TOffset, EByteOrder>;
         friend Archive<BinaryInputArchive>;
 
     public:
