Version bumps: Riglogic 8.1.0, GS 9.1.0, DNA 7.1.0, all dependencies updated to required versions.

diff --git a/Engine/Plugins/Animation/RigLogic/Source/RigLogicLib/Public/terse/archives/binary/OutputArchive.h b/Engine/Plugins/Animation/RigLogic/Source/RigLogicLib/Public/terse/archives/binary/OutputArchive.h
--- a/Engine/Plugins/Animation/RigLogic/Source/RigLogicLib/Public/terse/archives/binary/OutputArchive.h
+++ b/Engine/Plugins/Animation/RigLogic/Source/RigLogicLib/Public/terse/archives/binary/OutputArchive.h
@@ -3,9 +3,14 @@
 #pragma once
 
 #include "terse/Archive.h"
-#include "terse/archives/binary/Traits.h"
+#include "terse/archives/Traits.h"
+#include "terse/types/Anchor.h"
+#include "terse/types/ArchiveOffset.h"
+#include "terse/types/ArchiveSize.h"
+#include "terse/types/Blob.h"
 #include "terse/types/DynArray.h"
-#include "terse/utils/ArchiveOffset.h"
+#include "terse/types/Transparent.h"
+#include "terse/types/Versioned.h"
 #include "terse/utils/ByteSwap.h"
 
 #ifdef _MSC_VER
@@ -30,7 +35,7 @@
 
 namespace terse {
 
-template<class TExtender, class TStream, typename TSize, typename TOffset>
+template<class TExtender, class TStream, typename TSize, typename TOffset, Endianness EByteOrder>
 class ExtendableBinaryOutputArchive : public Archive<TExtender> {
     private:
         using BaseArchive = Archive<TExtender>;
@@ -40,14 +45,37 @@ class ExtendableBinaryOutputArchive : public Archive<TExtender> {
         using SizeType = TSize;
         using OffsetType = TOffset;
 
+        static constexpr Endianness endianness() {
+            return EByteOrder;
+        }
+
     public:
         ExtendableBinaryOutputArchive(TExtender* extender, TStream* stream_) : BaseArchive{extender}, stream{stream_} {
         }
 
+        bool isOk() {
+            return true;
+        }
+
+        void sync() {
+        }
+
+        void label(const char*  /*unused*/) {
+        }
+
     protected:
+        template<typename T>
+        void process(Transparent<T>&& source) {
+            process(source.data);
+        }
+
+        void process(Anchor<OffsetType>& source) {
+            source.value = static_cast<OffsetType>(stream->tell());
+        }
+
         void process(ArchiveOffset<OffsetType>& source) {
             // Record the position where the offset is going to be written so it can be seeked to later when
-            // it's proxy is found
+            // its proxy is found
             #if !defined(__clang__) && defined(__GNUC__)
                 #pragma GCC diagnostic push
                 #pragma GCC diagnostic ignored "-Wuseless-cast"
@@ -56,8 +84,8 @@ class ExtendableBinaryOutputArchive : public Archive<TExtender> {
             #if !defined(__clang__) && defined(__GNUC__)
                 #pragma GCC diagnostic pop
             #endif
-            // Since the actual offset value is not yet known at the time when it's declaration is encountered,
-            // fill it's place with zeros as a placeholder, and it will be later populated when it's associated
+            // Since the actual offset value is not yet known at the time when its declaration is encountered,
+            // fill its place with zeros as a placeholder, and it will be later populated when its associated
             // proxy is found
             process(typename ArchiveOffset<OffsetType>::ValueType{});
         }
@@ -66,14 +94,91 @@ class ExtendableBinaryOutputArchive : public Archive<TExtender> {
             // The current position of the stream needs to be written to the position where the associated
             // `ArchiveOffset` is found.
             auto current = stream->tell();
-            assert(current <= std::numeric_limits<typename ArchiveOffset<OffsetType>::ValueType>::max());
+            assert(current <= std::numeric_limits<OffsetType>::max());
+            source.target->value = static_cast<OffsetType>(current);
             // Seek to the actual position of the offset marker and write the stream's above captured position there
             stream->seek(source.target->position);
-            process(static_cast<typename ArchiveOffset<OffsetType>::ValueType>(current));
+            process(static_cast<OffsetType>(current));
             // Return to the earlier captured stream position so processing can safely resume from there
             stream->seek(current);
         }
 
+        void process(typename ArchiveOffset<OffsetType>::Proxy&& source) {
+            process(source);
+        }
+
+        void process(ArchiveSize<SizeType, OffsetType>& source) {
+            // Record the position where the size is going to be written so it can be seeked to later when
+            // its proxy is found
+            #if !defined(__clang__) && defined(__GNUC__)
+                #pragma GCC diagnostic push
+                #pragma GCC diagnostic ignored "-Wuseless-cast"
+            #endif
+            source.position = static_cast<decltype(source.position)>(stream->tell());
+            #if !defined(__clang__) && defined(__GNUC__)
+                #pragma GCC diagnostic pop
+            #endif
+            // Since the actual size value is not yet known at the time when it's declaration is encountered,
+            // fill it's place with zeros as a placeholder, and it will be later populated when it's associated
+            // proxy is found
+            process(typename ArchiveSize<SizeType, OffsetType>::ValueType{});
+        }
+
+        void process(typename ArchiveSize<SizeType, OffsetType>::Proxy& source) {
+            // The current position of the stream minus the base offset needs to be written to the position
+            // where the associated `ArchiveSize` is found.
+            const auto current = stream->tell();
+            assert(source.base != nullptr);
+            assert(source.base->value <= current);
+            auto size = current - source.base->value;
+            assert(size <= std::numeric_limits<SizeType>::max());
+            source.target->value = static_cast<SizeType>(size);
+            // Seek to the actual position of the size marker and write the above calculated size
+            stream->seek(source.target->position);
+            process(static_cast<SizeType>(size));
+            // Return to the earlier captured stream position so processing can safely resume from there
+            stream->seek(current);
+        }
+
+        void process(typename ArchiveSize<SizeType, OffsetType>::Proxy&& source) {
+            process(source);
+        }
+
+        template<typename T, typename V>
+        typename std::enable_if<traits::has_versioned_save_member<T, V>::value,
+                                void>::type process(const Versioned<T, V>& source) {
+            const_cast<T&>(source.data).save(*static_cast<TExtender*>(this), V{});
+        }
+
+        template<typename T, typename V>
+        typename std::enable_if<traits::has_versioned_serialize_member<T, V>::value,
+                                void>::type process(const Versioned<T, V>& source) {
+            const_cast<T&>(source.data).serialize(*static_cast<TExtender*>(this), V{});
+        }
+
+        template<typename T, typename V>
+        typename std::enable_if<traits::has_versioned_save_function<T, V>::value,
+                                void>::type process(const Versioned<T, V>& source) {
+            save(*static_cast<TExtender*>(this), V{}, const_cast<T&>(source.data));
+        }
+
+        template<typename T, typename V>
+        typename std::enable_if<traits::has_versioned_serialize_function<T, V>::value,
+                                void>::type process(const Versioned<T, V>& source) {
+            serialize(*static_cast<TExtender*>(this), V{}, const_cast<T&>(source.data));
+        }
+
+        template<typename T, typename V>
+        typename std::enable_if<!traits::has_versioned_save_member<T, V>::value &&
+                                !traits::has_versioned_serialize_member<T, V>::value &&
+                                !traits::has_versioned_save_function<T, V>::value &&
+                                !traits::has_versioned_serialize_function<T, V>::value,
+                                void>::type process(Versioned<T, V>&& dest) {
+            // If no versioned serializer is found, but version information was passed along, fall back to unversioned
+            // functions if available.
+            BaseArchive::dispatch(dest.data);
+        }
+
         template<typename T>
         typename std::enable_if<traits::has_save_member<T>::value,
                                 void>::type process(const T& source) {
@@ -102,11 +207,20 @@ class ExtendableBinaryOutputArchive : public Archive<TExtender> {
         typename std::enable_if<!traits::has_save_member<T>::value && !traits::has_serialize_member<T>::value &&
                                 !traits::has_save_function<T>::value && !traits::has_serialize_function<T>::value,
                                 void>::type process(const T& source) {
-            T networkOrder;
-            std::memcpy(&networkOrder, &source, sizeof(T));
-            hostToNetwork(networkOrder);
+            T swapped;
+            std::memcpy(&swapped, &source, sizeof(T));
+            SwapTo<EByteOrder>::swap(swapped);
             // NOLINTNEXTLINE(cppcoreguidelines-pro-type-reinterpret-cast)
-            stream->write(reinterpret_cast<char*>(&networkOrder), sizeof(T));
+            stream->write(reinterpret_cast<char*>(&swapped), sizeof(T));
+        }
+
+        template<typename T, typename ... Args>
+        void process(const Blob<T, Args...>& source) {
+            using ValueType = typename Blob<T, Args...>::value_type;
+            if (source.size() != 0ul) {
+                // NOLINTNEXTLINE(cppcoreguidelines-pro-type-reinterpret-cast)
+                stream->write(reinterpret_cast<const char*>(source.data()), source.size() * sizeof(ValueType));
+            }
         }
 
         template<typename T, std::size_t N>
@@ -150,20 +264,44 @@ class ExtendableBinaryOutputArchive : public Archive<TExtender> {
         }
 
         template<class TContainer>
-        void processElements(const TContainer& source) {
+        typename std::enable_if<!traits::is_batchable<TContainer>::value>::type
+        processElements(const TContainer& source) {
             for (const auto& element : source) {
                 BaseArchive::dispatch(element);
             }
         }
 
+        template<class TContainer>
+        typename std::enable_if<traits::is_batchable<TContainer>::value && traits::has_wide_elements<TContainer>::value>::type
+        processElements(const TContainer& source) {
+            for (const auto& element : source) {
+                BaseArchive::dispatch(element);
+            }
+        }
+
+        template<class TContainer>
+        typename std::enable_if<traits::is_batchable<TContainer>::value && !traits::has_wide_elements<TContainer>::value>::type
+        processElements(const TContainer& source) {
+            using ValueType = typename TContainer::value_type;
+            const auto size = source.size();
+            if (size != 0ul) {
+                // NOLINTNEXTLINE(cppcoreguidelines-pro-type-reinterpret-cast)
+                stream->write(reinterpret_cast<const char*>(&source[0]), size * sizeof(ValueType));
+            }
+        }
+
     private:
         TStream* stream;
 };
 
-template<class TStream, typename TSize = std::uint32_t, typename TOffset = TSize>
-class BinaryOutputArchive : public ExtendableBinaryOutputArchive<BinaryOutputArchive<TStream>, TStream, TSize, TOffset> {
+template<class TStream, typename TSize = std::uint32_t, typename TOffset = TSize, Endianness EByteOrder = Endianness::Network>
+class BinaryOutputArchive : public ExtendableBinaryOutputArchive<BinaryOutputArchive<TStream, TSize, TOffset, EByteOrder>,
+                                                                 TStream,
+                                                                 TSize,
+                                                                 TOffset,
+                                                                 EByteOrder> {
     public:
-        using BaseArchive = ExtendableBinaryOutputArchive<BinaryOutputArchive, TStream, TSize, TOffset>;
+        using BaseArchive = ExtendableBinaryOutputArchive<BinaryOutputArchive, TStream, TSize, TOffset, EByteOrder>;
         friend Archive<BinaryOutputArchive>;
 
     public:
diff --git a/Engine/Plugins/Animation/RigLogic/Source/RigLogicLib/Public/terse/utils/ByteSwap.h b/Engine/Plugins/Animation/RigLogic/Source/RigLogicLib/Public/terse/utils/ByteSwap.h
--- a/Engine/Plugins/Animation/RigLogic/Source/RigLogicLib/Public/terse/utils/ByteSwap.h
+++ b/Engine/Plugins/Animation/RigLogic/Source/RigLogicLib/Public/terse/utils/ByteSwap.h
@@ -50,56 +50,117 @@ struct block128 {
 
 }  // namespace impl
 
+enum class Endianness {
+    Little,
+    Big,
+    Network = Big
+};
+
+template<Endianness EByteOrder>
+struct ByteSwapper;
+
+template<>
+struct ByteSwapper<Endianness::Little> {
+
+    template<typename T>
+    static void swap(T& value) {
+        #ifdef TARGET_LITTLE_ENDIAN
+            static_cast<void>(value);
+        #else
+            using UIntType = typename traits::uint_of_size<sizeof(T)>::type;
+            static_assert(sizeof(T) == sizeof(UIntType), "No matching unsigned integral type found for the given type.");
+            // Using memcpy is the only well-defined way of reconstructing arbitrary types from raw bytes.
+            // The seemingly unnecessary copies and memcpy calls are all optimized away,
+            // compiler knows what's up.
+            UIntType swapped;
+            std::memcpy(&swapped, &value, sizeof(T));
+            swapped = bswap(swapped);
+            std::memcpy(&value, &swapped, sizeof(T));
+        #endif  // TARGET_LITTLE_ENDIAN
+    }
+
+    template<typename T>
+    static void swap(T* values) {
+        #ifdef TARGET_LITTLE_ENDIAN
+            static_cast<void>(values);
+        #else
+            using UIntType = typename traits::uint_of_size<sizeof(T)>::type;
+            static_assert(sizeof(T) == sizeof(UIntType), "No matching unsigned integral type found for the given type.");
+            // Using memcpy is the only well-defined way of reconstructing arbitrary types from raw bytes.
+            // The seemingly unnecessary copies and memcpy calls are all optimized away,
+            // compiler knows what's up.
+            alignas(impl::block128::alignment()) UIntType swapped[16ul / sizeof(T)];
+            std::memcpy(static_cast<UIntType*>(swapped), values, 16ul);
+            bswap(static_cast<UIntType*>(swapped));
+            std::memcpy(values, static_cast<UIntType*>(swapped), 16ul);
+        #endif  // TARGET_LITTLE_ENDIAN
+    }
+
+};
+
+template<>
+struct ByteSwapper<Endianness::Big> {
+
+    template<typename T>
+    static void swap(T& value) {
+        #ifdef TARGET_LITTLE_ENDIAN
+            using UIntType = typename traits::uint_of_size<sizeof(T)>::type;
+            static_assert(sizeof(T) == sizeof(UIntType), "No matching unsigned integral type found for the given type.");
+            // Using memcpy is the only well-defined way of reconstructing arbitrary types from raw bytes.
+            // The seemingly unnecessary copies and memcpy calls are all optimized away,
+            // compiler knows what's up.
+            UIntType swapped;
+            std::memcpy(&swapped, &value, sizeof(T));
+            swapped = bswap(swapped);
+            std::memcpy(&value, &swapped, sizeof(T));
+        #else
+            static_cast<void>(value);
+        #endif  // TARGET_LITTLE_ENDIAN
+    }
+
+    template<typename T>
+    static void swap(T* values) {
+        #ifdef TARGET_LITTLE_ENDIAN
+            using UIntType = typename traits::uint_of_size<sizeof(T)>::type;
+            static_assert(sizeof(T) == sizeof(UIntType), "No matching unsigned integral type found for the given type.");
+            // Using memcpy is the only well-defined way of reconstructing arbitrary types from raw bytes.
+            // The seemingly unnecessary copies and memcpy calls are all optimized away,
+            // compiler knows what's up.
+            alignas(impl::block128::alignment()) UIntType swapped[16ul / sizeof(T)];
+            std::memcpy(static_cast<UIntType*>(swapped), values, 16ul);
+            bswap(static_cast<UIntType*>(swapped));
+            std::memcpy(values, static_cast<UIntType*>(swapped), 16ul);
+        #else
+            static_cast<void>(values);
+        #endif  // TARGET_LITTLE_ENDIAN
+    }
+
+};
+
+template<Endianness EByteOrder>
+using SwapTo = ByteSwapper<EByteOrder>;
+
+template<Endianness EByteOrder>
+using SwapFrom = ByteSwapper<EByteOrder>;
+
 template<typename T>
 inline void networkToHost(T& value) {
-    using UIntType = typename traits::uint_of_size<sizeof(T)>::type;
-    static_assert(sizeof(T) == sizeof(UIntType), "No matching unsigned integral type found for the given type.");
-    // Using memcpy is the only well-defined way of reconstructing arbitrary types from raw bytes.
-    // The seemingly unnecessary copies and memcpy calls are all optimized away,
-    // compiler knows what's up.
-    UIntType hostOrder;
-    std::memcpy(&hostOrder, &value, sizeof(T));
-    hostOrder = ntoh(hostOrder);
-    std::memcpy(&value, &hostOrder, sizeof(T));
+    SwapFrom<Endianness::Network>::swap(value);
 }
 
 template<typename T>
 inline void networkToHost128(T* values) {
-    using UIntType = typename traits::uint_of_size<sizeof(T)>::type;
-    static_assert(sizeof(T) == sizeof(UIntType), "No matching unsigned integral type found for the given type.");
-    // Using memcpy is the only well-defined way of reconstructing arbitrary types from raw bytes.
-    // The seemingly unnecessary copies and memcpy calls are all optimized away,
-    // compiler knows what's up.
-    alignas(impl::block128::alignment()) UIntType hostOrder[16ul / sizeof(T)];
-    std::memcpy(static_cast<UIntType*>(hostOrder), values, 16ul);
-    ntoh(static_cast<UIntType*>(hostOrder));
-    std::memcpy(values, static_cast<UIntType*>(hostOrder), 16ul);
+    SwapFrom<Endianness::Network>::swap(values);
 }
 
 template<typename T>
 inline void hostToNetwork(T& value) {
-    using UIntType = typename traits::uint_of_size<sizeof(T)>::type;
-    static_assert(sizeof(T) == sizeof(UIntType), "No matching unsigned integral type found for the given type.");
-    // Using memcpy is the only well-defined way of reconstructing arbitrary types from raw bytes.
-    // The seemingly unnecessary copies and memcpy calls are all optimized away,
-    // compiler knows what's up.
-    UIntType networkOrder;
-    std::memcpy(&networkOrder, &value, sizeof(T));
-    networkOrder = hton(networkOrder);
-    std::memcpy(&value, &networkOrder, sizeof(T));
+    SwapTo<Endianness::Network>::swap(value);
 }
 
 template<typename T>
 inline void hostToNetwork128(T* values) {
-    using UIntType = typename traits::uint_of_size<sizeof(T)>::type;
-    static_assert(sizeof(T) == sizeof(UIntType), "No matching unsigned integral type found for the given type.");
-    // Using memcpy is the only well-defined way of reconstructing arbitrary types from raw bytes.
-    // The seemingly unnecessary copies and memcpy calls are all optimized away,
-    // compiler knows what's up.
-    alignas(impl::block128::alignment()) UIntType networkOrder[16ul / sizeof(T)];
-    std::memcpy(static_cast<UIntType*>(networkOrder), values, 16ul);
-    hton(static_cast<UIntType*>(networkOrder));
-    std::memcpy(values, static_cast<UIntType*>(networkOrder), 16ul);
+    SwapTo<Endianness::Network>::swap(values);
 }
 
 }  // namespace terse
diff --git a/Engine/Plugins/Animation/RigLogic/Source/RigLogicLib/Public/terse/utils/Endianness.h b/Engine/Plugins/Animation/RigLogic/Source/RigLogicLib/Public/terse/utils/Endianness.h
--- a/Engine/Plugins/Animation/RigLogic/Source/RigLogicLib/Public/terse/utils/Endianness.h
+++ b/Engine/Plugins/Animation/RigLogic/Source/RigLogicLib/Public/terse/utils/Endianness.h
@@ -9,7 +9,8 @@
 #endif
 
 #if !defined(NO_ENDIAN_H)
-    #if defined(USE_ENDIAN_H) || defined(USE_MACHINE_ENDIAN_H) || defined(USE_SYS_ENDIAN_H)
+    #if defined(USE_ENDIAN_H) || defined(USE_MACHINE_ENDIAN_H) || defined(USE_SYS_ENDIAN_H) || defined(USE_SYS_ISA_DEFS_H) || \
+    defined(USE_SYS_PARAM_H)
         #define OVERRIDDEN_ENDIAN_H
     #endif
 
@@ -20,6 +21,10 @@
             #define USE_MACHINE_ENDIAN_H
         #elif defined(__FreeBSD__) || defined(__NetBSD__) || defined(__OpenBSD__)
             #define USE_SYS_ENDIAN_H
+        #elif defined(__sun)
+            #define USE_SYS_ISA_DEFS_H
+        #elif defined(__MINGW32__) || defined(__MINGW64__) || !(defined(_WIN64) || defined(_WIN32))
+            #define USE_SYS_PARAM_H
         #endif
     #endif
 
@@ -29,9 +34,9 @@
         #include <machine/endian.h>
     #elif defined(USE_SYS_ENDIAN_H)
         #include <sys/endian.h>
-    #elif defined(__sun)
+    #elif defined(USE_SYS_ISA_DEFS_H)
         #include <sys/isa_defs.h>
-    #elif defined(__MINGW32__) || defined(__MINGW64__) || !(defined(_WIN64) || defined(_WIN32))
+    #elif defined(USE_SYS_PARAM_H)
         #include <sys/param.h>
     #endif
 #endif
@@ -227,6 +232,11 @@
     #error "Platform not supported, no byte swap functions defined."
 #endif
 
+#if defined(__clang__) || defined(__GNUC__)
+    #pragma GCC diagnostic push
+    #pragma GCC diagnostic ignored "-Wunused-value"
+#endif
+
 // Process single values
 
 inline std::uint8_t ntoh(std::uint8_t x) {
@@ -264,14 +274,7 @@ inline std::uint64_t hton(std::uint64_t x) {
 // Process multiple blocks simultaneously
 
 inline void ntoh(std::uint8_t* x) {
-    #if defined(__clang__) || defined(__GNUC__)
-        #pragma GCC diagnostic push
-        #pragma GCC diagnostic ignored "-Wunused-value"
-    #endif
     ntoh8x16(x);
-    #if defined(__clang__) || defined(__GNUC__)
-        #pragma GCC diagnostic pop
-    #endif
 }
 
 inline void ntoh(std::uint16_t* x) {
@@ -287,14 +290,7 @@ inline void ntoh(std::uint64_t* x) {
 }
 
 inline void hton(std::uint8_t* x) {
-    #if defined(__clang__) || defined(__GNUC__)
-        #pragma GCC diagnostic push
-        #pragma GCC diagnostic ignored "-Wunused-value"
-    #endif
     hton8x16(x);
-    #if defined(__clang__) || defined(__GNUC__)
-        #pragma GCC diagnostic pop
-    #endif
 }
 
 inline void hton(std::uint16_t* x) {
@@ -308,3 +304,43 @@ inline void hton(std::uint32_t* x) {
 inline void hton(std::uint64_t* x) {
     hton64x2(x);
 }
+
+// Byte swap overloads
+
+inline std::uint8_t bswap(std::uint8_t x) {
+    // No operation
+    return x;
+}
+
+inline std::uint16_t bswap(std::uint16_t x) {
+    return bswap16(x);
+}
+
+inline std::uint32_t bswap(std::uint32_t x) {
+    return bswap32(x);
+}
+
+inline std::uint64_t bswap(std::uint64_t x) {
+    return bswap64(x);
+}
+
+inline void bswap(std::uint8_t* x) {
+    // No operation
+    static_cast<void>(x);
+}
+
+inline void bswap(std::uint16_t* x) {
+    bswap16x8(x);
+}
+
+inline void bswap(std::uint32_t* x) {
+    bswap32x4(x);
+}
+
+inline void bswap(std::uint64_t* x) {
+    bswap64x2(x);
+}
+
+#if defined(__clang__) || defined(__GNUC__)
+    #pragma GCC diagnostic pop
+#endif
diff --git a/Engine/Plugins/Animation/RigLogic/Source/RigLogicLib/Public/terse/version/Version.h b/Engine/Plugins/Animation/RigLogic/Source/RigLogicLib/Public/terse/version/Version.h
--- a/Engine/Plugins/Animation/RigLogic/Source/RigLogicLib/Public/terse/version/Version.h
+++ b/Engine/Plugins/Animation/RigLogic/Source/RigLogicLib/Public/terse/version/Version.h
@@ -2,7 +2,7 @@
 
 #pragma once
 
-#define TERSE_MAJOR_VERSION 1
-#define TERSE_MINOR_VERSION 4
-#define TERSE_PATCH_VERSION 3
-#define TERSE_VERSION_STRING "1.4.3"
+#define TERSE_MAJOR_VERSION 3
+#define TERSE_MINOR_VERSION 1
+#define TERSE_PATCH_VERSION 2
+#define TERSE_VERSION_STRING "3.1.2"
diff --git a/Engine/Plugins/Animation/RigLogic/Source/RigLogicLib/Public/trio/version/Version.h b/Engine/Plugins/Animation/RigLogic/Source/RigLogicLib/Public/trio/version/Version.h
--- a/Engine/Plugins/Animation/RigLogic/Source/RigLogicLib/Public/trio/version/Version.h
+++ b/Engine/Plugins/Animation/RigLogic/Source/RigLogicLib/Public/trio/version/Version.h
@@ -4,5 +4,5 @@
 
 #define TRIO_MAJOR_VERSION 4
 #define TRIO_MINOR_VERSION 0
-#define TRIO_PATCH_VERSION 0
-#define TRIO_VERSION_STRING "4.0.0"
+#define TRIO_PATCH_VERSION 3
+#define TRIO_VERSION_STRING "4.0.3"
diff --git a/Engine/Plugins/Animation/RigLogic/Source/RigLogicLibTest/Private/dnatests/Defs.h b/Engine/Plugins/Animation/RigLogic/Source/RigLogicLibTest/Private/dnatests/Defs.h
--- a/Engine/Plugins/Animation/RigLogic/Source/RigLogicLibTest/Private/dnatests/Defs.h
+++ b/Engine/Plugins/Animation/RigLogic/Source/RigLogicLibTest/Private/dnatests/Defs.h
@@ -4,7 +4,7 @@
 
 #ifdef _MSC_VER
     #pragma warning(push)
-    #pragma warning(disable : 4121 4365 4987 4668 4996 6330)
+    #pragma warning(disable : 4061 4121 4365 4987 4668 4996 6330)
 #endif
 #ifdef __clang__
     #pragma clang diagnostic push
diff --git a/Engine/Plugins/Animation/RigLogic/Source/RigLogicLibTest/Private/dnatests/Fixtures.cpp b/Engine/Plugins/Animation/RigLogic/Source/RigLogicLibTest/Private/dnatests/Fixtures.cpp
--- a/Engine/Plugins/Animation/RigLogic/Source/RigLogicLibTest/Private/dnatests/Fixtures.cpp
+++ b/Engine/Plugins/Animation/RigLogic/Source/RigLogicLibTest/Private/dnatests/Fixtures.cpp
@@ -754,22 +754,22 @@ const unsigned char geometry[] = {
     0x00, 0x08,  // Maximum influence per vertex
     0x00, 0x00, 0x00, 0x03,  // Skin weights structure count: 3 (for each vertex)
     0x00, 0x00, 0x00, 0x03,  // Weights length: 3 (for each influencing joint)
-    0x40, 0xe0, 0x00, 0x00,  // 7.0f
-    0x41, 0x00, 0x00, 0x00,  // 8.0f
-    0x41, 0x10, 0x00, 0x00,  // 9.0f
+    0x3f, 0x33, 0x33, 0x33,  // 0.7f
+    0x3d, 0xcc, 0xcc, 0xcd,  // 0.1f
+    0x3e, 0x4c, 0xcc, 0xcd,  // 0.2f
     0x00, 0x00, 0x00, 0x03,  // Influencing joint count: 3 (for each weight)
     0x00, 0x00,  // Joint: 0
     0x00, 0x01,  // Joint: 1
     0x00, 0x02,  // Joint: 2
     0x00, 0x00, 0x00, 0x02,  // Weights length: 2 (for each influencing joint)
-    0x40, 0xe0, 0x00, 0x00,  // 7.0f
-    0x41, 0x00, 0x00, 0x00,  // 8.0f
+    0x3f, 0x00, 0x00, 0x00,  // 0.5f
+    0x3f, 0x00, 0x00, 0x00,  // 0.5f
     0x00, 0x00, 0x00, 0x02,  // Influencing joint count: 2 (for each weight)
     0x00, 0x03,  // Joint: 3
     0x00, 0x04,  // Joint: 4
     0x00, 0x00, 0x00, 0x02,  // Weights length: 2 (for each influencing joint)
-    0x41, 0x00, 0x00, 0x00,  // 8.0f
-    0x41, 0x10, 0x00, 0x00,  // 9.0f
+    0x3e, 0xcc, 0xcc, 0xcd,  // 0.4f
+    0x3f, 0x19, 0x99, 0x9a,  // 0.6f
     0x00, 0x00, 0x00, 0x02,  // Influencing joint count: 2 (for each weight)
     0x00, 0x05,  // Joint: 5
     0x00, 0x06,  // Joint: 6
@@ -845,22 +845,22 @@ const unsigned char geometry[] = {
     0x00, 0x08,  // Maximum influence per vertex
     0x00, 0x00, 0x00, 0x03,  // Skin weights structure count: 3 (for each vertex)
     0x00, 0x00, 0x00, 0x03,  // Weights length: 3 (for each influencing joint)
-    0x40, 0x80, 0x00, 0x00,  // 4.0f
-    0x40, 0xa0, 0x00, 0x00,  // 5.0f
-    0x40, 0xc0, 0x00, 0x00,  // 6.0f
+    0x3e, 0xcc, 0xcc, 0xcd,  // 0.4f
+    0x3e, 0x99, 0x99, 0x9a,  // 0.3f
+    0x3e, 0x99, 0x99, 0x9a,  // 0.3f
     0x00, 0x00, 0x00, 0x03,  // Influencing joint count: 3 (for each weight)
     0x00, 0x00,  // Joint: 0
     0x00, 0x01,  // Joint: 1
     0x00, 0x02,  // Joint: 2
     0x00, 0x00, 0x00, 0x02,  // Weights length: 2 (for each influencing joint)
-    0x40, 0x80, 0x00, 0x00,  // 4.0f
-    0x40, 0xa0, 0x00, 0x00,  // 5.0f
+    0x3f, 0x4c, 0xcc, 0xcd,  // 0.8f
+    0x3e, 0x4c, 0xcc, 0xcd,  // 0.2f
     0x00, 0x00, 0x00, 0x02,  // Influencing joint count: 2 (for each weight)
     0x00, 0x03,  // Joint: 3
     0x00, 0x04,  // Joint: 4
     0x00, 0x00, 0x00, 0x02,  // Weights length: 2 (for each influencing joint)
-    0x40, 0xa0, 0x00, 0x00,  // 5.0f
-    0x40, 0xc0, 0x00, 0x00,  // 6.0f
+    0x3d, 0xcc, 0xcc, 0xcd,  // 0.1f
+    0x3f, 0x66, 0x66, 0x66,  // 0.9f
     0x00, 0x00, 0x00, 0x02,  // Influencing joint count: 2 (for each weight)
     0x00, 0x05,  // Joint: 5
     0x00, 0x06,  // Joint: 6
@@ -936,22 +936,22 @@ const unsigned char geometry[] = {
     0x00, 0x08,  // Maximum influence per vertex
     0x00, 0x00, 0x00, 0x03,  // Skin weights structure count: 3 (for each vertex)
     0x00, 0x00, 0x00, 0x03,  // Weights length: 3 (for each influencing joint)
-    0x3f, 0x80, 0x00, 0x00,  // 1.0f
-    0x40, 0x00, 0x00, 0x00,  // 2.0f
-    0x40, 0x40, 0x00, 0x00,  // 3.0f
+    0x3d, 0xcc, 0xcc, 0xcd,  // 0.1f
+    0x3e, 0x99, 0x99, 0x9a,  // 0.3f
+    0x3f, 0x19, 0x99, 0x9a,  // 0.6f
     0x00, 0x00, 0x00, 0x03,  // Influencing joint count: 3 (for each weight)
     0x00, 0x00,  // Joint: 0
     0x00, 0x01,  // Joint: 1
     0x00, 0x02,  // Joint: 2
     0x00, 0x00, 0x00, 0x02,  // Weights length: 2 (for each influencing joint)
-    0x3f, 0x80, 0x00, 0x00,  // 1.0f
-    0x40, 0x00, 0x00, 0x00,  // 2.0f
+    0x3e, 0x99, 0x99, 0x9a,  // 0.3f
+    0x3f, 0x33, 0x33, 0x33,  // 0.7f
     0x00, 0x00, 0x00, 0x02,  // Influencing joint count: 2 (for each weight)
     0x00, 0x03,  // Joint: 3
     0x00, 0x04,  // Joint: 4
     0x00, 0x00, 0x00, 0x02,  // Weights length: 2 (for each influencing joint)
-    0x40, 0x00, 0x00, 0x00,  // 2.0f
-    0x40, 0x40, 0x00, 0x00,  // 3.0f
+    0x3e, 0x4c, 0xcc, 0xcd,  // 0.2f
+    0x3f, 0x4c, 0xcc, 0xcd,  // 0.8f
     0x00, 0x00, 0x00, 0x02,  // Influencing joint count: 2 (for each weight)
     0x00, 0x05,  // Joint: 5
     0x00, 0x06,  // Joint: 6
@@ -1854,40 +1854,40 @@ const pma::Matrix<std::uint16_t> maxInfluencePerVertex = {
     }
 };
 const pma::Matrix<pma::Matrix<float> > skinWeightsValues = {
-    {  // MaxLOD-0 - MaxLOD-1
+    {  // MaxLOD-0 - MinLOD-1
         {  // Mesh-0
-            {7.0f, 8.0f, 9.0f},
-            {7.0f, 8.0f},
-            {8.0f, 9.0f}
+            {0.7f, 0.1f, 0.2f},
+            {0.5f, 0.5f},
+            {0.4f, 0.6f}
         },
         {  // Mesh-1
-            {4.0f, 5.0f, 6.0f},
-            {4.0f, 5.0f},
-            {5.0f, 6.0f}
+            {0.4f, 0.3f, 0.3f},
+            {0.8f, 0.2f},
+            {0.1f, 0.9f}
         },
         {  // Mesh-2
-            {1.0f, 2.0f, 3.0f},
-            {1.0f, 2.0f},
-            {2.0f, 3.0f}
+            {0.1f, 0.3f, 0.6f},
+            {0.3f, 0.7f},
+            {0.2f, 0.8f}
         }
     },
     {  // MaxLOD-1 - MinLOD-1
         {  // Mesh-0 (Mesh-2 under MaxLOD-0)
-            {1.0f, 2.0f, 3.0f},
-            {1.0f},
-            {3.0f}
+            {0.1f, 0.3f, 0.6f},
+            {1.0f},  // 0.3f normalized to 1.0f
+            {1.0f}  // 0.8f normalized to 1.0f
         }
     },
     {  // MaxLOD-0 - MinLOD-0
         {  // Mesh-0
-            {7.0f, 8.0f, 9.0f},
-            {7.0f, 8.0f},
-            {8.0f, 9.0f}
+            {0.7f, 0.1f, 0.2f},
+            {0.5f, 0.5f},
+            {0.4f, 0.6f}
         },
         {  // Mesh-1
-            {4.0f, 5.0f, 6.0f},
-            {4.0f, 5.0f},
-            {5.0f, 6.0f}
+            {0.4f, 0.3f, 0.3f},
+            {0.8f, 0.2f},
+            {0.1f, 0.9f}
         }
     }
 };
diff --git a/Engine/Plugins/Animation/RigLogic/Source/RigLogicLibTest/Private/dnatests/TestFilteredInputArchive.cpp b/Engine/Plugins/Animation/RigLogic/Source/RigLogicLibTest/Private/dnatests/TestFilteredInputArchive.cpp
--- a/Engine/Plugins/Animation/RigLogic/Source/RigLogicLibTest/Private/dnatests/TestFilteredInputArchive.cpp
+++ b/Engine/Plugins/Animation/RigLogic/Source/RigLogicLibTest/Private/dnatests/TestFilteredInputArchive.cpp
@@ -1,10 +1,13 @@
 // Copyright Epic Games, Inc. All Rights Reserved.
 
 #include "dnatests/Defs.h"
+#include "dnatests/FakeDNAReader.h"
 #include "dnatests/Fixtures.h"
 
 #include "dna/DataLayer.h"
 #include "dna/DNA.h"
+#include "dna/BinaryStreamReader.h"
+#include "dna/BinaryStreamWriter.h"
 #include "dna/stream/FilteredInputArchive.h"
 
 #include <pma/resources/AlignedMemoryResource.h>
@@ -182,3 +185,274 @@ TEST_F(GeometryFilteringTest, IgnoreBlendShapeTargets) {
         ASSERT_TRUE(mesh.blendShapeTargets.size() == 0ul);
     }
 }
+
+namespace {
+
+class FilterLODsDNAReader : public dna::FakeDNAReader {
+    public:
+        FilterLODsDNAReader(dna::MemoryResource* memRes = nullptr) : jointNames{memRes}, jointIndicesPerLOD{memRes} {
+            lodCount = 6u;
+            jointNames = {{"body_joint1", memRes}, {"body_joint2", memRes}, {"body_joint3", memRes}, {"body_joint4", memRes},
+                {"body_joint5", memRes},
+                {"head_joint1", memRes}, {"head_joint2", memRes}, {"head_joint3", memRes}, {"head_joint4",
+                                                                                            memRes},
+                {"head_joint5", memRes},
+                {"head_joint6", memRes}, {"head_joint7", memRes}, {"head_joint8", memRes}, {"head_joint9",
+                                                                                            memRes},
+                {"head_joint10", memRes}};
+
+            jointIndicesPerLOD.resize(lodCount);
+            jointIndicesPerLOD[0] = {5u, 6u, 7u, 8u, 9u, 10u, 11u, 12u, 13u, 14u};
+            jointIndicesPerLOD[1] = {5u, 6u, 7u, 8u, 9u, 11u, 12u, 13u};
+            jointIndicesPerLOD[2] = {5u, 6u, 7u, 8u, 9u, 12u, 13u};
+            jointIndicesPerLOD[3] = {5u, 6u, 7u, 8u, 9u};
+            jointIndicesPerLOD[4] = {5u, 6u, 9u};
+            jointIndicesPerLOD[5] = {5u, 6u};
+        }
+
+        std::uint16_t getLODCount() const override {
+            return lodCount;
+        }
+
+        dna::ConstArrayView<std::uint16_t> getJointIndicesForLOD(std::uint16_t lod) const override {
+            return dna::ConstArrayView<std::uint16_t>(jointIndicesPerLOD[lod]);
+        }
+
+        std::uint16_t getJointCount() const override {
+            return static_cast<std::uint16_t>(jointNames.size());
+        }
+
+        dna::StringView getJointName(std::uint16_t i) const override {
+            return dna::StringView{jointNames[i]};
+        }
+
+    private:
+        std::uint16_t lodCount;
+        dna::Vector<dna::String<char> > jointNames;
+        dna::Matrix<std::uint16_t> jointIndicesPerLOD;
+};
+
+class FilterLODsTest : public ::testing::TestWithParam<std::vector<std::uint16_t> > {
+    protected:
+        void SetUp() override {
+            FilterLODsDNAReader dnaReader;
+            stream = dna::makeScoped<dna::MemoryStream>();
+            writer = dna::makeScoped<dna::BinaryStreamWriter>(stream.get());
+            writer->setFrom(&dnaReader);
+            writer->write();
+            lods = GetParam();
+            reader = dna::makeScoped<dna::BinaryStreamReader>(stream.get());
+            reader->read();
+            readerSpecificLODs = dna::makeScoped<dna::BinaryStreamReader>(stream.get(),
+                                                                          dna::DataLayer::All,
+                                                                          lods.data(),
+                                                                          static_cast<std::uint16_t>(lods.size()));
+            readerSpecificLODs->read();
+        }
+
+    protected:
+        dna::ScopedPtr<dna::BinaryStreamReader, dna::FactoryDestroy<dna::BinaryStreamReader> > reader;
+        dna::ScopedPtr<dna::BinaryStreamReader, dna::FactoryDestroy<dna::BinaryStreamReader> > readerSpecificLODs;
+        dna::ScopedPtr<dna::BinaryStreamWriter, dna::FactoryDestroy<dna::BinaryStreamWriter> > writer;
+        dna::ScopedPtr<dna::MemoryStream> stream;
+        std::vector<std::uint16_t> lods;
+};
+}  // namespace
+
+TEST_P(FilterLODsTest, KeepJointsNotInLODs) {
+    std::uint16_t allJoints = reader->getJointCount();
+    std::uint16_t jointsInLOD0 = static_cast<std::uint16_t>(reader->getJointIndicesForLOD(0).size());
+    std::uint16_t jointsNotInLODs = static_cast<std::uint16_t>(allJoints - jointsInLOD0);
+    ASSERT_EQ(allJoints, 15);
+    ASSERT_EQ(jointsInLOD0, 10);
+    ASSERT_EQ(jointsNotInLODs, 5);
+
+    ASSERT_TRUE(lods.size() > 0);
+    std::sort(lods.begin(), lods.end());
+    std::uint16_t maxLOD = lods[0];
+    std::uint16_t jointsInMaxLOD = static_cast<std::uint16_t>(reader->getJointIndicesForLOD(maxLOD).size());
+
+    std::uint16_t expectedJoints = static_cast<std::uint16_t>(jointsInMaxLOD + jointsNotInLODs);
+    std::uint16_t actualJoints = readerSpecificLODs->getJointCount();
+    ASSERT_EQ(actualJoints, expectedJoints);
+}
+
+INSTANTIATE_TEST_SUITE_P(FilteredDNAInputArchiveTest, FilterLODsTest,
+                         ::testing::Values(std::vector<std::uint16_t>{0}, std::vector<std::uint16_t>{2},
+                                           std::vector<std::uint16_t>{5}, std::vector<std::uint16_t>{3, 1},
+                                           std::vector<std::uint16_t>{4, 3, 5}, std::vector<std::uint16_t>{0, 1, 2, 3, 4}));
+
+namespace {
+
+class FilterSkinWeightsDNAReader : public dna::FakeDNAReader {
+    public:
+        FilterSkinWeightsDNAReader(dna::MemoryResource* memRes = nullptr) : jointNames{memRes}, meshNames{memRes},
+            jointIndicesPerLOD{memRes}, meshesPerLOD{memRes}, skinWeightsValues{memRes}, skinWeightsJointIndices{memRes} {
+            lodCount = 3u;
+            jointNames = {{"head_joint1", memRes}, {"head_joint2", memRes}, {"head_joint3", memRes}, {"head_joint4",
+                                                                                                      memRes},
+                {"head_joint5", memRes},
+                {"head_joint6", memRes}, {"head_joint7", memRes}, {"head_joint8", memRes}, {"head_joint9",
+                                                                                            memRes},
+                {"head_joint10", memRes}};
+            meshNames = {"mesh1", "mesh2"};
+            jointIndicesPerLOD.resize(lodCount);
+            jointIndicesPerLOD[0] = {0u, 1u, 2u, 3u, 4u, 5u, 6u, 7u, 8u, 9u};
+            jointIndicesPerLOD[1] = {4u, 5u, 7u};
+            jointIndicesPerLOD[2] = {};
+            meshesPerLOD.resize(lodCount);
+            meshesPerLOD[0] = {0u, 1u};
+            meshesPerLOD[1] = {0u, 1u};
+            meshesPerLOD[2] = {0u, 1u};
+            skinWeightsValues = {
+                {  // Mesh-0
+                    {0.7f, 0.1f, 0.2f},
+                    {0.5f, 0.5f},
+                    {0.4f, 0.6f}
+                },
+                {  // Mesh-1
+                    {0.4f, 0.3f, 0.3f},
+                    {0.8f, 0.2f},
+                    {0.1f, 0.9f}
+                }
+            };
+            skinWeightsJointIndices = {
+                {  // Mesh-0
+                    {0, 1, 2},
+                    {3, 4},
+                    {6, 9}
+                },
+                {  // Mesh-1
+                    {0, 1, 2},
+                    {7, 8},
+                    {5, 6}
+                }
+            };
+        }
+
+        std::uint16_t getLODCount() const override {
+            return lodCount;
+        }
+
+        dna::ConstArrayView<std::uint16_t> getJointIndicesForLOD(std::uint16_t lod) const override {
+            return dna::ConstArrayView<std::uint16_t>(jointIndicesPerLOD[lod]);
+        }
+
+        std::uint16_t getJointCount() const override {
+            return static_cast<std::uint16_t>(jointNames.size());
+        }
+
+        dna::StringView getJointName(std::uint16_t i) const override {
+            return dna::StringView{jointNames[i]};
+        }
+
+        std::uint16_t getMeshCount() const override {
+            return static_cast<std::uint16_t>(meshNames.size());
+        }
+
+        dna::StringView getMeshName(std::uint16_t i) const override {
+            return dna::StringView{meshNames[i]};
+        }
+
+        dna::ConstArrayView<std::uint16_t> getMeshIndicesForLOD(std::uint16_t lod) const override {
+            return dna::ConstArrayView<std::uint16_t>(meshesPerLOD[lod]);
+        }
+
+        std::uint32_t getSkinWeightsCount(std::uint16_t meshIndex) const override {
+            if (meshIndex < skinWeightsJointIndices.size()) {
+                return static_cast<std::uint16_t>(skinWeightsJointIndices[meshIndex].size());
+            }
+            return {};
+        }
+
+        dna::ConstArrayView<float> getSkinWeightsValues(std::uint16_t meshIndex, std::uint32_t vertexIndex) const override {
+            if (vertexIndex < getSkinWeightsCount(meshIndex)) {
+                return dna::ConstArrayView<float>(skinWeightsValues[meshIndex][vertexIndex]);
+            }
+            return {};
+        }
+
+        dna::ConstArrayView<std::uint16_t> getSkinWeightsJointIndices(std::uint16_t meshIndex,
+                                                                      std::uint32_t vertexIndex) const override {
+            if (vertexIndex < getSkinWeightsCount(meshIndex)) {
+                return dna::ConstArrayView<std::uint16_t>(skinWeightsJointIndices[meshIndex][vertexIndex]);
+            }
+            return {};
+        }
+
+    private:
+        std::uint16_t lodCount;
+        dna::Vector<dna::String<char> > jointNames;
+        dna::Vector<dna::String<char> > meshNames;
+        dna::Matrix<std::uint16_t> jointIndicesPerLOD;
+        dna::Matrix<std::uint16_t> meshesPerLOD;
+        dna::Vector<dna::Matrix<float> > skinWeightsValues;
+        dna::Vector<dna::Matrix<std::uint16_t> > skinWeightsJointIndices;
+};
+
+class FilterSkinWeightsTest : public ::testing::Test {
+    protected:
+        void SetUp() override {
+            FilterSkinWeightsDNAReader dnaReader;
+            stream = dna::makeScoped<dna::MemoryStream>();
+            writer = dna::makeScoped<dna::BinaryStreamWriter>(stream.get());
+            writer->setFrom(&dnaReader);
+            writer->write();
+            reader = dna::makeScoped<dna::BinaryStreamReader>(stream.get());
+            reader->read();
+
+            lods = {2u};
+            readerLOD2 = dna::makeScoped<dna::BinaryStreamReader>(stream.get(),
+                                                                  dna::DataLayer::All,
+                                                                  lods.data(),
+                                                                  static_cast<std::uint16_t>(lods.size()));
+            readerLOD2->read();
+        }
+
+    protected:
+        dna::ScopedPtr<dna::BinaryStreamReader, dna::FactoryDestroy<dna::BinaryStreamReader> > reader;
+        dna::ScopedPtr<dna::BinaryStreamReader, dna::FactoryDestroy<dna::BinaryStreamReader> > readerLOD2;
+        dna::ScopedPtr<dna::BinaryStreamWriter, dna::FactoryDestroy<dna::BinaryStreamWriter> > writer;
+        dna::ScopedPtr<dna::MemoryStream> stream;
+        std::vector<std::uint16_t> lods;
+};
+}  // namespace
+
+TEST_F(FilterSkinWeightsTest, RemoveAllJoints) {
+    std::uint16_t jointCount = reader->getJointCount();
+    ASSERT_EQ(jointCount, 10u);
+
+    ASSERT_EQ(reader->getSkinWeightsCount(0u), 3u);
+    ASSERT_EQ(reader->getSkinWeightsCount(1u), 3u);
+
+    ASSERT_EQ(reader->getSkinWeightsJointIndices(0u, 0u).size(), 3u);
+    ASSERT_EQ(reader->getSkinWeightsJointIndices(0u, 1u).size(), 2u);
+    ASSERT_EQ(reader->getSkinWeightsJointIndices(0u, 2u).size(), 2u);
+    ASSERT_EQ(reader->getSkinWeightsJointIndices(1u, 0u).size(), 3u);
+    ASSERT_EQ(reader->getSkinWeightsJointIndices(1u, 1u).size(), 2u);
+    ASSERT_EQ(reader->getSkinWeightsJointIndices(1u, 2u).size(), 2u);
+
+
+    ASSERT_EQ(reader->getSkinWeightsValues(0u, 0u).size(), 3u);
+    ASSERT_EQ(reader->getSkinWeightsValues(0u, 1u).size(), 2u);
+    ASSERT_EQ(reader->getSkinWeightsValues(0u, 2u).size(), 2u);
+    ASSERT_EQ(reader->getSkinWeightsValues(1u, 0u).size(), 3u);
+    ASSERT_EQ(reader->getSkinWeightsValues(1u, 1u).size(), 2u);
+    ASSERT_EQ(reader->getSkinWeightsValues(1u, 2u).size(), 2u);
+
+
+    ASSERT_EQ(readerLOD2->getSkinWeightsJointIndices(0u, 0u).size(), 0u);
+    ASSERT_EQ(readerLOD2->getSkinWeightsJointIndices(0u, 1u).size(), 0u);
+    ASSERT_EQ(readerLOD2->getSkinWeightsJointIndices(0u, 2u).size(), 0u);
+    ASSERT_EQ(readerLOD2->getSkinWeightsJointIndices(1u, 0u).size(), 0u);
+    ASSERT_EQ(readerLOD2->getSkinWeightsJointIndices(1u, 1u).size(), 0u);
+    ASSERT_EQ(readerLOD2->getSkinWeightsJointIndices(1u, 2u).size(), 0u);
+
+
+    ASSERT_EQ(readerLOD2->getSkinWeightsValues(0u, 0u).size(), 0u);
+    ASSERT_EQ(readerLOD2->getSkinWeightsValues(0u, 1u).size(), 0u);
+    ASSERT_EQ(readerLOD2->getSkinWeightsValues(0u, 2u).size(), 0u);
+    ASSERT_EQ(readerLOD2->getSkinWeightsValues(1u, 0u).size(), 0u);
+    ASSERT_EQ(readerLOD2->getSkinWeightsValues(1u, 1u).size(), 0u);
+    ASSERT_EQ(readerLOD2->getSkinWeightsValues(1u, 2u).size(), 0u);
+}
diff --git a/Engine/Plugins/Animation/RigLogic/Source/RigLogicLibTest/Private/dnatests/TestStreamReadWriteIntegration.cpp b/Engine/Plugins/Animation/RigLogic/Source/RigLogicLibTest/Private/dnatests/TestStreamReadWriteIntegration.cpp
--- a/Engine/Plugins/Animation/RigLogic/Source/RigLogicLibTest/Private/dnatests/TestStreamReadWriteIntegration.cpp
+++ b/Engine/Plugins/Animation/RigLogic/Source/RigLogicLibTest/Private/dnatests/TestStreamReadWriteIntegration.cpp
@@ -7,8 +7,10 @@
 
 #include "dna/DataLayer.h"
 #include "dna/DNA.h"
-#include "dna/StreamReader.h"
-#include "dna/StreamWriter.h"
+#include "dna/BinaryStreamReader.h"
+#include "dna/BinaryStreamWriter.h"
+#include "dna/JSONStreamReader.h"
+#include "dna/JSONStreamWriter.h"
 
 #include <pma/resources/AlignedMemoryResource.h>
 
@@ -338,7 +340,7 @@ static void assertReaderHasAllData(Reader* reader, const LODParameters& params)
     verifyGeometry(reader, params);
 }
 
-TEST_P(StreamReadWriteIntegrationTest, VerifyAllDNADataAfterSetFromIsUsed) {
+TEST(StreamReadWriteIntegrationTest, JSONVerifyAllDNADataAfterSetFromIsUsed) {
     pma::AlignedMemoryResource memRes;
 
     const auto bytes = raw::getBytes();
@@ -346,23 +348,42 @@ TEST_P(StreamReadWriteIntegrationTest, VerifyAllDNADataAfterSetFromIsUsed) {
     source->write(reinterpret_cast<const char*>(bytes.data()), bytes.size());
     source->seek(0);
 
-    auto sourceReader = StreamReader::create(source.get(), DataLayer::All, 0u, &memRes);
+    auto sourceReader = pma::makeScoped<BinaryStreamReader>(source.get(), DataLayer::All, static_cast<std::uint16_t>(0), &memRes);
     sourceReader->read();
 
     auto clone = pma::makeScoped<trio::MemoryStream>();
-    auto cloneWriter = StreamWriter::create(clone.get(), &memRes);
-    cloneWriter->setFrom(sourceReader);
+    auto cloneWriter = pma::makeScoped<JSONStreamWriter>(clone.get(), 4u, &memRes);
+    cloneWriter->setFrom(sourceReader.get());
     cloneWriter->write();
 
     clone->seek(0ul);
-    auto cloneReader = StreamReader::create(clone.get(), DataLayer::All, params.maxLOD, &memRes);
+    auto cloneReader = pma::makeScoped<JSONStreamReader>(clone.get(), &memRes);
     cloneReader->read();
 
-    assertReaderHasAllData(cloneReader, params);
+    assertReaderHasAllData(cloneReader.get(), {0u, 0u});
+}
+
+TEST_P(StreamReadWriteIntegrationTest, BinaryVerifyAllDNADataAfterSetFromIsUsed) {
+    pma::AlignedMemoryResource memRes;
+
+    const auto bytes = raw::getBytes();
+    auto source = pma::makeScoped<trio::MemoryStream>();
+    source->write(reinterpret_cast<const char*>(bytes.data()), bytes.size());
+    source->seek(0);
+
+    auto sourceReader = pma::makeScoped<BinaryStreamReader>(source.get(), DataLayer::All, static_cast<std::uint16_t>(0), &memRes);
+    sourceReader->read();
+
+    auto clone = pma::makeScoped<trio::MemoryStream>();
+    auto cloneWriter = pma::makeScoped<BinaryStreamWriter>(clone.get(), &memRes);
+    cloneWriter->setFrom(sourceReader.get());
+    cloneWriter->write();
+
+    clone->seek(0ul);
+    auto cloneReader = pma::makeScoped<BinaryStreamReader>(clone.get(), DataLayer::All, params.maxLOD, &memRes);
+    cloneReader->read();
 
-    StreamReader::destroy(cloneReader);
-    StreamWriter::destroy(cloneWriter);
-    StreamReader::destroy(sourceReader);
+    assertReaderHasAllData(cloneReader.get(), params);
 }
 
 INSTANTIATE_TEST_SUITE_P(StreamReadWriteIntegrationTest, StreamReadWriteIntegrationTest, ::testing::Values(
diff --git a/Engine/Plugins/Animation/RigLogic/Source/RigLogicLibTest/Private/dnatests/TestStreamReadWriteIntegration.h b/Engine/Plugins/Animation/RigLogic/Source/RigLogicLibTest/Private/dnatests/TestStreamReadWriteIntegration.h
--- a/Engine/Plugins/Animation/RigLogic/Source/RigLogicLibTest/Private/dnatests/TestStreamReadWriteIntegration.h
+++ b/Engine/Plugins/Animation/RigLogic/Source/RigLogicLibTest/Private/dnatests/TestStreamReadWriteIntegration.h
@@ -4,9 +4,6 @@
 
 #include "dnatests/Defs.h"
 
-#include "dna/StreamReader.h"
-
-
 struct LODParameters {
     std::uint16_t maxLOD;
     std::uint16_t currentLOD;
diff --git a/Engine/Plugins/Animation/RigLogic/Source/RigLogicLibTest/Private/pmatests/Defs.h b/Engine/Plugins/Animation/RigLogic/Source/RigLogicLibTest/Private/pmatests/Defs.h
--- a/Engine/Plugins/Animation/RigLogic/Source/RigLogicLibTest/Private/pmatests/Defs.h
+++ b/Engine/Plugins/Animation/RigLogic/Source/RigLogicLibTest/Private/pmatests/Defs.h
@@ -4,7 +4,7 @@
 
 #ifdef _MSC_VER
     #pragma warning(push)
-    #pragma warning(disable : 4121 4365 4987 4668 4996 6330)
+    #pragma warning(disable : 4061 4121 4365 4987 4668 4996 6330)
 #endif
 #ifdef __clang__
     #pragma clang diagnostic push
diff --git a/Engine/Plugins/Animation/RigLogic/Source/RigLogicLibTest/Private/rltests/conditionaltable/ConditionalTableFixtures.cpp b/Engine/Plugins/Animation/RigLogic/Source/RigLogicLibTest/Private/rltests/conditionaltable/ConditionalTableFixtures.cpp
--- a/Engine/Plugins/Animation/RigLogic/Source/RigLogicLibTest/Private/rltests/conditionaltable/ConditionalTableFixtures.cpp
+++ b/Engine/Plugins/Animation/RigLogic/Source/RigLogicLibTest/Private/rltests/conditionaltable/ConditionalTableFixtures.cpp
@@ -35,6 +35,7 @@ const std::uint16_t distinctOutputs = 2ul;
 
 rl4::ConditionalTable ConditionalTableFactory::withSingleIO(rl4::Vector<float>&& fromValues,
                                                             rl4::Vector<float>&& toValues,
+                                                            rl4::Vector<float>&& cutValues,
                                                             rl4::MemoryResource* memRes) {
     return rl4::ConditionalTable{
         rl4::Vector<std::uint16_t>{single::inputIndices, single::inputIndices + single::size, memRes},
@@ -42,9 +43,10 @@ rl4::ConditionalTable ConditionalTableFactory::withSingleIO(rl4::Vector<float>&&
         std::move(fromValues),
         std::move(toValues),
         rl4::Vector<float>{single::slopeValues, single::slopeValues + single::size, memRes},
-        rl4::Vector<float>{single::cutValues, single::cutValues + single::size, memRes},
+        std::move(cutValues),
         single::distinctInputs,
-        single::distinctOutputs
+        single::distinctOutputs,
+        memRes
     };
 }
 
@@ -59,7 +61,8 @@ rl4::ConditionalTable ConditionalTableFactory::withMultipleIO(rl4::Vector<float>
         rl4::Vector<float>{multiple::slopeValues, multiple::slopeValues + multiple::size, memRes},
         rl4::Vector<float>{multiple::cutValues, multiple::cutValues + multiple::size, memRes},
         multiple::distinctInputs,
-        multiple::distinctOutputs
+        multiple::distinctOutputs,
+        memRes
     };
 }
 
@@ -67,6 +70,7 @@ rl4::ConditionalTable ConditionalTableFactory::withSingleIODefaults(rl4::MemoryR
     return withSingleIO(
         rl4::Vector<float>{single::fromValues, single::fromValues + single::size, memRes},
         rl4::Vector<float>{single::toValues, single::toValues + single::size, memRes},
+        rl4::Vector<float>{single::cutValues, single::cutValues + single::size, memRes},
         memRes
         );
 }
diff --git a/Engine/Plugins/Animation/RigLogic/Source/RigLogicLibTest/Private/rltests/conditionaltable/ConditionalTableFixtures.h b/Engine/Plugins/Animation/RigLogic/Source/RigLogicLibTest/Private/rltests/conditionaltable/ConditionalTableFixtures.h
--- a/Engine/Plugins/Animation/RigLogic/Source/RigLogicLibTest/Private/rltests/conditionaltable/ConditionalTableFixtures.h
+++ b/Engine/Plugins/Animation/RigLogic/Source/RigLogicLibTest/Private/rltests/conditionaltable/ConditionalTableFixtures.h
@@ -25,6 +25,7 @@ struct ConditionalTableFactory {
 
     static rl4::ConditionalTable withSingleIO(rl4::Vector<float>&& fromValues,
                                               rl4::Vector<float>&& toValues,
+                                              rl4::Vector<float>&& cutValues,
                                               rl4::MemoryResource* memRes);
     static rl4::ConditionalTable withMultipleIO(rl4::Vector<float>&& fromValues,
                                                 rl4::Vector<float>&& toValues,
diff --git a/Engine/Plugins/Animation/RigLogic/Source/RigLogicLibTest/Private/rltests/conditionaltable/TestConditionalTable.cpp b/Engine/Plugins/Animation/RigLogic/Source/RigLogicLibTest/Private/rltests/conditionaltable/TestConditionalTable.cpp
--- a/Engine/Plugins/Animation/RigLogic/Source/RigLogicLibTest/Private/rltests/conditionaltable/TestConditionalTable.cpp
+++ b/Engine/Plugins/Animation/RigLogic/Source/RigLogicLibTest/Private/rltests/conditionaltable/TestConditionalTable.cpp
@@ -12,6 +12,8 @@ namespace {
 struct CalcTestData {
     float fromValues[2];
     float toValues[2];
+    float cutValues[2];
+    float inValues[1];
     float expected[1];
 };
 
@@ -27,11 +29,12 @@ TEST_P(ConditionalTableTest, CheckCalculationBorderCases) {
     auto conditionals = ConditionalTableFactory::withSingleIO(
         rl4::Vector<float>{testData.fromValues, testData.fromValues + conditionalsSize, &amr},
         rl4::Vector<float>{testData.toValues, testData.toValues + conditionalsSize, &amr},
+        rl4::Vector<float>{testData.cutValues, testData.cutValues + conditionalsSize, &amr},
         &amr
         );
     float outputs[] = {0.0f};
-    conditionals.calculate(conditionalTableInputs.data(), outputs);
-    ASSERT_ELEMENTS_EQ(outputs, testData.expected, 1);
+    conditionals.calculate(testData.inValues, outputs);
+    ASSERT_ELEMENTS_NEAR(outputs, testData.expected, 1, 0.00001f);
 }
 
 INSTANTIATE_TEST_SUITE_P(ConditionalTableTest,
@@ -39,22 +42,23 @@ INSTANTIATE_TEST_SUITE_P(ConditionalTableTest,
                          ::testing::Values(
                              // {{fromValues}, {toValues}, {expected}}
                              // In-value below from-value
-                             CalcTestData{{0.3f, 0.6f}, {0.6f, 1.0f}, {0.0f}},
+                             CalcTestData{{0.3f, 0.6f}, {0.6f, 1.0f}, {0.1f, 0.3f}, {0.1f}, {0.0f}},
                              // In-value equals from-value
-                             CalcTestData{{0.1f, 0.6f}, {0.6f, 1.0f}, {0.0f}},
+                             CalcTestData{{0.1f, 0.6f}, {0.6f, 1.0f}, {0.1f, 0.3f}, {0.1f}, {0.2f}},
                              // In-value equals to-value
-                             CalcTestData{{0.0f, 0.1f}, {0.6f, 1.0f}, {0.2f}},
+                             CalcTestData{{0.0f, 0.2f}, {0.1f, 1.0f}, {0.1f, 0.3f}, {0.1f}, {0.2f}},
                              // In-value equals both from-value and to-value
-                             CalcTestData{{0.0f, 0.1f}, {0.1f, 1.0f}, {0.2f}},
+                             CalcTestData{{0.0f, 0.1f}, {0.1f, 1.0f}, {0.1f, 0.3f}, {0.1f}, {0.2f}},
                              // In-value between from-value and to-value
-                             CalcTestData{{0.0f, 0.6f}, {0.6f, 1.0f}, {0.2f}},
+                             CalcTestData{{0.0f, 0.6f}, {0.6f, 1.0f}, {0.1f, 0.3f}, {0.1f}, {0.2f}},
                              // In-value above to-value
-                             CalcTestData{{0.0f, 0.04f}, {0.04f, 0.09f}, {0.0f}},
-                             // Original output not reset, but summed with new value
-                             CalcTestData{{0.0f, 0.0f}, {0.2f, 0.2f}, {0.6f}}
+                             CalcTestData{{0.0f, 0.04f}, {0.04f, 0.09f}, {0.1f, 0.3f}, {0.1f}, {0.0f}},
+                             // In-value equals lower-bound from-value
+                             CalcTestData{{-1.0f, 0.0f}, {0.0f, 1.0f}, {1.4f, 0.3f}, {-1.0f}, {0.4f}}
                              ));
 
 TEST(ConditionalTableTest, OutputClamped) {
+    pma::AlignedMemoryResource amr;
     const std::uint16_t inputCount = 1u;
     const std::uint16_t outputCount = 1u;
     rl4::Vector<std::uint16_t> inputIndices{0};
@@ -71,22 +75,48 @@ TEST(ConditionalTableTest, OutputClamped) {
                                        std::move(slopeValues),
                                        std::move(cutValues),
                                        inputCount,
-                                       outputCount};
+                                       outputCount,
+                                       &amr};
     conditionals.calculate(conditionalTableInputs.data(), outputs);
     const float expected[] = {1.0f};
     ASSERT_ELEMENTS_EQ(outputs, expected, 1ul);
 }
 
+TEST(ConditionalTableTest, OutputIsAccumulated) {
+    pma::AlignedMemoryResource amr;
+    const std::uint16_t inputCount = 2u;
+    const std::uint16_t outputCount = 1u;
+    rl4::Vector<std::uint16_t> inputIndices{0, 1};
+    rl4::Vector<std::uint16_t> outputIndices{0, 0};
+    rl4::Vector<float> slopeValues{1.0f, 1.0f};
+    rl4::Vector<float> cutValues{0.2f, 0.4f};
+    rl4::Vector<float> fromValues{0.0f, 0.0f};
+    rl4::Vector<float> toValues{0.2f, 0.2f};
+    float outputs[1] = {};
+    rl4::ConditionalTable conditionals{std::move(inputIndices),
+                                       std::move(outputIndices),
+                                       std::move(fromValues),
+                                       std::move(toValues),
+                                       std::move(slopeValues),
+                                       std::move(cutValues),
+                                       inputCount,
+                                       outputCount,
+                                       &amr};
+    conditionals.calculate(conditionalTableInputs.data(), outputs);
+    const float expected[] = {0.9f};
+    ASSERT_ELEMENTS_NEAR(outputs, expected, 1ul, 0.00001f);
+}
+
 TEST(ConditionalTableTest, OutputIsResetOnEachCalculation) {
     pma::AlignedMemoryResource amr;
     auto conditionals = ConditionalTableFactory::withSingleIODefaults(&amr);
 
-    float outputs[2ul];
-    const float expected[] = {0.6f};
+    float outputs[1ul] = {};
+    const float expected[] = {0.2f};
 
     conditionals.calculate(conditionalTableInputs.data(), outputs);
-    ASSERT_ELEMENTS_EQ(outputs, expected, 1ul);
+    ASSERT_ELEMENTS_NEAR(outputs, expected, 1ul, 0.00001f);
 
     conditionals.calculate(conditionalTableInputs.data(), outputs);
-    ASSERT_ELEMENTS_EQ(outputs, expected, 1ul);
+    ASSERT_ELEMENTS_NEAR(outputs, expected, 1ul, 0.00001f);
 }
diff --git a/Engine/Plugins/Animation/RigLogic/Source/RigLogicLibTest/Private/rltests/dna/FakeReader.h b/Engine/Plugins/Animation/RigLogic/Source/RigLogicLibTest/Private/rltests/dna/FakeReader.h
--- a/Engine/Plugins/Animation/RigLogic/Source/RigLogicLibTest/Private/rltests/dna/FakeReader.h
+++ b/Engine/Plugins/Animation/RigLogic/Source/RigLogicLibTest/Private/rltests/dna/FakeReader.h
@@ -23,9 +23,6 @@ class FakeReader : public Reader {
     public:
         ~FakeReader();
 
-        void unload(DataLayer layer) override {
-        }
-
         // DescriptorReader methods start
         StringView getName() const override {
             return {};
@@ -470,6 +467,10 @@ class FakeReader : public Reader {
             return {};
         }
 
+        void unload(DataLayer  /**/) override {
+
+        }
+
 };
 
 }  // namespace dna
diff --git a/Engine/Plugins/Animation/RigLogic/Source/RigLogicLibTest/Private/rltests/riglogic/TestRigLogic.cpp b/Engine/Plugins/Animation/RigLogic/Source/RigLogicLibTest/Private/rltests/riglogic/TestRigLogic.cpp
--- a/Engine/Plugins/Animation/RigLogic/Source/RigLogicLibTest/Private/rltests/riglogic/TestRigLogic.cpp
+++ b/Engine/Plugins/Animation/RigLogic/Source/RigLogicLibTest/Private/rltests/riglogic/TestRigLogic.cpp
@@ -17,7 +17,7 @@ class RigLogicTest : public ::testing::Test {
             stream->write(reinterpret_cast<const char*>(bytes.data()), bytes.size());
             stream->seek(0);
 
-            reader = dna::StreamReader::create(stream.get());
+            reader = dna::BinaryStreamReader::create(stream.get());
             reader->read();
 
             rigLogic = rl4::RigLogic::create(reader);
@@ -27,13 +27,13 @@ class RigLogicTest : public ::testing::Test {
         void TearDown() override {
             rl4::RigInstance::destroy(rigInstance);
             rl4::RigLogic::destroy(rigLogic);
-            dna::StreamReader::destroy(reader);
+            dna::BinaryStreamReader::destroy(reader);
         }
 
     protected:
         pma::AlignedMemoryResource memRes;
         pma::ScopedPtr<trio::MemoryStream, pma::FactoryDestroy<trio::MemoryStream> > stream;
-        dna::StreamReader* reader;
+        dna::BinaryStreamReader* reader;
         rl4::RigLogic* rigLogic;
         rl4::RigInstance* rigInstance;
 };
diff --git a/Engine/Plugins/Animation/RigLogic/Source/RigLogicLibTest/Private/tersetests/archives/binary/Common.h b/Engine/Plugins/Animation/RigLogic/Source/RigLogicLibTest/Private/tersetests/archives/binary/Common.h
--- a/Engine/Plugins/Animation/RigLogic/Source/RigLogicLibTest/Private/tersetests/archives/binary/Common.h
+++ b/Engine/Plugins/Animation/RigLogic/Source/RigLogicLibTest/Private/tersetests/archives/binary/Common.h
@@ -2,7 +2,8 @@
 
 #pragma once
 
-#include "terse/utils/ArchiveOffset.h"
+#include "terse/types/ArchiveOffset.h"
+#include "terse/types/Versioned.h"
 
 #include <pma/TypeDefs.h>
 
@@ -120,3 +121,29 @@ template<class TArchive>
 void save(TArchive& archive, const SerializableByFreeLoadSave& source) {
     archive(source.a, source.b);
 }
+
+struct UpgradedSerializableType {
+    std::uint32_t a;
+    std::uint16_t b;
+
+    template<class Archive>
+    void load(Archive& archive, terse::Version<1>) {
+        archive(a);
+    }
+
+    template<class Archive>
+    void load(Archive& archive, terse::Version<2>) {
+        archive(a, b);
+    }
+
+    template<class Archive>
+    void save(Archive& archive, terse::Version<1>) {
+        archive(a);
+    }
+
+    template<class Archive>
+    void save(Archive& archive, terse::Version<2>) {
+        archive(a, b);
+    }
+
+};
diff --git a/Engine/Plugins/Animation/RigLogic/Source/RigLogicLibTest/Private/tersetests/types/TestDynArray.cpp b/Engine/Plugins/Animation/RigLogic/Source/RigLogicLibTest/Private/tersetests/types/TestDynArray.cpp
--- a/Engine/Plugins/Animation/RigLogic/Source/RigLogicLibTest/Private/tersetests/types/TestDynArray.cpp
+++ b/Engine/Plugins/Animation/RigLogic/Source/RigLogicLibTest/Private/tersetests/types/TestDynArray.cpp
@@ -102,9 +102,14 @@ TEST(DynArrayTest, MoveConstruct) {
     terse::DynArray<int, pma::PolyAllocator<int> > arr{values, 10ul};
     terse::DynArray<int, pma::PolyAllocator<int> > arrCopy = std::move(arr);
 
-    ASSERT_TRUE(arr.empty());
-    ASSERT_EQ(arr.size(), 0ul);
-    ASSERT_EQ(arr.data(), nullptr);
+    // Hide these assertions from clang static analyzer, as move operations for DynArrays are
+    // well defined, but clang doesn't know that, so it will always warn about accessing
+    // moved-from objects.
+    #ifndef __clang_analyzer__
+        ASSERT_TRUE(arr.empty());
+        ASSERT_EQ(arr.size(), 0ul);
+        ASSERT_EQ(arr.data(), nullptr);
+    #endif  // __clang_analyzer__
 
     ASSERT_EQ(arrCopy.size(), 10ul);
     for (std::size_t i{}; i < arrCopy.size(); ++i) {
@@ -120,9 +125,14 @@ TEST(DynArrayTest, MoveAssign) {
     terse::DynArray<int, pma::PolyAllocator<int> > arrCopy;
     arrCopy = std::move(arr);
 
-    ASSERT_TRUE(arr.empty());
-    ASSERT_EQ(arr.size(), 0ul);
-    ASSERT_EQ(arr.data(), nullptr);
+    // Hide these assertions from clang static analyzer, as move operations for DynArrays are
+    // well defined, but clang doesn't know that, so it will always warn about accessing
+    // moved-from objects.
+    #ifndef __clang_analyzer__
+        ASSERT_TRUE(arr.empty());
+        ASSERT_EQ(arr.size(), 0ul);
+        ASSERT_EQ(arr.data(), nullptr);
+    #endif  // __clang_analyzer__
 
     ASSERT_EQ(arrCopy.size(), 10ul);
     for (std::size_t i{}; i < arrCopy.size(); ++i) {
diff --git a/Engine/Plugins/Animation/RigLogic/Source/RigLogicLibTest/Private/trimdtests/Defs.h b/Engine/Plugins/Animation/RigLogic/Source/RigLogicLibTest/Private/trimdtests/Defs.h
--- a/Engine/Plugins/Animation/RigLogic/Source/RigLogicLibTest/Private/trimdtests/Defs.h
+++ b/Engine/Plugins/Animation/RigLogic/Source/RigLogicLibTest/Private/trimdtests/Defs.h
@@ -4,7 +4,7 @@
 
 #ifdef _MSC_VER
     #pragma warning(push)
-    #pragma warning(disable : 4121 4365 4987 4668 4996 6330)
+    #pragma warning(disable : 4061 4121 4365 4987 4668 4996 6330)
 #endif
 #ifdef __clang__
     #pragma clang diagnostic push
@@ -37,29 +37,29 @@
 #endif
 
 #ifndef ASSERT_ELEMENTS_EQ
-#define ASSERT_ELEMENTS_EQ(result, expected, count)     \
-    for (std::size_t i = 0ul; i < count; ++i) {         \
-        ASSERT_EQ(result[i], expected[i]);              \
-    }
+    #define ASSERT_ELEMENTS_EQ(result, expected, count)     \
+        for (std::size_t i = 0ul; i < count; ++i) {         \
+            ASSERT_EQ(result[i], expected[i]);              \
+        }
 #endif
 
 #ifndef EXPECT_ELEMENTS_EQ
-#define EXPECT_ELEMENTS_EQ(result, expected, count)     \
-    for (std::size_t i = 0ul; i < count; ++i) {         \
-        EXPECT_EQ(result[i], expected[i]);              \
-    }
+    #define EXPECT_ELEMENTS_EQ(result, expected, count)     \
+        for (std::size_t i = 0ul; i < count; ++i) {         \
+            EXPECT_EQ(result[i], expected[i]);              \
+        }
 #endif
 
 #ifndef ASSERT_ELEMENTS_NEAR
-#define ASSERT_ELEMENTS_NEAR(result, expected, count, threshold)    \
-    for (std::size_t i = 0ul; i < count; ++i) {                     \
-        ASSERT_NEAR(result[i], expected[i], threshold);             \
-    }
+    #define ASSERT_ELEMENTS_NEAR(result, expected, count, threshold)    \
+        for (std::size_t i = 0ul; i < count; ++i) {                     \
+            ASSERT_NEAR(result[i], expected[i], threshold);             \
+        }
 #endif
 
 #ifndef EXPECT_ELEMENTS_NEAR
-#define EXPECT_ELEMENTS_NEAR(result, expected, count, threshold)    \
-    for (std::size_t i = 0ul; i < count; ++i) {                     \
-        EXPECT_NEAR(result[i], expected[i], threshold);             \
-    }
+    #define EXPECT_ELEMENTS_NEAR(result, expected, count, threshold)    \
+        for (std::size_t i = 0ul; i < count; ++i) {                     \
+            EXPECT_NEAR(result[i], expected[i], threshold);             \
+        }
 #endif
diff --git a/Engine/Plugins/Animation/RigLogic/Source/RigLogicModule/Private/DNAAsset.cpp b/Engine/Plugins/Animation/RigLogic/Source/RigLogicModule/Private/DNAAsset.cpp
--- a/Engine/Plugins/Animation/RigLogic/Source/RigLogicModule/Private/DNAAsset.cpp
+++ b/Engine/Plugins/Animation/RigLogic/Source/RigLogicModule/Private/DNAAsset.cpp
@@ -33,19 +33,19 @@ static constexpr uint32 AVG_GEOMETRY_SIZE = 50 * 1024 * 1024;
 
 static TSharedPtr<IDNAReader> ReadDNAFromStream(rl4::BoundedIOStream* Stream, EDNADataLayer Layer, uint16 MaxLOD)
 {
-	auto DNAStreamReader = rl4::makeScoped<dna::StreamReader>(Stream, static_cast<dna::DataLayer>(Layer), MaxLOD, FMemoryResource::Instance());
+	auto DNAStreamReader = rl4::makeScoped<dna::BinaryStreamReader>(Stream, static_cast<dna::DataLayer>(Layer), MaxLOD, FMemoryResource::Instance());
 	DNAStreamReader->read();
 	if (!rl4::Status::isOk())
 	{
 		UE_LOG(LogDNAAsset, Error, TEXT("%s"), ANSI_TO_TCHAR(rl4::Status::get().message));
 		return nullptr;
 	}
-	return MakeShared<FDNAReader<dna::StreamReader>>(DNAStreamReader.release());
+	return MakeShared<FDNAReader<dna::BinaryStreamReader>>(DNAStreamReader.release());
 }
 
 static void WriteDNAToStream(const IDNAReader* Source, EDNADataLayer Layer, rl4::BoundedIOStream* Destination)
 {
-	auto DNAWriter = rl4::makeScoped<dna::StreamWriter>(Destination, FMemoryResource::Instance());
+	auto DNAWriter = rl4::makeScoped<dna::BinaryStreamWriter>(Destination, FMemoryResource::Instance());
 	if (Source != nullptr)
 	{
 		DNAWriter->setFrom(Source->Unwrap(), static_cast<dna::DataLayer>(Layer), FMemoryResource::Instance());
diff --git a/Engine/Plugins/Animation/RigLogic/Source/RigLogicModule/Private/DNAUtils.cpp b/Engine/Plugins/Animation/RigLogic/Source/RigLogicModule/Private/DNAUtils.cpp
--- a/Engine/Plugins/Animation/RigLogic/Source/RigLogicModule/Private/DNAUtils.cpp
+++ b/Engine/Plugins/Animation/RigLogic/Source/RigLogicModule/Private/DNAUtils.cpp
@@ -10,7 +10,8 @@
 
 #include "riglogic/RigLogic.h"
 
-static TSharedPtr<IDNAReader> ReadDNAStream(rl4::ScopedPtr<dna::StreamReader> DNAStreamReader)
+
+static TSharedPtr<IDNAReader> ReadDNAStream(rl4::ScopedPtr<dna::BinaryStreamReader> DNAStreamReader)
 {
 	DNAStreamReader->read();
 	if (!rl4::Status::isOk())
@@ -18,14 +19,14 @@ static TSharedPtr<IDNAReader> ReadDNAStream(rl4::ScopedPtr<dna::StreamReader> DN
 		UE_LOG(LogDNAReader, Error, TEXT("%s"), ANSI_TO_TCHAR(rl4::Status::get().message));
 		return nullptr;
 	}
-	return MakeShared<FDNAReader<dna::StreamReader>>(DNAStreamReader.release());
+	return MakeShared<FDNAReader<dna::BinaryStreamReader>>(DNAStreamReader.release());
 }
 
 TSharedPtr<IDNAReader> ReadDNAFromFile(const FString& Path, EDNADataLayer Layer, uint16_t MaxLOD)
 {
 	LLM_SCOPE_BYNAME(TEXT("Animation/RigLogic"));
 	auto DNAFileStream = rl4::makeScoped<rl4::MemoryMappedFileStream>(TCHAR_TO_UTF8(*Path), rl4::MemoryMappedFileStream::AccessMode::Read, FMemoryResource::Instance());
-	auto DNAStreamReader = rl4::makeScoped<dna::StreamReader>(DNAFileStream.get(), static_cast<dna::DataLayer>(Layer), MaxLOD, FMemoryResource::Instance());
+	auto DNAStreamReader = rl4::makeScoped<dna::BinaryStreamReader>(DNAFileStream.get(), static_cast<dna::DataLayer>(Layer), MaxLOD, FMemoryResource::Instance());
 	return ReadDNAStream(MoveTemp(DNAStreamReader));
 }
 
@@ -33,7 +34,7 @@ TSharedPtr<IDNAReader> ReadDNAFromFile(const FString& Path, EDNADataLayer Layer,
 {
 	LLM_SCOPE_BYNAME(TEXT("Animation/RigLogic"));
 	auto DNAFileStream = rl4::makeScoped<rl4::MemoryMappedFileStream>(TCHAR_TO_UTF8(*Path), rl4::MemoryMappedFileStream::AccessMode::Read, FMemoryResource::Instance());
-	auto DNAStreamReader = rl4::makeScoped<dna::StreamReader>(DNAFileStream.get(), static_cast<dna::DataLayer>(Layer), LODs.GetData(), static_cast<uint16>(LODs.Num()), FMemoryResource::Instance());
+	auto DNAStreamReader = rl4::makeScoped<dna::BinaryStreamReader>(DNAFileStream.get(), static_cast<dna::DataLayer>(Layer), LODs.GetData(), static_cast<uint16>(LODs.Num()), FMemoryResource::Instance());
 	return ReadDNAStream(MoveTemp(DNAStreamReader));
 }
 
@@ -41,7 +42,7 @@ TSharedPtr<IDNAReader> ReadDNAFromBuffer(TArray<uint8>* DNABuffer, EDNADataLayer
 {
 	LLM_SCOPE_BYNAME(TEXT("Animation/RigLogic"));
 	FRigLogicMemoryStream DNAMemoryStream(DNABuffer);
-	auto DNAStreamReader = rl4::makeScoped<dna::StreamReader>(&DNAMemoryStream, static_cast<dna::DataLayer>(Layer), MaxLOD, FMemoryResource::Instance());
+	auto DNAStreamReader = rl4::makeScoped<dna::BinaryStreamReader>(&DNAMemoryStream, static_cast<dna::DataLayer>(Layer), MaxLOD, FMemoryResource::Instance());
 	return ReadDNAStream(MoveTemp(DNAStreamReader));
 }
 
@@ -49,7 +50,7 @@ TSharedPtr<IDNAReader> ReadDNAFromBuffer(TArray<uint8>* DNABuffer, EDNADataLayer
 {
 	LLM_SCOPE_BYNAME(TEXT("Animation/RigLogic"));
 	FRigLogicMemoryStream DNAMemoryStream(DNABuffer);
-	auto DNAStreamReader = rl4::makeScoped<dna::StreamReader>(&DNAMemoryStream, static_cast<dna::DataLayer>(Layer), LODs.GetData(), static_cast<uint16>(LODs.Num()), FMemoryResource::Instance());
+	auto DNAStreamReader = rl4::makeScoped<dna::BinaryStreamReader>(&DNAMemoryStream, static_cast<dna::DataLayer>(Layer), LODs.GetData(), static_cast<uint16>(LODs.Num()), FMemoryResource::Instance());
 	return ReadDNAStream(MoveTemp(DNAStreamReader));
 }
 
@@ -58,7 +59,7 @@ TArray<uint8> ReadStreamFromDNA(const IDNAReader* Reader, EDNADataLayer Layer)
 	LLM_SCOPE_BYNAME(TEXT("Animation/RigLogic"));
 	TArray<char> DNABuffer;
 	auto DeltaDnaStream = rl4::makeScoped<trio::MemoryStream>();
-	auto DNAStreamWriter = rl4::makeScoped<dna::StreamWriter>(DeltaDnaStream.get(), FMemoryResource::Instance());
+	auto DNAStreamWriter = rl4::makeScoped<dna::BinaryStreamWriter>(DeltaDnaStream.get(), FMemoryResource::Instance());
 	DNAStreamWriter->setFrom(Reader->Unwrap(), static_cast<dna::DataLayer>(Layer), FMemoryResource::Instance());
 	DNAStreamWriter->write();
 	DNABuffer.AddZeroed(DeltaDnaStream->size());
@@ -70,7 +71,7 @@ void WriteDNAToFile(const IDNAReader* Reader, EDNADataLayer Layer, const FString
 {
 	LLM_SCOPE_BYNAME(TEXT("Animation/RigLogic"));
 	auto DNAFileStream = rl4::makeScoped<rl4::FileStream>(TCHAR_TO_UTF8(*Path), rl4::FileStream::AccessMode::Write, rl4::FileStream::OpenMode::Binary, FMemoryResource::Instance());
-	auto DNAStreamWriter = rl4::makeScoped<dna::StreamWriter>(DNAFileStream.get(), FMemoryResource::Instance());
+	auto DNAStreamWriter = rl4::makeScoped<dna::BinaryStreamWriter>(DNAFileStream.get(), FMemoryResource::Instance());
 	DNAStreamWriter->setFrom(Reader->Unwrap(), static_cast<dna::DataLayer>(Layer), FMemoryResource::Instance());
 	DNAStreamWriter->write();
 }
