Fixed warrning found by static analysis in terse lib CharOutputStreamBuf.

diff --git a/Engine/Plugins/Animation/RigLogic/Source/RigLogicLib/Public/terse/types/CharOutputStreamBuf.h b/Engine/Plugins/Animation/RigLogic/Source/RigLogicLib/Public/terse/types/CharOutputStreamBuf.h
--- a/Engine/Plugins/Animation/RigLogic/Source/RigLogicLib/Public/terse/types/CharOutputStreamBuf.h
+++ b/Engine/Plugins/Animation/RigLogic/Source/RigLogicLib/Public/terse/types/CharOutputStreamBuf.h
@@ -29,7 +29,7 @@ class CharOutputStreamBuf : public std::streambuf {
         }
 
         ~CharOutputStreamBuf() {
-            sync();
+            CharOutputStreamBuf::sync();
         }
 
         std::streamsize xsputn(const CharType* source, std::streamsize size) override {
