Fix warning from clang static analysis (possible null dereference)

diff --git a/Engine/Plugins/Animation/RigLogic/Source/RigLogicLib/Private/pma/resources/ArenaMemoryResource.cpp b/Engine/Plugins/Animation/RigLogic/Source/RigLogicLib/Private/pma/resources/ArenaMemoryResource.cpp
--- a/Engine/Plugins/Animation/RigLogic/Source/RigLogicLib/Private/pma/resources/ArenaMemoryResource.cpp
+++ b/Engine/Plugins/Animation/RigLogic/Source/RigLogicLib/Private/pma/resources/ArenaMemoryResource.cpp
@@ -28,6 +28,9 @@ inline T* alignPointer(T* ptr, std::size_t alignment) {
 class ArenaMemoryResource::Impl {
     public:
         static Impl* create(std::size_t initialSize_, std::size_t regionSize_, float growthFactor_, MemoryResource* upstream_) {
+            if (upstream_ == nullptr) {
+                return nullptr;
+            }
             PolyAllocator<Impl> alloc{upstream_};
             return alloc.newObject(initialSize_, regionSize_, growthFactor_, upstream_);
         }
