Added RL_SHARED to public defintions of RL.

diff --git a/Engine/Plugins/Animation/RigLogic/Source/RigLogicLib/RigLogicLib.Build.cs b/Engine/Plugins/Animation/RigLogic/Source/RigLogicLib/RigLogicLib.Build.cs
--- a/Engine/Plugins/Animation/RigLogic/Source/RigLogicLib/RigLogicLib.Build.cs
+++ b/Engine/Plugins/Animation/RigLogic/Source/RigLogicLib/RigLogicLib.Build.cs
@@ -34,7 +34,8 @@ public class RigLogicLib : ModuleRules
         if (Target.LinkType != TargetLinkType.Monolithic)
         {
             PrivateDefinitions.Add("RL_BUILD_SHARED=1");
-        }
+			PublicDefinitions.Add("RL_SHARED=1");
+		}
 
         if (Target.Platform == UnrealTargetPlatform.Win64)
         {
