FileName description added to DNA Import

diff --git a/Engine/Plugins/Animation/RigLogic/Source/RigLogicEditor/Private/DNAImporter.cpp b/Engine/Plugins/Animation/RigLogic/Source/RigLogicEditor/Private/DNAImporter.cpp
--- a/Engine/Plugins/Animation/RigLogic/Source/RigLogicEditor/Private/DNAImporter.cpp
+++ b/Engine/Plugins/Animation/RigLogic/Source/RigLogicEditor/Private/DNAImporter.cpp
@@ -52,7 +52,10 @@ FDNAAssetImportOptions* GetImportOptions(FDNAImporter * DNAImporter, UDNAAssetIm
 			}
 		}
 
+		DNAImporter->SetDNAFileName(*FPaths::GetCleanFilename(FullPath));
+
 		//This option must always be the same value has the skeletalmesh one.
+		
 
 		//////////////////////////////////////////////////////////////////////////
 		// Set the information section data
@@ -60,8 +63,7 @@ FDNAAssetImportOptions* GetImportOptions(FDNAImporter * DNAImporter, UDNAAssetIm
 		//Make sure the file is open to be able to read the header before showing the options
 		//If the file is already open it will simply return false.
 		// do analytics on getting DNA data
-		ImportUI->FileVersion = DNAImporter->GetDNAFileVersion();
-		ImportUI->FileCreator = DNAImporter->GetDNAFileCreator();
+		ImportUI->FileName = DNAImporter->GetDNAFileName();
 
 		TSharedPtr<SWindow> ParentWindow;
 
diff --git a/Engine/Plugins/Animation/RigLogic/Source/RigLogicEditor/Public/DNAAssetImportUI.h b/Engine/Plugins/Animation/RigLogic/Source/RigLogicEditor/Public/DNAAssetImportUI.h
--- a/Engine/Plugins/Animation/RigLogic/Source/RigLogicEditor/Public/DNAAssetImportUI.h
+++ b/Engine/Plugins/Animation/RigLogic/Source/RigLogicEditor/Public/DNAAssetImportUI.h
@@ -44,11 +44,7 @@ public:
 		// DNA Asset file informations
 		// Transient value that are set everytime we show the options dialog. These are information only and should be string.
 
-	/* The DNA file version */
-	UPROPERTY(VisibleAnywhere, Transient, Category = DNAFileInformation, meta = (ImportType = "DNAAsset", DisplayName = "File Version"))
-		FString FileVersion;
-
-	/* The file creator information */
-	UPROPERTY(VisibleAnywhere, Transient, Category = DNAFileInformation, meta = (ImportType = "DNAAsset", DisplayName = "File Creator"))
-		FString FileCreator;
+	/* The DNA file name */
+	UPROPERTY(VisibleAnywhere, Transient, Category = DNAFileInformation, meta = (ImportType = "DNAAsset", DisplayName = "File Name"))
+		FString FileName;
 };
diff --git a/Engine/Plugins/Animation/RigLogic/Source/RigLogicEditor/Public/DNAImporter.h b/Engine/Plugins/Animation/RigLogic/Source/RigLogicEditor/Public/DNAImporter.h
--- a/Engine/Plugins/Animation/RigLogic/Source/RigLogicEditor/Public/DNAImporter.h
+++ b/Engine/Plugins/Animation/RigLogic/Source/RigLogicEditor/Public/DNAImporter.h
@@ -46,8 +46,8 @@ public:
 	 */
 	FDNAAssetImportOptions* GetImportOptions() const;
 	void PartialCleanUp();
-	FString GetDNAFileVersion() { return DNAFileVersion; }
-	FString GetDNAFileCreator() { return DNAFileCreator; }
+	FString GetDNAFileName() { return DNAFileName; }
+	void SetDNAFileName(FString DNAName) { this->DNAFileName = DNAName; };
 protected:
 	enum IMPORTPHASE
 	{
@@ -66,8 +66,7 @@ protected:
 	// base path of DNA file
 	FString FileBasePath;
 	TWeakObjectPtr<UObject> Parent;
-	FString DNAFileVersion;
-	FString DNAFileCreator;
+	FString DNAFileName;
 
 	FDNAImporter();
 
