Suppress false-positive warning from clang static analyzer

diff --git a/Engine/Plugins/Animation/RigLogic/Source/RigLogicLib/Public/terse/types/ArchiveOffset.h b/Engine/Plugins/Animation/RigLogic/Source/RigLogicLib/Public/terse/types/ArchiveOffset.h
--- a/Engine/Plugins/Animation/RigLogic/Source/RigLogicLib/Public/terse/types/ArchiveOffset.h
+++ b/Engine/Plugins/Animation/RigLogic/Source/RigLogicLib/Public/terse/types/ArchiveOffset.h
@@ -32,7 +32,9 @@ struct ArchiveOffset {
         ArchiveOffset* target;
 
         explicit Proxy(ArchiveOffset& ptr) : target{std::addressof(ptr)} {
-            target->proxy = this;
+            #if !defined(__clang_analyzer__)
+                target->proxy = this;
+            #endif
         }
 
         ~Proxy() {
diff --git a/Engine/Plugins/Animation/RigLogic/Source/RigLogicLib/Public/terse/types/ArchiveSize.h b/Engine/Plugins/Animation/RigLogic/Source/RigLogicLib/Public/terse/types/ArchiveSize.h
--- a/Engine/Plugins/Animation/RigLogic/Source/RigLogicLib/Public/terse/types/ArchiveSize.h
+++ b/Engine/Plugins/Animation/RigLogic/Source/RigLogicLib/Public/terse/types/ArchiveSize.h
@@ -39,8 +39,10 @@ struct ArchiveSize {
         Anchor<TOffset>* base;
 
         Proxy(ArchiveSize& size, Anchor<TOffset>& offset) : target{std::addressof(size)}, base{std::addressof(offset)} {
-            target->proxy = this;
-            base->onMove(onBaseMoved, this);
+            #if !defined(__clang_analyzer__)
+                target->proxy = this;
+                base->onMove(onBaseMoved, this);
+            #endif
         }
 
         ~Proxy() {
