Fix warnings caused by removal of UDNAIndexMapping class by reintroducing it as a deprecated class

diff --git a/Engine/Plugins/Animation/RigLogic/Source/RigLogicModule/Private/DNAIndexMapping.h b/Engine/Plugins/Animation/RigLogic/Source/RigLogicModule/Private/DNAIndexMapping.h
--- a/Engine/Plugins/Animation/RigLogic/Source/RigLogicModule/Private/DNAIndexMapping.h
+++ b/Engine/Plugins/Animation/RigLogic/Source/RigLogicModule/Private/DNAIndexMapping.h
@@ -4,6 +4,9 @@
 
 #include "CoreMinimal.h"
 #include "BoneIndices.h"
+#include "Engine/AssetUserData.h"
+
+#include "DNAIndexMapping.generated.h"
 
 class IBehaviorReader;
 class USkeleton;
@@ -34,3 +37,9 @@ struct FDNAIndexMapping
 	void MapMaskMultipliers(const IBehaviorReader* DNABehavior, const USkeleton* Skeleton);
 
 };
+
+UCLASS(NotBlueprintable, hidecategories = (Object), deprecated)
+class UDEPRECATED_DNAIndexMapping : public UAssetUserData
+{
+	GENERATED_BODY()
+};
\ No newline at end of file
