Exsposed source path of DNA Asset to AssetUserData in SkeletalMesh Details Panel

diff --git a/Engine/Plugins/Animation/RigLogic/Source/RigLogicModule/Private/DNAAsset.cpp b/Engine/Plugins/Animation/RigLogic/Source/RigLogicModule/Private/DNAAsset.cpp
--- a/Engine/Plugins/Animation/RigLogic/Source/RigLogicModule/Private/DNAAsset.cpp
+++ b/Engine/Plugins/Animation/RigLogic/Source/RigLogicModule/Private/DNAAsset.cpp
@@ -212,7 +212,7 @@ bool UDNAAsset::Init(const FString& DNAFilename)
 		UE_LOG(LogDNAAsset, Warning, TEXT("%s"), ANSI_TO_TCHAR(rl4::Status::get().message));
 	}
 
-	this->DNAFileName = DNAFilename; //memorize for re-import
+	DnaFileName = DNAFilename; //memorize for re-import
 	
 	if (!FPaths::FileExists(DNAFilename))
 	{
diff --git a/Engine/Plugins/Animation/RigLogic/Source/RigLogicModule/Public/DNAAsset.h b/Engine/Plugins/Animation/RigLogic/Source/RigLogicModule/Public/DNAAsset.h
--- a/Engine/Plugins/Animation/RigLogic/Source/RigLogicModule/Public/DNAAsset.h
+++ b/Engine/Plugins/Animation/RigLogic/Source/RigLogicModule/Public/DNAAsset.h
@@ -51,8 +51,8 @@ public:
 	TSharedPtr<IGeometryReader> GetGeometryReader();
 #endif
 
-	UPROPERTY()
-	FString DNAFileName; 
+	UPROPERTY(VisibleAnywhere, Category = ImportSettings)
+	FString DnaFileName;
 
 	bool Init(const FString& Filename);
 	void Serialize(FArchive& Ar) override;
