Implemented ContextMenu DNA Actions for Skeletal Mesh

diff --git a/Engine/Plugins/Animation/RigLogic/Source/RigLogicEditor/Private/DNAImporter.cpp b/Engine/Plugins/Animation/RigLogic/Source/RigLogicEditor/Private/DNAImporter.cpp
--- a/Engine/Plugins/Animation/RigLogic/Source/RigLogicEditor/Private/DNAImporter.cpp
+++ b/Engine/Plugins/Animation/RigLogic/Source/RigLogicEditor/Private/DNAImporter.cpp
@@ -15,10 +15,14 @@
 #include "Framework/Application/SlateApplication.h"
 #include "HAL/PlatformApplicationMisc.h"
 #include "DNAAssetImportWindow.h"
+#include "DesktopPlatformModule.h"
+#include "EditorDirectories.h"
 
 TSharedPtr<FDNAImporter> FDNAImporter::StaticInstance;
 TSharedPtr<FDNAImporter> FDNAImporter::StaticPreviewInstance;
 
+#define LOCTEXT_NAMESPACE "DNAAssetImport"
+
 FDNAAssetImportOptions* GetImportOptions(FDNAImporter * DNAImporter, UDNAAssetImportUI * ImportUI, bool bShowOptionDialog, bool bIsAutomated, const FString & FullPath, bool & OutOperationCanceled, const FString & InFilename)
 {
 	OutOperationCanceled = false;
@@ -194,3 +198,43 @@ void FDNAImporter::PartialCleanUp()
 	// reset
 	CurPhase = NOTSTARTED;
 }
+
+FString FDNAImporter::PromptForDNAImportFile()
+{
+	const FText PromptTitle = LOCTEXT("DNAPromptTitle", "Choose a file to import for DNA");
+
+	FString ChosenFilename("");
+
+	FString ExtensionStr;
+	ExtensionStr += TEXT("All model files|*.dna|");
+	ExtensionStr += TEXT("DNA files|*.dna|");
+	ExtensionStr += TEXT("All files|*.*");
+
+	// First, display the file open dialog for selecting the file.
+	TArray<FString> OpenFilenames;
+	IDesktopPlatform* DesktopPlatform = FDesktopPlatformModule::Get();
+	bool bOpen = false;
+	if (DesktopPlatform)
+	{
+		bOpen = DesktopPlatform->OpenFileDialog(
+			FSlateApplication::Get().FindBestParentWindowHandleForDialogs(nullptr),
+			PromptTitle.ToString(),
+			*FEditorDirectories::Get().GetLastDirectory(ELastDirectory::UNR),
+			TEXT(""),
+			*ExtensionStr,
+			EFileDialogFlags::None,
+			OpenFilenames
+		);
+	}
+
+	if (bOpen)
+	{
+		if (OpenFilenames.Num() == 1)
+		{
+			ChosenFilename = OpenFilenames[0];
+			FEditorDirectories::Get().SetLastDirectory(ELastDirectory::UNR, FPaths::GetPath(ChosenFilename));
+		}
+	}
+
+	return ChosenFilename;
+}
diff --git a/Engine/Plugins/Animation/RigLogic/Source/RigLogicEditor/Private/RigLogicEditor.cpp b/Engine/Plugins/Animation/RigLogic/Source/RigLogicEditor/Private/RigLogicEditor.cpp
--- a/Engine/Plugins/Animation/RigLogic/Source/RigLogicEditor/Private/RigLogicEditor.cpp
+++ b/Engine/Plugins/Animation/RigLogic/Source/RigLogicEditor/Private/RigLogicEditor.cpp
@@ -3,18 +3,130 @@
 #include "RigLogicEditor.h"
 #include "Modules/ModuleManager.h"
 #include "Logging/LogMacros.h"
-
+#include "DNAImporter.h"
 #include "RigUnit_RigLogic.h"
+#include "IAssetTypeActions.h"
+#include "AssetToolsModule.h"
+#include "Toolkits/AssetEditorToolkit.h"
+#include "Framework/MultiBox/MultiBoxBuilder.h"
+#include "ISettingsModule.h"
+#include "PropertyEditorModule.h"
+#include <ContentBrowserModule.h>
+#include <DNAAssetImportFactory.h>
+#include "Framework/Notifications/NotificationManager.h"
+#include "Widgets/Notifications/SNotificationList.h"
+#include "Misc/Paths.h"
+
+
 
 IMPLEMENT_MODULE(FRigLogicEditor, RigLogicEditor)
 
 DEFINE_LOG_CATEGORY_STATIC(LogRigLogicEditor, Log, All);
 
+#define LOCTEXT_NAMESPACE "DNAAssetImport"
+
 void FRigLogicEditor::StartupModule()
 {
+	FContentBrowserModule& ContentBrowserModule = FModuleManager::LoadModuleChecked<FContentBrowserModule>(TEXT("ContentBrowser"));	
+	TArray<FContentBrowserMenuExtender_SelectedAssets>& MenuExtenderDelegates = ContentBrowserModule.GetAllAssetViewContextMenuExtenders(); 
+	
+	MenuExtenderDelegates.Add(FContentBrowserMenuExtender_SelectedAssets::CreateStatic(&FRigLogicEditor::OnExtendSkelMeshWithDNASelectionMenu));
 }
 
-void FRigLogicEditor::ShutdownModule()
+
+TSharedRef<FExtender> FRigLogicEditor::OnExtendSkelMeshWithDNASelectionMenu(const TArray<FAssetData>& SelectedAssets)
+{
+	TSharedRef<FExtender> Extender = MakeShared<FExtender>();
+
+	Extender->AddMenuExtension(
+		"GetAssetActions",
+		EExtensionHook::After,
+		nullptr,
+		FMenuExtensionDelegate::CreateStatic(&FRigLogicEditor::CreateDnaActionsSubMenu, SelectedAssets)
+	);
+	return Extender;
+}
+
+void FRigLogicEditor::CreateDnaActionsSubMenu(FMenuBuilder& MenuBuilder, const TArray<FAssetData> SelectedAssets)
+{
+	const FAssetData& Asset = SelectedAssets[0];
+
+	if (Asset.GetClass()->IsChildOf(USkeletalMesh::StaticClass()))
+	{
+		MenuBuilder.AddSubMenu(
+			LOCTEXT("DNASkeletalMeshSubmenu", "MetaHuman DNA"),
+			LOCTEXT("DNAImportSubmenu_ToolTip", "DNA related actions"),
+			FNewMenuDelegate::CreateStatic(&FRigLogicEditor::GetDNAMenu, SelectedAssets),
+			false,
+			FSlateIcon(FAppStyle::GetAppStyleSetName(), "Icons.Import")
+		);	
+	}
+	
+}
+
+void FRigLogicEditor::GetDNAMenu(FMenuBuilder& MenuBuilder, const TArray<FAssetData> SelectedAssets)
 {
+	auto Mesh = SelectedAssets[0].GetAsset();
+	
+	MenuBuilder.AddMenuEntry(
+		LOCTEXT("Add DNA", "Import new DNA File"),
+		LOCTEXT("ImportDNA", "Import DNA"),
+		FSlateIcon(FAppStyle::GetAppStyleSetName(), "Icons.Import"),
+		FUIAction(FExecuteAction::CreateStatic(&FRigLogicEditor::ExecuteDNAImport, static_cast<UObject*>(Mesh)))
+	);
+
+	MenuBuilder.AddMenuEntry(
+		LOCTEXT("Reimport DNA", "Reimport existing DNA File"),
+		LOCTEXT("Reimport DNA", "Reimport DNA"),
+		FSlateIcon(FAppStyle::GetAppStyleSetName(), "Icons.Import"),
+		FUIAction(FExecuteAction::CreateStatic(&FRigLogicEditor::ExecuteDNAReimport, static_cast<UObject*>(Mesh)))
+	);
 }
 
+void FRigLogicEditor::ExecuteDNAImport(UObject* Mesh)
+{
+	FAssetToolsModule& AssetToolsModule = FModuleManager::Get().LoadModuleChecked<FAssetToolsModule>("AssetTools");
+	
+	FDNAImporter* DNAImporter = FDNAImporter::GetInstance();
+	const TArray<FString> Filenames = { DNAImporter->PromptForDNAImportFile() };
+	UFactory* Factory = Cast<UFactory>(UDNAAssetImportFactory::StaticClass()->GetDefaultObject());
+	FString Path = FPaths::GetPath(Mesh->GetPathName());
+	
+	AssetToolsModule.Get().ImportAssets(Filenames,Path, Factory,true, nullptr,true);
+}
+
+void FRigLogicEditor::ExecuteDNAReimport(class UObject* Mesh)
+{
+	FAssetToolsModule& AssetToolsModule = FModuleManager::Get().LoadModuleChecked<FAssetToolsModule>("AssetTools");
+	USkeletalMesh* SkelMesh = Cast<USkeletalMesh>(Mesh);
+	const TArray<UAssetUserData*>* AssetData = SkelMesh->GetAssetUserDataArray();
+
+	if(!AssetData->IsEmpty())
+	{
+		const UDNAAsset* DNAAsset = Cast<UDNAAsset>(&AssetData[0]);
+		if(DNAAsset)
+		{
+			const TArray<FString> Filenames = { DNAAsset->DnaFileName };
+			UFactory* Factory = Cast<UFactory>(UDNAAssetImportFactory::StaticClass()->GetDefaultObject());
+			FString Path = FPaths::GetPath(Mesh->GetPathName());
+
+			AssetToolsModule.Get().ImportAssets(Filenames, Path, Factory, true, nullptr, true);
+		}
+	}
+	else
+	{
+		FNotificationInfo Info(LOCTEXT("DNA_ReimportErrorMessage", "There is no DNA file attached to do Reimport"));
+		Info.ExpireDuration = 3.0f;
+		Info.bUseLargeFont = false;
+		TSharedPtr<SNotificationItem> Notification = FSlateNotificationManager::Get().AddNotification(Info);
+		if (Notification.IsValid())
+		{
+			Notification->SetCompletionState(SNotificationItem::CS_Fail);
+		}
+	}
+}
+
+void FRigLogicEditor::ShutdownModule()
+{
+
+}
\ No newline at end of file
diff --git a/Engine/Plugins/Animation/RigLogic/Source/RigLogicEditor/Public/DNAImporter.h b/Engine/Plugins/Animation/RigLogic/Source/RigLogicEditor/Public/DNAImporter.h
--- a/Engine/Plugins/Animation/RigLogic/Source/RigLogicEditor/Public/DNAImporter.h
+++ b/Engine/Plugins/Animation/RigLogic/Source/RigLogicEditor/Public/DNAImporter.h
@@ -46,6 +46,7 @@ public:
 	 */
 	FDNAAssetImportOptions* GetImportOptions() const;
 	void PartialCleanUp();
+	FString PromptForDNAImportFile();
 	FString GetDNAFileName() { return DNAFileName; }
 	void SetDNAFileName(FString DNAName) { this->DNAFileName = DNAName; };
 protected:
diff --git a/Engine/Plugins/Animation/RigLogic/Source/RigLogicEditor/Public/RigLogicEditor.h b/Engine/Plugins/Animation/RigLogic/Source/RigLogicEditor/Public/RigLogicEditor.h
--- a/Engine/Plugins/Animation/RigLogic/Source/RigLogicEditor/Public/RigLogicEditor.h
+++ b/Engine/Plugins/Animation/RigLogic/Source/RigLogicEditor/Public/RigLogicEditor.h
@@ -5,6 +5,10 @@
 #include "CoreMinimal.h"
 #include "Modules/ModuleInterface.h"
 
+class FExtender;
+class FMenuBuilder;
+struct FAssetData;
+
 class RIGLOGICEDITOR_API FRigLogicEditor : public IModuleInterface 
 {
 public:
@@ -13,4 +17,11 @@ public:
 
 private:
 	TArray<TSharedRef<class IAssetTypeActions>> AssetTypeActions;
-};
+
+	static TSharedRef<FExtender> OnExtendSkelMeshWithDNASelectionMenu(const TArray<FAssetData>& SelectedAssets);
+	static void CreateDnaActionsSubMenu(FMenuBuilder& MenuBuilder, const TArray<FAssetData> SelectedAssets);
+	static void GetDNAMenu(FMenuBuilder& MenuBuilder, const TArray<FAssetData> SelectedAssets);
+
+	static void ExecuteDNAImport(class UObject* Mesh);
+	static void ExecuteDNAReimport(class UObject* Mesh);
+};
\ No newline at end of file
diff --git a/Engine/Plugins/Animation/RigLogic/Source/RigLogicEditor/RigLogicEditor.Build.cs b/Engine/Plugins/Animation/RigLogic/Source/RigLogicEditor/RigLogicEditor.Build.cs
--- a/Engine/Plugins/Animation/RigLogic/Source/RigLogicEditor/RigLogicEditor.Build.cs
+++ b/Engine/Plugins/Animation/RigLogic/Source/RigLogicEditor/RigLogicEditor.Build.cs
@@ -34,7 +34,9 @@ namespace UnrealBuildTool.Rules
                     "SlateCore",
                     "ApplicationCore",
                     "Slate",
-                    "InputCore"
+                    "InputCore",
+					"EditorWidgets",
+					"DesktopPlatform"
                 }
             );
 
@@ -42,10 +44,12 @@ namespace UnrealBuildTool.Rules
                 new string[]
                 {
                     "PropertyEditor",
-                    "AssetTools"
-                }
+                    "AssetTools",
+					"Settings"
+				}
             );
 
-        }
-    }
+
+		}
+	}
 }
