Update RigLogic to v9.1.1 and GeneSplicer to v9.4.0

diff --git a/Engine/Plugins/Animation/RigLogic/RigLogic.uplugin b/Engine/Plugins/Animation/RigLogic/RigLogic.uplugin
--- a/Engine/Plugins/Animation/RigLogic/RigLogic.uplugin
+++ b/Engine/Plugins/Animation/RigLogic/RigLogic.uplugin
@@ -1,9 +1,9 @@
 {
 	"FileVersion": 3,
-	"Version": 8.1,
-	"VersionName": "8.1.0",
-	"FriendlyName": "RigLogic Plugin v8.1.0",
-	"Description": "3Lateral RigLogic Plugin for Facial Animation v8.1.0",
+	"Version": 9.1,
+	"VersionName": "9.1.1",
+	"FriendlyName": "RigLogic Plugin v9.1.1",
+	"Description": "3Lateral RigLogic Plugin for Facial Animation v9.1.1",
 	"Category": "Animation",
 	"CreatedBy": "3Lateral",
 	"CreatedByURL": "https://www.3lateral.com/",
diff --git a/Engine/Plugins/Animation/RigLogic/Source/RigLogicEditor/RigLogicEditor.Build.cs b/Engine/Plugins/Animation/RigLogic/Source/RigLogicEditor/RigLogicEditor.Build.cs
--- a/Engine/Plugins/Animation/RigLogic/Source/RigLogicEditor/RigLogicEditor.Build.cs
+++ b/Engine/Plugins/Animation/RigLogic/Source/RigLogicEditor/RigLogicEditor.Build.cs
@@ -1,54 +1,45 @@
 // Copyright Epic Games, Inc. All Rights Reserved.
 
-using System.IO;
-
 namespace UnrealBuildTool.Rules
 {
-    public class RigLogicEditor : ModuleRules
-    {
-        private string ModulePath
-        {
-            get { return ModuleDirectory; }
-        }
+	public class RigLogicEditor : ModuleRules
+	{
+		public RigLogicEditor(ReadOnlyTargetRules Target) : base(Target)
+		{
+			PublicDependencyModuleNames.AddRange(
+				new string[]
+				{
+					"Core",
+					"CoreUObject",
+					"Engine",
+					"ControlRig",
+					"UnrealEd",
+					"EditorFramework",
+					"MainFrame",
+					"RigLogicModule",
+					"RigLogicLib",
+					"PropertyEditor",
+					"SlateCore",
+					"ApplicationCore",
+					"Slate",
+					"InputCore"
+				}
+			);
 
-        public RigLogicEditor(ReadOnlyTargetRules Target) : base(Target)
-        {
-            if (Target.LinkType != TargetLinkType.Monolithic)
-            {
-                PublicDefinitions.Add("RL_SHARED=1");
-            }
+			PrivateIncludePathModuleNames.AddRange(
+				new string[]
+				{
+					"PropertyEditor",
+					"AssetTools"
+				}
+			);
 
-            PublicDependencyModuleNames.AddRange(
-                new string[]
-                {
-                    "Core",
-                    "CoreUObject",
-                    "Engine",
-                    "ControlRig",
-                    "UnrealEd",
-                    "EditorFramework",
-                    "MainFrame",
-                    "RigLogicModule",
-                    "RigLogicLib",
-                    "PropertyEditor",
-                    "SlateCore",
-                    "ApplicationCore",
-                    "Slate",
                     "InputCore",
 					"EditorWidgets",
 					"DesktopPlatform"
-                }
-            );
-
-            PrivateIncludePathModuleNames.AddRange(
-                new string[]
-                {
-                    "PropertyEditor",
                     "AssetTools",
 					"Settings"
 				}
-            );
-
 
 		}
 	}
diff --git a/Engine/Plugins/Animation/RigLogic/Source/RigLogicLib/Private/dna/BaseImpl.h b/Engine/Plugins/Animation/RigLogic/Source/RigLogicLib/Private/dna/BaseImpl.h
--- a/Engine/Plugins/Animation/RigLogic/Source/RigLogicLib/Private/dna/BaseImpl.h
+++ b/Engine/Plugins/Animation/RigLogic/Source/RigLogicLib/Private/dna/BaseImpl.h
@@ -3,6 +3,7 @@
 #pragma once
 
 #include "dna/DNA.h"
+#include "dna/Helpers.h"
 #include "dna/types/Aliases.h"
 
 namespace dna {
@@ -11,7 +12,12 @@ class BaseImpl {
     protected:
         explicit BaseImpl(MemoryResource* memRes_) :
             memRes{memRes_},
-            dna{memRes} {
+            dna{UnknownLayerPolicy::Preserve, UpgradeFormatPolicy::Allowed, memRes} {
+        }
+
+        BaseImpl(UnknownLayerPolicy unknownPolicy, UpgradeFormatPolicy upgradePolicy, MemoryResource* memRes_) :
+            memRes{memRes_},
+            dna{unknownPolicy, upgradePolicy, memRes} {
         }
 
         ~BaseImpl() = default;
@@ -27,9 +33,13 @@ class BaseImpl {
             return memRes;
         }
 
+        void rawCopyInto(DNA& destination, DataLayer layer, UnknownLayerPolicy policy, MemoryResource* memRes_) {
+            copy(dna, destination, layer, policy, memRes_);
+        }
+
     protected:
         MemoryResource* memRes;
-        DNA dna;
+        mutable DNA dna;
 
 };
 
diff --git a/Engine/Plugins/Animation/RigLogic/Source/RigLogicLib/Private/dna/DNA.h b/Engine/Plugins/Animation/RigLogic/Source/RigLogicLib/Private/dna/DNA.h
--- a/Engine/Plugins/Animation/RigLogic/Source/RigLogicLib/Private/dna/DNA.h
+++ b/Engine/Plugins/Animation/RigLogic/Source/RigLogicLib/Private/dna/DNA.h
@@ -2,11 +2,17 @@
 
 #pragma once
 
+#include "dna/DataLayer.h"
 #include "dna/LODMapping.h"
 #include "dna/SurjectiveMapping.h"
 #include "dna/TypeDefs.h"
+#include "dna/utils/Extd.h"
 
+#include <terse/types/Anchor.h>
 #include <terse/types/ArchiveOffset.h>
+#include <terse/types/ArchiveSize.h>
+#include <terse/types/Transparent.h>
+#include <terse/types/Versioned.h>
 
 #ifdef _MSC_VER
     #pragma warning(push)
@@ -26,16 +32,121 @@
 
 namespace dna {
 
-template<typename TFrom, typename TTo = TFrom>
-struct RawSurjectiveMapping : public SurjectiveMapping<TFrom, TTo> {
-    using SurjectiveMapping<TFrom, TTo>::SurjectiveMapping;
+template<typename TVersion, TVersion... Args>
+struct Dispatcher;
+
+template<typename TVersion, TVersion Fallback>
+struct Dispatcher<TVersion, Fallback> {
+
+    template<class Archive, typename TSerializable>
+    void operator()(Archive& archive, TSerializable& object, TVersion  /*unused*/, bool transparent) {
+        // Dispatching occurs only after DNA version has already been read, identified and confirmed
+        // to be supported (the unknown enum won't be passed here ever). Likewise, if a newer
+        // version is encountered than the latest version listed here, it will be dispatched to the
+        // latest available serializer.
+        auto versionedObject = terse::versioned(object, terse::Version<Fallback>{});
+        if (transparent) {
+            archive(terse::transparent(versionedObject));
+        } else {
+            archive(versionedObject);
+        }
+    }
+
+};
+
+template<typename TVersion, TVersion Expected, TVersion... Rest>
+struct Dispatcher<TVersion, Expected, Rest...> : Dispatcher<TVersion, Rest...> {
+    using Super = Dispatcher<TVersion, Rest...>;
+
+    template<class Archive, typename TSerializable>
+    void operator()(Archive& archive, TSerializable& object, TVersion version, bool transparent) {
+        if (version == Expected) {
+            auto versionedObject = terse::versioned(object, terse::Version<Expected>{});
+            if (transparent) {
+                archive(terse::transparent(versionedObject));
+            } else {
+                archive(versionedObject);
+            }
+        } else {
+            Super::operator()(archive, object, version, transparent);
+        }
+    }
+
+};
+
+static constexpr std::uint32_t sid4(const char* name) {
+    return ((static_cast<std::uint32_t>(name[0]) << 24u) +
+            (static_cast<std::uint32_t>(name[1]) << 16u) +
+            (static_cast<std::uint32_t>(name[2]) << 8u) +
+            (static_cast<std::uint32_t>(name[3]) << 0u));
+}
+
+static constexpr std::uint32_t rev(std::uint16_t generation, std::uint16_t version) {
+    return (static_cast<std::uint32_t>(generation) << 16u) + static_cast<std::uint32_t>(version);
+}
+
+static constexpr std::uint16_t high16(std::uint32_t packed) {
+    return static_cast<std::uint16_t>(packed >> 16u);
+}
+
+static constexpr std::uint16_t low16(std::uint32_t packed) {
+    return static_cast<std::uint16_t>(packed & 0x0000FFFFu);
+}
+
+enum FileVersion : std::uint64_t {
+    unknown = 0u,
+    v21 = rev(2, 1),
+    v22 = rev(2, 2),
+    v23 = rev(2, 3),
+    latest = v23
+};
+
+using Dispatch = Dispatcher<FileVersion, FileVersion::v21, FileVersion::v22, FileVersion::v23>;
+
+template<typename A, FileVersion B>
+using IfHigher = typename std::enable_if<(A::value() > B), void>::type;
+
+template<typename A, FileVersion B>
+using IfLower = typename std::enable_if < (A::value() < B), void > ::type;
+
+enum class UpgradeFormatPolicy {
+    Disallowed,
+    Allowed
+};
+
+struct SerializationContext {
+    FileVersion version;
+    void* data;
+};
+
+struct Version {
+    std::uint16_t generation;
+    std::uint16_t version;
+
+    Version(FileVersion fileVersion) :
+        generation{high16(static_cast<std::uint32_t>(fileVersion))},
+        version{low16(static_cast<std::uint32_t>(fileVersion))} {
+    }
+
+    operator FileVersion() const {
+        return static_cast<FileVersion>(rev(generation, version));
+    }
 
     template<class Archive>
     void serialize(Archive& archive) {
-        archive.label("from");
-        archive(this->from);
-        archive.label("to");
-        archive(this->to);
+        archive.label("generation");
+        archive(generation);
+        archive.label("version");
+        archive(version);
+    }
+
+    bool matches(FileVersion fileVersion) const {
+        const Version expected{fileVersion};
+        return (generation == expected.generation) && (version == expected.version);
+    }
+
+    bool supported() const {
+        return (generation == Version{FileVersion::latest}.generation);
     }
 
 };
@@ -58,6 +169,7 @@ struct ExpectedValue {
     void save(Archive& archive) {
         archive.label("value");
         archive(expected);
+        got = expected;
     }
 
     bool matches() const {
@@ -87,29 +199,6 @@ struct Signature {
 
 };
 
-struct Version {
-    ExpectedValue<std::uint16_t> generation;
-    ExpectedValue<std::uint16_t> version;
-
-    Version(std::uint16_t generation_, std::uint16_t version_) :
-        generation{generation_},
-        version{version_} {
-    }
-
-    template<class Archive>
-    void serialize(Archive& archive) {
-        archive.label("generation");
-        archive(generation);
-        archive.label("version");
-        archive(version);
-    }
-
-    bool matches() const {
-        return (generation.matches() && version.matches());
-    }
-
-};
-
 struct SectionLookupTable {
     terse::ArchiveOffset<std::uint32_t> descriptor;
     terse::ArchiveOffset<std::uint32_t> definition;
@@ -142,6 +231,256 @@ struct SectionLookupTable {
 
 };
 
+struct Index {
+    std::uint32_t id;
+    std::uint32_t version;
+    terse::ArchiveOffset<std::uint32_t> offset;
+    terse::ArchiveSize<std::uint32_t, std::uint32_t> size;
+
+    template<std::size_t Length = 16ul>
+    std::array<char, Length> label() {
+        std::array<char, Length> output;
+        std::memcpy(output.data(), &id, sizeof(id));
+        std::reverse(output.begin(), extd::advanced(output.begin(), sizeof(id)));
+        const std::uint32_t major = version >> 16u;
+        const std::uint32_t minor = version & 0x0000FFFFu;
+        std::snprintf(output.data() + sizeof(id), output.size() - sizeof(id), "%u.%u", major, minor);
+        return output;
+    }
+
+    template<class Archive>
+    void serialize(Archive& archive) {
+        archive.label("id");
+        archive(id);
+        archive.label("version");
+        archive(version);
+        archive.label("offset");
+        archive(offset);
+        archive.label("size");
+        archive(size);
+    }
+
+};
+
+struct IndexTable {
+    Vector<Index> entries;
+
+    explicit IndexTable(MemoryResource* memRes) : entries{memRes} {
+    }
+
+    Index* get(std::uint32_t id, std::uint32_t version) {
+        for (auto& index : entries) {
+            if ((index.id == id) && (index.version == version)) {
+                return &index;
+            }
+        }
+        return nullptr;
+    }
+
+    const Index* get(std::uint32_t id, std::uint32_t version) const {
+        return const_cast<IndexTable*>(this)->get(id, version);
+    }
+
+    Index* create(std::uint32_t id, std::uint32_t version) {
+        assert(get(id, version) == nullptr);
+        entries.push_back({id, version, {}, {}});
+        return &entries.back();
+    }
+
+    void destroy(std::uint32_t id, std::uint32_t version) {
+        entries.erase(std::remove_if(entries.begin(), entries.end(), [id, version](const Index& index) {
+                return (index.id == id) && (index.version == version);
+            }), entries.end());
+    }
+
+    template<class Archive>
+    void serialize(Archive& archive) {
+        archive.label("entries");
+        archive(entries);
+    }
+
+};
+
+struct UnknownLayer {
+    std::uint32_t id;
+    std::uint32_t version;
+    Blob<char> data;
+
+    UnknownLayer(std::uint32_t id_, std::uint32_t version_, MemoryResource* memRes) : id{id_}, version{version_}, data{memRes} {
+    }
+
+    std::uint32_t layerId() {
+        return id;
+    }
+
+    std::uint32_t layerVersion() {
+        return version;
+    }
+
+    bool layerSupported(FileVersion  /*unused*/) {
+        return true;
+    }
+
+    template<class Archive>
+    void serialize(Archive& archive) {
+        // Serialization of the layer type itself is transparent
+        // and should not generate an additional, new structure level in JSON
+        auto context = static_cast<SerializationContext*>(archive.getUserData());
+        auto index = static_cast<Index*>(context->data);
+
+        const auto layerName = index->label();
+        archive.label(layerName.data());
+
+        archive(terse::proxy(index->offset));
+        terse::Anchor<std::uint32_t> base;
+        archive(base);
+        data.setSize(index->size.value);
+        archive(data);
+        archive(terse::proxy(index->size, base));
+    }
+
+};
+
+template<std::uint32_t Id, std::uint32_t Version, typename TData, FileVersion AvailableFromFileVersion>
+struct Layer : TData {
+
+    template<typename ... Args>
+    explicit Layer(Args&& ... args) : TData{std::forward<Args>(args)...} {
+    }
+
+    std::uint32_t layerId() {
+        return Id;
+    }
+
+    std::uint32_t layerVersion() {
+        return Version;
+    }
+
+    bool layerSupported(FileVersion inVersion) {
+        return (inVersion >= AvailableFromFileVersion);
+    }
+
+    template<class Archive>
+    void serialize(Archive& archive) {
+        // Serialization of the layer type itself is transparent
+        // and should not generate an additional, new structure level in JSON
+        auto context = static_cast<SerializationContext*>(archive.getUserData());
+        auto index = static_cast<Index*>(context->data);
+
+        const auto layerName = index->label();
+        archive.label(layerName.data());
+
+        archive(terse::proxy(index->offset));
+        terse::Anchor<std::uint32_t> base;
+        archive(base);
+        Dispatch()(archive, static_cast<TData&>(*this), context->version, false);
+        archive(terse::proxy(index->size, base));
+    }
+
+};
+
+template<typename ... TLayer>
+struct LayerContainer;
+
+template<>
+struct LayerContainer<> {
+    UnknownLayerPolicy unknownPolicy;
+    UpgradeFormatPolicy upgradePolicy;
+    Vector<UnknownLayer> unknownLayers;
+
+    LayerContainer(UnknownLayerPolicy unknownPolicy_, UpgradeFormatPolicy upgradePolicy_, MemoryResource* memRes) :
+        unknownPolicy{unknownPolicy_}, upgradePolicy{upgradePolicy_}, unknownLayers{memRes} {
+    }
+
+    template<typename TArchive>
+    void serialize(TArchive& archive) {
+        auto context = static_cast<SerializationContext*>(archive.getUserData());
+        auto index = static_cast<Index*>(context->data);
+
+        for (auto& layer : unknownLayers) {
+            if ((index->id == layer.layerId()) && (index->version == layer.layerVersion())) {
+                archive(terse::transparent(layer));
+                return;
+            }
+        }
+
+        // This code path should be executed only by input archives, when loading unknown layers.
+        // After initially adding all unknown layers, there is no way for new unknown layers to be
+        // encountered, so the above loop will take care of all during saving.
+        if (unknownPolicy == UnknownLayerPolicy::Preserve) {
+            auto memRes = unknownLayers.get_allocator().getMemoryResource();
+            unknownLayers.push_back(UnknownLayer{index->id, index->version, memRes});
+            archive(terse::transparent(unknownLayers.back()));
+        } else {
+            // Mark index for deletion
+            context->data = nullptr;
+        }
+    }
+
+    void updateIndexTable(IndexTable&  /*unused*/, FileVersion  /*unused*/) {
+        // Unknown layers are already in the index (that's the only way they could have been encountered)
+    }
+
+};
+
+template<typename TLayer, typename ... TRest>
+struct LayerContainer<TLayer, TRest...> : LayerContainer<TRest...> {
+    using Super = LayerContainer<TRest...>;
+
+    TLayer* layer;
+
+    LayerContainer(UnknownLayerPolicy unknownPolicy_,
+                   UpgradeFormatPolicy upgradePolicy_,
+                   MemoryResource* memRes,
+                   TLayer* layer_,
+                   TRest* ... rest) :
+        Super{unknownPolicy_, upgradePolicy_, memRes, std::forward<TRest*>(rest)...},
+        layer{layer_} {
+    }
+
+    template<typename TArchive>
+    void serialize(TArchive& archive) {
+        auto context = static_cast<SerializationContext*>(archive.getUserData());
+        auto index = static_cast<Index*>(context->data);
+        if ((index->id == layer->layerId()) && (index->version == layer->layerVersion())) {
+            archive(terse::transparent(*layer));
+        } else {
+            archive(terse::transparent(static_cast<Super&>(*this)));
+        }
+    }
+
+    void updateIndexTable(IndexTable& indexTable, FileVersion fileVersion) {
+        const bool exists = (indexTable.get(layer->layerId(), layer->layerVersion()) != nullptr);
+        if (layer->layerSupported(fileVersion)) {
+            // Do not merge ifs, as the below else branch depends solely on the condition of the
+            // layer not being supported in the file version
+            if (!exists && (this->upgradePolicy == UpgradeFormatPolicy::Allowed)) {
+                indexTable.create(layer->layerId(), layer->layerVersion());
+            }
+        } else {
+            if (exists && (this->unknownPolicy == UnknownLayerPolicy::Ignore)) {
+                indexTable.destroy(layer->layerId(), layer->layerVersion());
+            }
+        }
+        Super::updateIndexTable(indexTable, fileVersion);
+    }
+
+};
+
+template<typename TFrom, typename TTo = TFrom>
+struct RawSurjectiveMapping : public SurjectiveMapping<TFrom, TTo> {
+    using SurjectiveMapping<TFrom, TTo>::SurjectiveMapping;
+
+    template<class Archive>
+    void serialize(Archive& archive) {
+        archive.label("from");
+        archive(this->from);
+        archive.label("to");
+        archive(this->to);
+    }
+
+};
+
 struct RawCoordinateSystem {
     std::uint16_t xAxis;
     std::uint16_t yAxis;
@@ -175,7 +514,6 @@ struct RawLODMapping : public LODMapping {
 struct RawDescriptor {
     using StringPair = std::tuple<String<char>, String<char> >;
 
-    terse::ArchiveOffset<std::uint32_t>::Proxy marker;
     String<char> name;
     std::uint16_t archetype;
     std::uint16_t gender;
@@ -189,8 +527,7 @@ struct RawDescriptor {
     String<char> complexity;
     String<char> dbName;
 
-    RawDescriptor(terse::ArchiveOffset<std::uint32_t>& markerTarget, MemoryResource* memRes) :
-        marker{markerTarget},
+    explicit RawDescriptor(MemoryResource* memRes) :
         name{memRes},
         archetype{},
         gender{},
@@ -207,7 +544,6 @@ struct RawDescriptor {
 
     template<class Archive>
     void serialize(Archive& archive) {
-        archive(marker);
         archive.label("name");
         archive(name);
         archive.label("archetype");
@@ -311,8 +647,8 @@ struct RawVector3Vector {
 
 };
 
+
 struct RawDefinition {
-    terse::ArchiveOffset<std::uint32_t>::Proxy marker;
     RawLODMapping lodJointMapping;
     RawLODMapping lodBlendShapeMapping;
     RawLODMapping lodAnimatedMapMapping;
@@ -328,8 +664,7 @@ struct RawDefinition {
     RawVector3Vector neutralJointTranslations;
     RawVector3Vector neutralJointRotations;
 
-    RawDefinition(terse::ArchiveOffset<std::uint32_t>& markerTarget, MemoryResource* memRes) :
-        marker{markerTarget},
+    explicit RawDefinition(MemoryResource* memRes) :
         lodJointMapping{memRes},
         lodBlendShapeMapping{memRes},
         lodAnimatedMapMapping{memRes},
@@ -348,7 +683,6 @@ struct RawDefinition {
 
     template<class Archive>
     void serialize(Archive& archive) {
-        archive(marker);
         archive.label("lodJointMapping");
         archive(lodJointMapping);
         archive.label("lodBlendShapeMapping");
@@ -568,49 +902,48 @@ struct RawAnimatedMaps {
 };
 
 struct RawBehavior {
-    terse::ArchiveOffset<std::uint32_t>::Proxy marker;
-
-    terse::ArchiveOffset<std::uint32_t>::Proxy controlsMarker;
     RawControls controls;
-
-    terse::ArchiveOffset<std::uint32_t>::Proxy jointsMarker;
     RawJoints joints;
-
-    terse::ArchiveOffset<std::uint32_t>::Proxy blendShapeChannelsMarker;
     RawBlendShapeChannels blendShapeChannels;
-
-    terse::ArchiveOffset<std::uint32_t>::Proxy animatedMapsMarker;
     RawAnimatedMaps animatedMaps;
 
-    RawBehavior(terse::ArchiveOffset<std::uint32_t>& markerTarget,
-                terse::ArchiveOffset<std::uint32_t>& controlsMarkerTarget,
-                terse::ArchiveOffset<std::uint32_t>& jointsMarkerTarget,
-                terse::ArchiveOffset<std::uint32_t>& blendShapeChannelsMarkerTarget,
-                terse::ArchiveOffset<std::uint32_t>& animatedMapsMarkerTarget,
-                MemoryResource* memRes) :
-        marker{markerTarget},
-        controlsMarker{controlsMarkerTarget},
+    explicit RawBehavior(MemoryResource* memRes) :
         controls{memRes},
-        jointsMarker{jointsMarkerTarget},
         joints{memRes},
-        blendShapeChannelsMarker{blendShapeChannelsMarkerTarget},
         blendShapeChannels{memRes},
-        animatedMapsMarker{animatedMapsMarkerTarget},
         animatedMaps{memRes} {
     }
 
     template<class Archive>
-    void serialize(Archive& archive) {
-        archive(marker, controlsMarker);
+    void serialize(Archive& archive, terse::Version<FileVersion::v21>  /*unused*/) {
+        auto context = static_cast<SerializationContext*>(archive.getUserData());
+        auto sections = static_cast<SectionLookupTable*>(context->data);
+
+        archive(terse::proxy(sections->controls));
+        archive.label("controls");
+        archive(controls);
+
+        archive(terse::proxy(sections->joints));
+        archive.label("joints");
+        archive(joints);
+
+        archive(terse::proxy(sections->blendShapeChannels));
+        archive.label("blendShapeChannels");
+        archive(blendShapeChannels);
+
+        archive(terse::proxy(sections->animatedMaps));
+        archive.label("animatedMaps");
+        archive(animatedMaps);
+    }
+
+    template<class Archive, typename AnyVersion>
+    void serialize(Archive& archive, AnyVersion) {
         archive.label("controls");
         archive(controls);
-        archive(jointsMarker);
         archive.label("joints");
         archive(joints);
-        archive(blendShapeChannelsMarker);
         archive.label("blendShapeChannels");
         archive(blendShapeChannels);
-        archive(animatedMapsMarker);
         archive.label("animatedMaps");
         archive(animatedMaps);
     }
@@ -739,6 +1072,8 @@ struct RawBlendShapeTarget {
 
 struct RawMesh {
     terse::ArchiveOffset<std::uint32_t> offset;
+    terse::ArchiveSize<std::uint32_t, std::uint32_t> size;
+    terse::Anchor<std::uint32_t> baseMarker;
     RawVector3Vector positions;
     RawTextureCoordinateVector textureCoordinates;
     RawVector3Vector normals;
@@ -747,10 +1082,13 @@ struct RawMesh {
     std::uint16_t maximumInfluencePerVertex;
     Vector<RawVertexSkinWeights> skinWeights;
     Vector<RawBlendShapeTarget> blendShapeTargets;
-    terse::ArchiveOffset<std::uint32_t>::Proxy marker;
+    terse::ArchiveOffset<std::uint32_t>::Proxy offsetMarker;
+    terse::ArchiveSize<std::uint32_t, std::uint32_t>::Proxy sizeMarker;
 
     explicit RawMesh(MemoryResource* memRes) :
         offset{},
+        size{},
+        baseMarker{},
         positions{memRes},
         textureCoordinates{memRes},
         normals{memRes},
@@ -759,13 +1097,22 @@ struct RawMesh {
         maximumInfluencePerVertex{},
         skinWeights{memRes},
         blendShapeTargets{memRes},
-        marker{offset} {
+        offsetMarker{offset},
+        sizeMarker{size, baseMarker} {
     }
 
     template<class Archive>
     void serialize(Archive& archive) {
+        auto context = static_cast<SerializationContext*>(archive.getUserData());
+        Dispatch()(archive, *this, context->version, true);
+    }
+
+    template<class Archive>
+    void serialize(Archive& archive, terse::Version<FileVersion::v21>  /*unused*/) {
+        // Old format uses a hacky absolute file offset pointing to the next Mesh start
         archive.label("offset");
         archive(offset);
+
         archive.label("positions");
         archive(positions);
         archive.label("textureCoordinates");
@@ -782,109 +1129,309 @@ struct RawMesh {
         archive(skinWeights);
         archive.label("blendShapeTargets");
         archive(blendShapeTargets);
-        archive(marker);
+
+        archive(offsetMarker);
+    }
+
+    template<class Archive, typename AnyVersion>
+    void serialize(Archive& archive, AnyVersion) {
+        // The new format uses a size, denoting the size of the current mesh, thus allowing
+        // the same skip mechanism, but in a position independent way
+        archive.label("size");
+        archive(size);
+        archive(baseMarker);
+
+        archive.label("positions");
+        archive(positions);
+        archive.label("textureCoordinates");
+        archive(textureCoordinates);
+        archive.label("normals");
+        archive(normals);
+        archive.label("layouts");
+        archive(layouts);
+        archive.label("faces");
+        archive(faces);
+        archive.label("maximumInfluencePerVertex");
+        archive(maximumInfluencePerVertex);
+        archive.label("skinWeights");
+        archive(skinWeights);
+        archive.label("blendShapeTargets");
+        archive(blendShapeTargets);
+
+        // Offset is not written into file
+        archive(sizeMarker);
     }
 
 };
 
 struct RawGeometry {
-    terse::ArchiveOffset<std::uint32_t>::Proxy marker;
     Vector<RawMesh> meshes;
 
-    RawGeometry(terse::ArchiveOffset<std::uint32_t>& markerTarget, MemoryResource* memRes) :
-        marker{markerTarget},
-        meshes{memRes} {
+    explicit RawGeometry(MemoryResource* memRes) : meshes{memRes} {
     }
 
     template<class Archive>
     void serialize(Archive& archive) {
-        archive(marker);
         archive.label("meshes");
         archive(meshes);
     }
 
 };
 
+struct RawActivationFunction {
+    std::uint16_t functionId;
+    DynArray<float> parameters;
+
+    explicit RawActivationFunction(MemoryResource* memRes) : functionId{}, parameters{memRes} {
+    }
+
+    template<class Archive>
+    void serialize(Archive& archive) {
+        archive.label("functionId");
+        archive(functionId);
+        archive.label("parameters");
+        archive(parameters);
+    }
+
+};
+
+struct RawNeuralNetworkLayer {
+    DynArray<float> biases;
+    DynArray<float> weights;  // 2D Matrix
+    RawActivationFunction activationFunction;
+
+    explicit RawNeuralNetworkLayer(MemoryResource* memRes) :
+        biases{memRes},
+        weights{memRes},
+        activationFunction{memRes} {
+    }
+
+    template<class Archive>
+    void serialize(Archive& archive) {
+        archive.label("biases");
+        archive(biases);
+        archive.label("weights");
+        archive(weights);
+        archive.label("activationFunction");
+        archive(activationFunction);
+    }
+
+};
+
+struct RawNeuralNetwork {
+    terse::ArchiveSize<std::uint32_t, std::uint32_t> size;
+    terse::Anchor<std::uint32_t> baseMarker;
+    DynArray<std::uint16_t> outputIndices;
+    DynArray<std::uint16_t> inputIndices;
+    Vector<RawNeuralNetworkLayer> layers;
+    terse::ArchiveSize<std::uint32_t, std::uint32_t>::Proxy sizeMarker;
+
+    explicit RawNeuralNetwork(MemoryResource* memRes) :
+        size{},
+        baseMarker{},
+        outputIndices{memRes},
+        inputIndices{memRes},
+        layers{memRes},
+        sizeMarker{size, baseMarker} {
+    }
+
+    template<class Archive>
+    void serialize(Archive& archive) {
+        archive.label("size");
+        archive(size);
+        archive(baseMarker);
+        archive.label("outputIndices");
+        archive(outputIndices);
+        archive.label("inputIndices");
+        archive(inputIndices);
+        archive.label("layers");
+        archive(layers);
+        archive(sizeMarker);
+    }
+
+};
+
+struct RawMeshRegionMembership {
+    Matrix<String<char> > regionNames;  // [meshIndex][regionIndex]
+    Vector<Matrix<std::uint16_t> > indices;  // [meshIndex][regionIndex][neuralNetworkIndices...]
+
+    explicit RawMeshRegionMembership(MemoryResource* memRes) :
+        regionNames{memRes},
+        indices{memRes} {
+    }
+
+    template<class Archive>
+    void serialize(Archive& archive) {
+        archive.label("regionNames");
+        archive(regionNames);
+        archive.label("indices");
+        archive(indices);
+    }
+
+};
+
+struct RawMachineLearnedBehavior {
+    Vector<String<char> > mlControlNames;
+    RawLODMapping lodNeuralNetworkMapping;
+    RawMeshRegionMembership neuralNetworkToMeshRegion;
+    Vector<RawNeuralNetwork> neuralNetworks;
+
+    explicit RawMachineLearnedBehavior(MemoryResource* memRes) :
+        mlControlNames{memRes},
+        lodNeuralNetworkMapping{memRes},
+        neuralNetworkToMeshRegion{memRes},
+        neuralNetworks{memRes} {
+    }
+
+    template<class Archive>
+    void serialize(Archive& archive) {
+        archive.label("mlControlNames");
+        archive(mlControlNames);
+        archive.label("lodNeuralNetworkMapping");
+        archive(lodNeuralNetworkMapping);
+        archive.label("neuralNetworkToMeshRegion");
+        archive(neuralNetworkToMeshRegion);
+        archive.label("neuralNetworks");
+        archive(neuralNetworks);
+    }
+
+};
+
 struct DNA {
     MemoryResource* memRes;
     Signature<3> signature{{'D', 'N', 'A'}};
-    Version version{2, 1};
-    SectionLookupTable sections;
-    RawDescriptor descriptor;
-    RawDefinition definition;
-    RawBehavior behavior;
-    RawGeometry geometry;
-    Signature<3> eof{{'A', 'N', 'D'}};
-
-    explicit DNA(MemoryResource* memRes_) :
+    Version version{FileVersion::latest};
+    IndexTable indexTable;
+    Layer<sid4("desc"), rev(1, 1), RawDescriptor, FileVersion::v22> descriptor;
+    Layer<sid4("defn"), rev(1, 1), RawDefinition, FileVersion::v22> definition;
+    Layer<sid4("bhvr"), rev(1, 1), RawBehavior, FileVersion::v22> behavior;
+    Layer<sid4("geom"), rev(1, 1), RawGeometry, FileVersion::v22> geometry;
+    Layer<sid4("mlbh"), rev(1, 0), RawMachineLearnedBehavior, FileVersion::v23> machineLearnedBehavior;
+
+    using Layers =
+        LayerContainer<decltype(descriptor),
+                       decltype(definition),
+                       decltype(behavior),
+                       decltype(geometry),
+                       decltype(machineLearnedBehavior)>;
+    Layers layers;
+
+    DNA(UnknownLayerPolicy unknownPolicy, UpgradeFormatPolicy upgradePolicy, MemoryResource* memRes_) :
         memRes{memRes_},
-        sections{},
-        descriptor{sections.descriptor, memRes},
-        definition{sections.definition, memRes},
-        behavior{sections.behavior,
-                 sections.controls,
-                 sections.joints,
-                 sections.blendShapeChannels,
-                 sections.animatedMaps,
-                 memRes},
-        geometry{sections.geometry, memRes} {
+        indexTable{memRes_},
+        descriptor{memRes},
+        definition{memRes},
+        behavior{memRes},
+        geometry{memRes},
+        machineLearnedBehavior{memRes},
+        layers{unknownPolicy, upgradePolicy, memRes, &descriptor, &definition, &behavior, &geometry, &machineLearnedBehavior} {
+    }
+
+    explicit DNA(MemoryResource* memRes_) : DNA{UnknownLayerPolicy::Preserve, UpgradeFormatPolicy::Allowed, memRes_} {
     }
 
     template<class Archive>
-    void load(Archive& archive) {
+    void serialize(Archive& archive) {
         archive.label("signature");
         archive(signature);
         archive.label("version");
         archive(version);
-        if (signature.matches() && version.matches()) {
-            archive.label("sections");
-            archive(sections);
-            archive.label("descriptor");
-            archive(descriptor);
-            archive.label("definition");
-            archive(definition);
-            archive.label("behavior");
-            archive(behavior);
-            archive.label("geometry");
-            archive(geometry);
-            archive.label("eof");
-            archive(eof);
-            assert(eof.matches());
+
+        if (!signature.matches()) {
+            return;
         }
+
+        SerializationContext context{version, nullptr};
+        auto oldUserData = archive.getUserData();
+        archive.setUserData(&context);
+        Dispatch()(archive, *this, version, true);
+        archive.setUserData(oldUserData);
     }
 
     template<class Archive>
-    void save(Archive& archive) {
-        archive.label("signature");
-        archive(signature);
-        archive.label("version");
-        archive(version);
+    void serialize(Archive& archive, terse::Version<FileVersion::v21> v21) {
+        auto context = static_cast<SerializationContext*>(archive.getUserData());
+
+        SectionLookupTable sections;
+        context->data = &sections;
+
         archive.label("sections");
         archive(sections);
+
         archive.label("descriptor");
-        archive(descriptor);
+        archive(terse::proxy(sections.descriptor));
+        archive(terse::versioned(static_cast<RawDescriptor&>(descriptor), v21));
         archive.label("definition");
-        archive(definition);
+        archive(terse::proxy(sections.definition));
+        archive(terse::versioned(static_cast<RawDefinition&>(definition), v21));
         archive.label("behavior");
-        archive(behavior);
+        archive(terse::proxy(sections.behavior));
+        archive(terse::versioned(static_cast<RawBehavior&>(behavior), v21));
         archive.label("geometry");
-        archive(geometry);
-        archive.label("eof");
+        archive(terse::proxy(sections.geometry));
+        archive(terse::versioned(static_cast<RawGeometry&>(geometry), v21));
+
+        Signature<3> eof{{'A', 'N', 'D'}};
         archive(eof);
+        assert(eof.matches());
+
+        context->data = nullptr;
+    }
+
+    template<class Archive, typename AnyVersion>
+    IfHigher<AnyVersion, FileVersion::v21> load(Archive& archive, AnyVersion) {
+        auto context = static_cast<SerializationContext*>(archive.getUserData());
+
+        archive.label("index");
+        archive(indexTable);
+
+        std::size_t indexCount = indexTable.entries.size();
+        for (std::size_t i = {}; i < indexCount;) {
+            auto& index = indexTable.entries[i];
+            context->data = &index;
+            archive(terse::transparent(layers));
+            // Check if index was marked for deletion, and if so, get rid of it
+            if (context->data == nullptr) {
+                indexTable.destroy(index.id, index.version);
+                --indexCount;
+            } else {
+                ++i;
+            }
+        }
+
+        context->data = nullptr;
+    }
+
+    template<class Archive, typename AnyVersion>
+    IfHigher<AnyVersion, FileVersion::v21> save(Archive& archive, AnyVersion) {
+        auto context = static_cast<SerializationContext*>(archive.getUserData());
+
+        // Add layers to index that were not already there
+        layers.updateIndexTable(indexTable, context->version);
+
+        archive.label("index");
+        archive(indexTable);
+
+        for (auto& index : indexTable.entries) {
+            context->data = &index;
+            archive(terse::transparent(layers));
+        }
     }
 
     void unloadDefinition() {
-        definition = RawDefinition{sections.definition, memRes};
+        static_cast<RawDefinition&>(definition) = RawDefinition{memRes};
     }
 
     void unloadBehavior() {
-        behavior =
-            RawBehavior{sections.behavior, sections.controls, sections.joints, sections.blendShapeChannels, sections.animatedMaps,
-                        memRes};
+        static_cast<RawBehavior&>(behavior) = RawBehavior{memRes};
     }
 
     void unloadGeometry() {
-        geometry = RawGeometry{sections.geometry, memRes};
+        static_cast<RawGeometry&>(geometry) = RawGeometry{memRes};
+    }
+
+    void unloadMachineLearnedBehavior() {
+        static_cast<RawMachineLearnedBehavior&>(machineLearnedBehavior) = RawMachineLearnedBehavior{memRes};
     }
 
 };
diff --git a/Engine/Plugins/Animation/RigLogic/Source/RigLogicLib/Private/dna/DataLayerBitmask.h b/Engine/Plugins/Animation/RigLogic/Source/RigLogicLib/Private/dna/DataLayerBitmask.h
--- a/Engine/Plugins/Animation/RigLogic/Source/RigLogicLib/Private/dna/DataLayerBitmask.h
+++ b/Engine/Plugins/Animation/RigLogic/Source/RigLogicLib/Private/dna/DataLayerBitmask.h
@@ -13,32 +13,41 @@ enum class DataLayerBitmask {
     Behavior = 4,
     GeometryBlendShapesOnly = 8,
     GeometryRest = 16,
+    MachineLearnedBehavior = 32
 };
 
 inline DataLayerBitmask computeDataLayerBitmask(DataLayer layer) {
     DataLayerBitmask result = DataLayerBitmask::Descriptor;
-    if (layer == DataLayer::Definition) {
-        result |= DataLayerBitmask::Definition;
-    } else if (layer == DataLayer::Behavior) {
+    if (layer == DataLayer::All) {
         result |= DataLayerBitmask::Definition;
         result |= DataLayerBitmask::Behavior;
-    } else if (layer == DataLayer::Geometry) {
-        result |= DataLayerBitmask::Definition;
         result |= DataLayerBitmask::GeometryBlendShapesOnly;
         result |= DataLayerBitmask::GeometryRest;
-    } else if (layer == DataLayer::GeometryWithoutBlendShapes) {
-        result |= DataLayerBitmask::Definition;
-        result |= DataLayerBitmask::GeometryRest;
-    } else if (layer == DataLayer::AllWithoutBlendShapes) {
-        result |= DataLayerBitmask::Definition;
-        result |= DataLayerBitmask::Behavior;
-        result |= DataLayerBitmask::GeometryRest;
-    } else if (layer == DataLayer::All) {
+        result |= DataLayerBitmask::MachineLearnedBehavior;
+        return result;
+    }
+
+    if (contains(layer, DataLayer::Definition)) {
         result |= DataLayerBitmask::Definition;
+    }
+
+    if (contains(layer, DataLayer::Behavior)) {
         result |= DataLayerBitmask::Behavior;
+    }
+
+    if (contains(layer, DataLayer::Geometry)) {
         result |= DataLayerBitmask::GeometryBlendShapesOnly;
         result |= DataLayerBitmask::GeometryRest;
     }
+
+    if (contains(layer, DataLayer::GeometryWithoutBlendShapes)) {
+        result |= DataLayerBitmask::GeometryRest;
+    }
+
+    if (contains(layer, DataLayer::MachineLearnedBehavior)) {
+        result |= DataLayerBitmask::MachineLearnedBehavior;
+    }
+
     return result;
 }
 
diff --git a/Engine/Plugins/Animation/RigLogic/Source/RigLogicLib/Private/dna/DenormalizedData.h b/Engine/Plugins/Animation/RigLogic/Source/RigLogicLib/Private/dna/DenormalizedData.h
--- a/Engine/Plugins/Animation/RigLogic/Source/RigLogicLib/Private/dna/DenormalizedData.h
+++ b/Engine/Plugins/Animation/RigLogic/Source/RigLogicLib/Private/dna/DenormalizedData.h
@@ -2,7 +2,6 @@
 
 #pragma once
 
-#include "dna/layers/DefinitionReader.h"
 #include "dna/LODMapping.h"
 #include "dna/types/Aliases.h"
 #include "dna/utils/Extd.h"
diff --git a/Engine/Plugins/Animation/RigLogic/Source/RigLogicLib/Private/dna/Reader.cpp b/Engine/Plugins/Animation/RigLogic/Source/RigLogicLib/Private/dna/Reader.cpp
--- a/Engine/Plugins/Animation/RigLogic/Source/RigLogicLib/Private/dna/Reader.cpp
+++ b/Engine/Plugins/Animation/RigLogic/Source/RigLogicLib/Private/dna/Reader.cpp
@@ -4,10 +4,12 @@
 
 namespace dna {
 
+HeaderReader::~HeaderReader() = default;
 DescriptorReader::~DescriptorReader() = default;
 DefinitionReader::~DefinitionReader() = default;
 BehaviorReader::~BehaviorReader() = default;
 GeometryReader::~GeometryReader() = default;
+MachineLearnedBehaviorReader::~MachineLearnedBehaviorReader() = default;
 Reader::~Reader() = default;
 
 }  // namespace dna
diff --git a/Engine/Plugins/Animation/RigLogic/Source/RigLogicLib/Private/dna/ReaderImpl.h b/Engine/Plugins/Animation/RigLogic/Source/RigLogicLib/Private/dna/ReaderImpl.h
--- a/Engine/Plugins/Animation/RigLogic/Source/RigLogicLib/Private/dna/ReaderImpl.h
+++ b/Engine/Plugins/Animation/RigLogic/Source/RigLogicLib/Private/dna/ReaderImpl.h
@@ -4,9 +4,6 @@
 
 #include "dna/BaseImpl.h"
 #include "dna/DenormalizedData.h"
-#include "dna/layers/DefinitionReader.h"
-#include "dna/layers/Descriptor.h"
-#include "dna/layers/Geometry.h"
 #include "dna/TypeDefs.h"
 
 #ifdef _MSC_VER
@@ -27,7 +24,11 @@ class ReaderImpl : public TReaderBase, public virtual BaseImpl {
     public:
         explicit ReaderImpl(MemoryResource* memRes_);
 
-        // DescriptorReader methods start
+        // HeaderReader methods
+        std::uint16_t getFileFormatGeneration() const override;
+        std::uint16_t getFileFormatVersion() const override;
+
+        // DescriptorReader methods
         StringView getName() const override;
         Archetype getArchetype() const override;
         Gender getGender() const override;
@@ -43,7 +44,7 @@ class ReaderImpl : public TReaderBase, public virtual BaseImpl {
         StringView getDBComplexity() const override;
         StringView getDBName() const override;
 
-        // DefinitionReader methods start
+        // DefinitionReader methods
         std::uint16_t getGUIControlCount() const override;
         StringView getGUIControlName(std::uint16_t index) const override;
         std::uint16_t getRawControlCount() const override;
@@ -77,7 +78,7 @@ class ReaderImpl : public TReaderBase, public virtual BaseImpl {
         ConstArrayView<float> getNeutralJointRotationYs() const override;
         ConstArrayView<float> getNeutralJointRotationZs() const override;
 
-        // BehaviorReader methods start
+        // BehaviorReader methods
         ConstArrayView<std::uint16_t> getGUIToRawInputIndices() const override;
         ConstArrayView<std::uint16_t> getGUIToRawOutputIndices() const override;
         ConstArrayView<float> getGUIToRawFromValues() const override;
@@ -108,7 +109,7 @@ class ReaderImpl : public TReaderBase, public virtual BaseImpl {
         ConstArrayView<float> getAnimatedMapSlopeValues() const override;
         ConstArrayView<float> getAnimatedMapCutValues() const override;
 
-        // GeometryReader methods start
+        // GeometryReader methods
         std::uint32_t getVertexPositionCount(std::uint16_t meshIndex) const override;
         Position getVertexPosition(std::uint16_t meshIndex, std::uint32_t vertexIndex) const override;
         ConstArrayView<float> getVertexPositionXs(std::uint16_t meshIndex) const override;
@@ -149,6 +150,25 @@ class ReaderImpl : public TReaderBase, public virtual BaseImpl {
                                                          std::uint16_t blendShapeTargetIndex) const override;
         ConstArrayView<std::uint32_t> getBlendShapeTargetVertexIndices(std::uint16_t meshIndex,
                                                                        std::uint16_t blendShapeTargetIndex) const override;
+        // MachineLearnedBehaviorReader methods
+        std::uint16_t getMLControlCount() const override;
+        StringView getMLControlName(std::uint16_t index) const override;
+        std::uint16_t getNeuralNetworkCount() const override;
+        std::uint16_t getNeuralNetworkIndexListCount() const override;
+        ConstArrayView<std::uint16_t> getNeuralNetworkIndicesForLOD(std::uint16_t lod) const override;
+        std::uint16_t getMeshRegionCount(std::uint16_t meshIndex) const override;
+        StringView getMeshRegionName(std::uint16_t meshIndex, std::uint16_t regionIndex) const override;
+        ConstArrayView<std::uint16_t> getNeuralNetworkIndicesForMeshRegion(std::uint16_t meshIndex,
+                                                                           std::uint16_t regionIndex) const override;
+        ConstArrayView<std::uint16_t> getNeuralNetworkInputIndices(std::uint16_t netIndex) const override;
+        ConstArrayView<std::uint16_t> getNeuralNetworkOutputIndices(std::uint16_t netIndex) const override;
+        std::uint16_t getNeuralNetworkLayerCount(std::uint16_t netIndex) const override;
+        ActivationFunction getNeuralNetworkLayerActivationFunction(std::uint16_t netIndex,
+                                                                   std::uint16_t layerIndex) const override;
+        ConstArrayView<float> getNeuralNetworkLayerActivationFunctionParameters(std::uint16_t netIndex,
+                                                                                std::uint16_t layerIndex) const override;
+        ConstArrayView<float> getNeuralNetworkLayerBiases(std::uint16_t netIndex, std::uint16_t layerIndex) const override;
+        ConstArrayView<float> getNeuralNetworkLayerWeights(std::uint16_t netIndex, std::uint16_t layerIndex) const override;
 
     protected:
         mutable DenormalizedData<TReaderBase> cache;
@@ -172,6 +192,16 @@ inline ReaderImpl<TReaderBase>::ReaderImpl(MemoryResource* memRes_) : BaseImpl{m
     #pragma warning(push)
     #pragma warning(disable : 4505)
 #endif
+template<class TReaderBase>
+inline std::uint16_t ReaderImpl<TReaderBase>::getFileFormatGeneration() const {
+    return dna.version.generation;
+}
+
+template<class TReaderBase>
+inline std::uint16_t ReaderImpl<TReaderBase>::getFileFormatVersion() const {
+    return dna.version.version;
+}
+
 template<class TReaderBase>
 inline StringView ReaderImpl<TReaderBase>::getName() const {
     return {dna.descriptor.name.data(), dna.descriptor.name.size()};
@@ -964,6 +994,133 @@ const {
     return {};
 }
 
+template<class TReaderBase>
+inline std::uint16_t ReaderImpl<TReaderBase>::getMLControlCount() const {
+    return static_cast<std::uint16_t>(dna.machineLearnedBehavior.mlControlNames.size());
+}
+
+template<class TReaderBase>
+inline StringView ReaderImpl<TReaderBase>::getMLControlName(std::uint16_t index) const {
+    if (index < dna.machineLearnedBehavior.mlControlNames.size()) {
+        const auto& mlControlName = dna.machineLearnedBehavior.mlControlNames[index];
+        return {mlControlName.data(), mlControlName.size()};
+    }
+    return {};
+}
+
+template<class TReaderBase>
+inline std::uint16_t ReaderImpl<TReaderBase>::getNeuralNetworkCount() const {
+    return static_cast<std::uint16_t>(dna.machineLearnedBehavior.neuralNetworks.size());
+}
+
+template<class TReaderBase>
+inline std::uint16_t ReaderImpl<TReaderBase>::getNeuralNetworkIndexListCount() const {
+    return dna.machineLearnedBehavior.lodNeuralNetworkMapping.getIndexListCount();
+}
+
+template<class TReaderBase>
+ConstArrayView<std::uint16_t> ReaderImpl<TReaderBase>::getNeuralNetworkIndicesForLOD(std::uint16_t lod) const {
+    return dna.machineLearnedBehavior.lodNeuralNetworkMapping.getIndices(lod);
+}
+
+template<class TReaderBase>
+inline std::uint16_t ReaderImpl<TReaderBase>::getMeshRegionCount(std::uint16_t meshIndex) const {
+    if (meshIndex < dna.machineLearnedBehavior.neuralNetworkToMeshRegion.regionNames.size()) {
+        return static_cast<std::uint16_t>(dna.machineLearnedBehavior.neuralNetworkToMeshRegion.regionNames[meshIndex].size());
+    }
+    return {};
+}
+
+template<class TReaderBase>
+inline StringView ReaderImpl<TReaderBase>::getMeshRegionName(std::uint16_t meshIndex, std::uint16_t regionIndex) const {
+    if ((meshIndex < dna.machineLearnedBehavior.neuralNetworkToMeshRegion.regionNames.size()) &&
+        (regionIndex < dna.machineLearnedBehavior.neuralNetworkToMeshRegion.regionNames[meshIndex].size())) {
+        return StringView{dna.machineLearnedBehavior.neuralNetworkToMeshRegion.regionNames[meshIndex][regionIndex]};
+    }
+    return {};
+}
+
+template<class TReaderBase>
+ConstArrayView<std::uint16_t> ReaderImpl<TReaderBase>::getNeuralNetworkIndicesForMeshRegion(std::uint16_t meshIndex,
+                                                                                            std::uint16_t regionIndex) const {
+    if (meshIndex < dna.machineLearnedBehavior.neuralNetworkToMeshRegion.indices.size()) {
+        const auto& mesh = dna.machineLearnedBehavior.neuralNetworkToMeshRegion.indices[meshIndex];
+        if (regionIndex < mesh.size()) {
+            const auto& region = mesh[regionIndex];
+            return {region.data(), region.size()};
+        }
+    }
+    return {};
+}
+
+template<class TReaderBase>
+inline ConstArrayView<std::uint16_t> ReaderImpl<TReaderBase>::getNeuralNetworkInputIndices(std::uint16_t netIndex) const {
+    const auto& neuralNets = dna.machineLearnedBehavior.neuralNetworks;
+    if (netIndex < neuralNets.size()) {
+        return ConstArrayView<std::uint16_t>{neuralNets[netIndex].inputIndices};
+    }
+    return {};
+}
+
+template<class TReaderBase>
+inline ConstArrayView<std::uint16_t> ReaderImpl<TReaderBase>::getNeuralNetworkOutputIndices(std::uint16_t netIndex) const {
+    const auto& neuralNets = dna.machineLearnedBehavior.neuralNetworks;
+    if (netIndex < neuralNets.size()) {
+        return ConstArrayView<std::uint16_t>{neuralNets[netIndex].outputIndices};
+    }
+    return {};
+}
+
+template<class TReaderBase>
+inline std::uint16_t ReaderImpl<TReaderBase>::getNeuralNetworkLayerCount(std::uint16_t netIndex) const {
+    const auto& neuralNets = dna.machineLearnedBehavior.neuralNetworks;
+    if (netIndex < neuralNets.size()) {
+        return static_cast<std::uint16_t>(neuralNets[netIndex].layers.size());
+    }
+    return {};
+}
+
+template<class TReaderBase>
+inline ActivationFunction ReaderImpl<TReaderBase>::getNeuralNetworkLayerActivationFunction(std::uint16_t netIndex,
+                                                                                           std::uint16_t layerIndex) const {
+    const auto& neuralNets = dna.machineLearnedBehavior.neuralNetworks;
+    if ((netIndex < neuralNets.size()) && (layerIndex < neuralNets[netIndex].layers.size())) {
+        return static_cast<ActivationFunction>(neuralNets[netIndex].layers[layerIndex].activationFunction.functionId);
+    }
+    return {};
+}
+
+template<class TReaderBase>
+inline ConstArrayView<float> ReaderImpl<TReaderBase>::getNeuralNetworkLayerActivationFunctionParameters(std::uint16_t netIndex,
+                                                                                                        std::uint16_t layerIndex)
+const {
+    const auto& neuralNets = dna.machineLearnedBehavior.neuralNetworks;
+    if ((netIndex < neuralNets.size()) && (layerIndex < neuralNets[netIndex].layers.size())) {
+        return ConstArrayView<float>{neuralNets[netIndex].layers[layerIndex].activationFunction.parameters};
+    }
+    return {};
+}
+
+template<class TReaderBase>
+inline ConstArrayView<float> ReaderImpl<TReaderBase>::getNeuralNetworkLayerBiases(std::uint16_t netIndex,
+                                                                                  std::uint16_t layerIndex) const {
+    const auto& neuralNets = dna.machineLearnedBehavior.neuralNetworks;
+    if ((netIndex < neuralNets.size()) && (layerIndex < neuralNets[netIndex].layers.size())) {
+        return ConstArrayView<float>{neuralNets[netIndex].layers[layerIndex].biases};
+    }
+    return {};
+}
+
+template<class TReaderBase>
+inline ConstArrayView<float> ReaderImpl<TReaderBase>::getNeuralNetworkLayerWeights(std::uint16_t netIndex,
+                                                                                   std::uint16_t layerIndex) const {
+    const auto& neuralNets = dna.machineLearnedBehavior.neuralNetworks;
+    if ((netIndex < neuralNets.size()) && (layerIndex < neuralNets[netIndex].layers.size())) {
+        return ConstArrayView<float>{neuralNets[netIndex].layers[layerIndex].weights};
+    }
+    return {};
+}
+
 #ifdef _MSC_VER
     #pragma warning(pop)
 #endif
diff --git a/Engine/Plugins/Animation/RigLogic/Source/RigLogicLib/Private/dna/TypeDefs.h b/Engine/Plugins/Animation/RigLogic/Source/RigLogicLib/Private/dna/TypeDefs.h
--- a/Engine/Plugins/Animation/RigLogic/Source/RigLogicLib/Private/dna/TypeDefs.h
+++ b/Engine/Plugins/Animation/RigLogic/Source/RigLogicLib/Private/dna/TypeDefs.h
@@ -9,6 +9,7 @@
 #include <pma/resources/AlignedMemoryResource.h>
 #include <pma/resources/ArenaMemoryResource.h>
 #include <pma/resources/DefaultMemoryResource.h>
+#include <terse/types/Blob.h>
 #include <terse/types/DynArray.h>
 
 namespace dna {
@@ -18,6 +19,9 @@ using namespace pma;
 template<typename T>
 using AlignedAllocator = PolyAllocator<T, 64ul, AlignedMemoryResource>;
 
+template<typename T>
+using Blob = terse::Blob<T, PolyAllocator<T> >;
+
 template<typename T>
 using DynArray = terse::DynArray<T, PolyAllocator<T> >;
 
diff --git a/Engine/Plugins/Animation/RigLogic/Source/RigLogicLib/Private/dna/Writer.cpp b/Engine/Plugins/Animation/RigLogic/Source/RigLogicLib/Private/dna/Writer.cpp
--- a/Engine/Plugins/Animation/RigLogic/Source/RigLogicLib/Private/dna/Writer.cpp
+++ b/Engine/Plugins/Animation/RigLogic/Source/RigLogicLib/Private/dna/Writer.cpp
@@ -26,10 +26,12 @@ using VertexLayoutVector = Vector<VertexLayout>;
 
 }  // namespace
 
+HeaderWriter::~HeaderWriter() = default;
 DescriptorWriter::~DescriptorWriter() = default;
 DefinitionWriter::~DefinitionWriter() = default;
 BehaviorWriter::~BehaviorWriter() = default;
 GeometryWriter::~GeometryWriter() = default;
+MachineLearnedBehaviorWriter::~MachineLearnedBehaviorWriter() = default;
 Writer::~Writer() = default;
 
 template<typename TVector, typename TGetter>
@@ -42,6 +44,11 @@ static TVector collect_n(std::size_t count, TGetter getter, MemoryResource* memR
     return retval;
 }
 
+static void copyHeader(const HeaderReader* source, HeaderWriter* destination, MemoryResource*  /*unused*/) {
+    destination->setFileFormatGeneration(source->getFileFormatGeneration());
+    destination->setFileFormatVersion(source->getFileFormatVersion());
+}
+
 static void copyDescriptor(const DescriptorReader* source, DescriptorWriter* destination, MemoryResource*  /*unused*/) {
     destination->setName(source->getName());
     destination->setArchetype(source->getArchetype());
@@ -79,11 +86,11 @@ using IndicesGetter = std::function<ConstArrayView<std::uint16_t>(std::uint16_t)
 using IndicesSetter = std::function<void (std::uint16_t, const std::uint16_t*, std::uint16_t)>;
 using LODMappingSetter = std::function<void (std::uint16_t, std::uint16_t)>;
 
-static void copyNameIndices(IndicesGetter getIndices,
-                            IndicesSetter setIndices,
-                            LODMappingSetter setLODMapping,
-                            std::uint16_t lodCount,
-                            MemoryResource* memRes) {
+static void copyLODIndices(IndicesGetter getIndices,
+                           IndicesSetter setIndices,
+                           LODMappingSetter setLODMapping,
+                           std::uint16_t lodCount,
+                           MemoryResource* memRes) {
     Matrix<std::uint16_t> allIndices{memRes};
     std::uint16_t index = 0u;
     for (std::uint16_t lod = 0u; lod < lodCount; ++lod) {
@@ -147,25 +154,25 @@ static void copyDefinition(const DefinitionReader* source, DefinitionWriter* des
     }
 
     using namespace std::placeholders;
-    copyNameIndices(
+    copyLODIndices(
         std::bind(&DefinitionReader::getJointIndicesForLOD, source, _1),
         std::bind(&DefinitionWriter::setJointIndices, destination, _1, _2, _3),
         std::bind(&DefinitionWriter::setLODJointMapping, destination, _1, _2),
         lodCount,
         memRes);
-    copyNameIndices(
+    copyLODIndices(
         std::bind(&DefinitionReader::getBlendShapeChannelIndicesForLOD, source, _1),
         std::bind(&DefinitionWriter::setBlendShapeChannelIndices, destination, _1, _2, _3),
         std::bind(&DefinitionWriter::setLODBlendShapeChannelMapping, destination, _1, _2),
         lodCount,
         memRes);
-    copyNameIndices(
+    copyLODIndices(
         std::bind(&DefinitionReader::getAnimatedMapIndicesForLOD, source, _1),
         std::bind(&DefinitionWriter::setAnimatedMapIndices, destination, _1, _2, _3),
         std::bind(&DefinitionWriter::setLODAnimatedMapMapping, destination, _1, _2),
         lodCount,
         memRes);
-    copyNameIndices(
+    copyLODIndices(
         std::bind(&DefinitionReader::getMeshIndicesForLOD, source, _1),
         std::bind(&DefinitionWriter::setMeshIndices, destination, _1, _2, _3),
         std::bind(&DefinitionWriter::setLODMeshMapping, destination, _1, _2),
@@ -410,12 +417,98 @@ static void copyBlendShapeTargets(const GeometryReader* source, GeometryWriter*
     }
 }
 
-void Writer::setFrom(const Reader* source, DataLayer layer, MemoryResource* memRes) {
+static bool hasMachineLearnedBehavior(const MachineLearnedBehaviorReader* source) {
+    // Heuristic for determining whether source DNA actually has any machine learned behavior data
+    return (source->getMLControlCount() != 0u) || (source->getNeuralNetworkCount() != 0u);
+}
+
+static void copyMachineLearnedBehavior(const MachineLearnedBehaviorReader* source,
+                                       MachineLearnedBehaviorWriter* destination,
+                                       MemoryResource* memRes) {
+    destination->clearMLControlNames();
+    destination->clearMeshRegionNames();
+    destination->clearNeuralNetworks();
+    destination->clearNeuralNetworkIndices();
+    destination->clearNeuralNetworkIndicesPerMeshRegion();
+
+    if (!hasMachineLearnedBehavior(source)) {
+        // Source DNA was loaded without machine learned behavior layer
+        return;
+    }
+
+    for (std::uint16_t i = source->getMLControlCount(); i > 0u; --i) {
+        const auto idx = static_cast<std::uint16_t>(i - 1u);
+        destination->setMLControlName(idx, source->getMLControlName(idx).data());
+    }
+
+    using namespace std::placeholders;
+    copyLODIndices(
+        std::bind(&MachineLearnedBehaviorReader::getNeuralNetworkIndicesForLOD, source, _1),
+        std::bind(&MachineLearnedBehaviorWriter::setNeuralNetworkIndices, destination, _1, _2, _3),
+        std::bind(&MachineLearnedBehaviorWriter::setLODNeuralNetworkMapping, destination, _1, _2),
+        source->getLODCount(),
+        memRes);
+
+    for (std::uint16_t meshIndexPlusOne = source->getMeshCount(); meshIndexPlusOne > 0u; --meshIndexPlusOne) {
+        const auto meshIndex = static_cast<std::uint16_t>(meshIndexPlusOne - 1u);
+        for (std::uint16_t regionIndexPlusOne = source->getMeshRegionCount(meshIndex); regionIndexPlusOne > 0u;
+             --regionIndexPlusOne) {
+            const auto regionIndex = static_cast<std::uint16_t>(regionIndexPlusOne - 1u);
+            destination->setMeshRegionName(meshIndex, regionIndex, source->getMeshRegionName(meshIndex, regionIndex));
+            const auto netIndices = source->getNeuralNetworkIndicesForMeshRegion(meshIndex, regionIndex);
+            destination->setNeuralNetworkIndicesForMeshRegion(meshIndex,
+                                                              regionIndex,
+                                                              netIndices.data(),
+                                                              static_cast<std::uint16_t>(netIndices.size()));
+        }
+    }
+
+    for (std::uint16_t netIndexPlusOne = source->getNeuralNetworkCount(); netIndexPlusOne > 0u; --netIndexPlusOne) {
+        const auto netIndex = static_cast<std::uint16_t>(netIndexPlusOne - 1u);
+        const auto inputIndices = source->getNeuralNetworkInputIndices(netIndex);
+        destination->setNeuralNetworkInputIndices(netIndex, inputIndices.data(),
+                                                  static_cast<std::uint16_t>(inputIndices.size()));
+
+        const auto outputIndices = source->getNeuralNetworkOutputIndices(netIndex);
+        destination->setNeuralNetworkOutputIndices(netIndex, outputIndices.data(),
+                                                   static_cast<std::uint16_t>(outputIndices.size()));
+
+        const auto layerCount = source->getNeuralNetworkLayerCount(netIndex);
+        for (std::uint16_t layerIndexPlusOne = layerCount; layerIndexPlusOne > 0u; --layerIndexPlusOne) {
+            const auto layerIndex = static_cast<std::uint16_t>(layerIndexPlusOne - 1u);
+
+            const auto activationFunction = source->getNeuralNetworkLayerActivationFunction(netIndex, layerIndex);
+            destination->setNeuralNetworkLayerActivationFunction(netIndex, layerIndex, activationFunction);
+
+            const auto activationFunctionParams = source->getNeuralNetworkLayerActivationFunctionParameters(netIndex,
+                                                                                                            layerIndex);
+            destination->setNeuralNetworkLayerActivationFunctionParameters(netIndex,
+                                                                           layerIndex,
+                                                                           activationFunctionParams.data(),
+                                                                           static_cast<std::uint16_t>(activationFunctionParams.
+                                                                                                      size()));
+
+            const auto biases = source->getNeuralNetworkLayerBiases(netIndex, layerIndex);
+            destination->setNeuralNetworkLayerBiases(netIndex,
+                                                     layerIndex,
+                                                     biases.data(),
+                                                     static_cast<std::uint32_t>(biases.size()));
+
+            const auto weights = source->getNeuralNetworkLayerWeights(netIndex, layerIndex);
+            destination->setNeuralNetworkLayerWeights(netIndex, layerIndex, weights.data(),
+                                                      static_cast<std::uint32_t>(weights.size()));
+        }
+    }
+}
+
+void Writer::setFrom(const Reader* source, DataLayer layer, UnknownLayerPolicy  /*unused*/, MemoryResource* memRes) {
+    // Unknown layers are inaccessible when copying through the public API
     if (source == nullptr) {
         return;
     }
 
     const auto bitmask = computeDataLayerBitmask(layer);
+    copyHeader(source, this, memRes);
     copyDescriptor(source, this, memRes);
     if (contains(bitmask, DataLayerBitmask::Definition)) {
         copyDefinition(source, this, memRes);
@@ -429,6 +522,9 @@ void Writer::setFrom(const Reader* source, DataLayer layer, MemoryResource* memR
     if (contains(bitmask, DataLayerBitmask::GeometryBlendShapesOnly)) {
         copyBlendShapeTargets(source, this, memRes);
     }
+    if (contains(bitmask, DataLayerBitmask::MachineLearnedBehavior)) {
+        copyMachineLearnedBehavior(source, this, memRes);
+    }
 }
 
 }  // namespace dna
diff --git a/Engine/Plugins/Animation/RigLogic/Source/RigLogicLib/Private/dna/WriterImpl.h b/Engine/Plugins/Animation/RigLogic/Source/RigLogicLib/Private/dna/WriterImpl.h
--- a/Engine/Plugins/Animation/RigLogic/Source/RigLogicLib/Private/dna/WriterImpl.h
+++ b/Engine/Plugins/Animation/RigLogic/Source/RigLogicLib/Private/dna/WriterImpl.h
@@ -3,8 +3,6 @@
 #pragma once
 
 #include "dna/BaseImpl.h"
-#include "dna/layers/Descriptor.h"
-#include "dna/layers/Geometry.h"
 #include "dna/TypeDefs.h"
 #include "dna/utils/Extd.h"
 
@@ -38,6 +36,10 @@ class WriterImpl : public TWriterBase, public virtual BaseImpl {
     public:
         explicit WriterImpl(MemoryResource* memRes_);
 
+        // HeaderWriter methods
+        void setFileFormatGeneration(std::uint16_t generation) override;
+        void setFileFormatVersion(std::uint16_t version) override;
+
         // DescriptorWriter methods
         void setName(const char* name) override;
         void setArchetype(Archetype archetype) override;
@@ -158,6 +160,43 @@ class WriterImpl : public TWriterBase, public virtual BaseImpl {
                                               const std::uint32_t* vertexIndices,
                                               std::uint32_t count) override;
 
+        // MachineLearnedBehaviorWriter methods
+        void clearMLControlNames() override;
+        void setMLControlName(std::uint16_t index, const char* name) override;
+        void clearNeuralNetworks() override;
+        void clearNeuralNetworkIndices() override;
+        void setNeuralNetworkIndices(std::uint16_t index, const std::uint16_t* netIndices, std::uint16_t count) override;
+        void clearMeshRegionNames() override;
+        void clearMeshRegionNames(std::uint16_t meshIndex) override;
+        void setMeshRegionName(std::uint16_t meshIndex, std::uint16_t regionIndex, const char* name) override;
+        void clearLODNeuralNetworkMappings() override;
+        void setLODNeuralNetworkMapping(std::uint16_t lod, std::uint16_t index) override;
+        void clearNeuralNetworkIndicesPerMeshRegion() override;
+        void setNeuralNetworkIndicesForMeshRegion(std::uint16_t meshIndex,
+                                                  std::uint16_t regionIndex,
+                                                  const std::uint16_t* netIndices,
+                                                  std::uint16_t count) override;
+        void deleteNeuralNetwork(std::uint16_t netIndex) override;
+        void setNeuralNetworkInputIndices(std::uint16_t netIndex, const std::uint16_t* inputIndices,
+                                          std::uint16_t count) override;
+        void setNeuralNetworkOutputIndices(std::uint16_t netIndex, const std::uint16_t* outputIndices,
+                                           std::uint16_t count) override;
+        void clearNeuralNetworkLayers(std::uint16_t netIndex) override;
+        void setNeuralNetworkLayerActivationFunction(std::uint16_t netIndex, std::uint16_t layerIndex,
+                                                     ActivationFunction function) override;
+        void setNeuralNetworkLayerActivationFunctionParameters(std::uint16_t netIndex,
+                                                               std::uint16_t layerIndex,
+                                                               const float* activationFunctionParameters,
+                                                               std::uint16_t count) override;
+        void setNeuralNetworkLayerBiases(std::uint16_t netIndex,
+                                         std::uint16_t layerIndex,
+                                         const float* biases,
+                                         std::uint32_t count) override;
+        void setNeuralNetworkLayerWeights(std::uint16_t netIndex,
+                                          std::uint16_t layerIndex,
+                                          const float* weights,
+                                          std::uint32_t count) override;
+
 };
 
 
@@ -177,6 +216,16 @@ WriterImpl<TWriterBase>::WriterImpl(MemoryResource* memRes_) : BaseImpl{memRes_}
     #pragma warning(push)
     #pragma warning(disable : 4505)
 #endif
+template<class TWriterBase>
+inline void WriterImpl<TWriterBase>::setFileFormatGeneration(std::uint16_t generation) {
+    dna.version.generation = generation;
+}
+
+template<class TWriterBase>
+inline void WriterImpl<TWriterBase>::setFileFormatVersion(std::uint16_t version) {
+    dna.version.version = version;
+}
+
 template<class TWriterBase>
 inline void WriterImpl<TWriterBase>::setName(const char* name) {
     dna.descriptor.name = name;
@@ -753,6 +802,155 @@ inline void WriterImpl<TWriterBase>::setBlendShapeTargetVertexIndices(std::uint1
     blendShapeTargets[blendShapeTargetIndex].vertexIndices.assign(vertexIndices, vertexIndices + count);
 }
 
+template<class TWriterBase>
+inline void WriterImpl<TWriterBase>::clearMLControlNames() {
+    dna.machineLearnedBehavior.mlControlNames.clear();
+}
+
+template<class TWriterBase>
+inline void WriterImpl<TWriterBase>::setMLControlName(std::uint16_t index, const char* name) {
+    ensureHasSize(dna.machineLearnedBehavior.mlControlNames, index + 1ul, memRes);
+    dna.machineLearnedBehavior.mlControlNames[index] = name;
+}
+
+template<class TWriterBase>
+inline void WriterImpl<TWriterBase>::clearNeuralNetworks() {
+    dna.machineLearnedBehavior.neuralNetworks.clear();
+}
+
+template<class TWriterBase>
+inline void WriterImpl<TWriterBase>::clearNeuralNetworkIndices() {
+    dna.machineLearnedBehavior.lodNeuralNetworkMapping.resetIndices();
+}
+
+template<class TWriterBase>
+inline void WriterImpl<TWriterBase>::setNeuralNetworkIndices(std::uint16_t index,
+                                                             const std::uint16_t* netIndices,
+                                                             std::uint16_t count) {
+    dna.machineLearnedBehavior.lodNeuralNetworkMapping.clearIndices(index);
+    dna.machineLearnedBehavior.lodNeuralNetworkMapping.addIndices(index, netIndices, count);
+}
+
+template<class TWriterBase>
+inline void WriterImpl<TWriterBase>::clearLODNeuralNetworkMappings() {
+    dna.machineLearnedBehavior.lodNeuralNetworkMapping.resetLODs();
+}
+
+template<class TWriterBase>
+inline void WriterImpl<TWriterBase>::setLODNeuralNetworkMapping(std::uint16_t lod, std::uint16_t index) {
+    dna.machineLearnedBehavior.lodNeuralNetworkMapping.associateLODWithIndices(lod, index);
+}
+
+template<class TWriterBase>
+inline void WriterImpl<TWriterBase>::clearMeshRegionNames() {
+    dna.machineLearnedBehavior.neuralNetworkToMeshRegion.regionNames.clear();
+}
+
+template<class TWriterBase>
+inline void WriterImpl<TWriterBase>::clearMeshRegionNames(std::uint16_t meshIndex) {
+    if (meshIndex < dna.machineLearnedBehavior.neuralNetworkToMeshRegion.regionNames.size()) {
+        dna.machineLearnedBehavior.neuralNetworkToMeshRegion.regionNames[meshIndex].clear();
+    }
+}
+
+template<class TWriterBase>
+inline void WriterImpl<TWriterBase>::setMeshRegionName(std::uint16_t meshIndex, std::uint16_t regionIndex, const char* name) {
+    ensureHasSize(dna.machineLearnedBehavior.neuralNetworkToMeshRegion.regionNames, meshIndex + 1ul, memRes);
+    ensureHasSize(dna.machineLearnedBehavior.neuralNetworkToMeshRegion.regionNames[meshIndex], regionIndex + 1ul, memRes);
+    dna.machineLearnedBehavior.neuralNetworkToMeshRegion.regionNames[meshIndex][regionIndex] = name;
+}
+
+template<class TWriterBase>
+inline void WriterImpl<TWriterBase>::clearNeuralNetworkIndicesPerMeshRegion() {
+    dna.machineLearnedBehavior.neuralNetworkToMeshRegion.indices.clear();
+}
+
+template<class TWriterBase>
+inline void WriterImpl<TWriterBase>::setNeuralNetworkIndicesForMeshRegion(std::uint16_t meshIndex,
+                                                                          std::uint16_t regionIndex,
+                                                                          const std::uint16_t* netIndices,
+                                                                          std::uint16_t count) {
+    ensureHasSize(dna.machineLearnedBehavior.neuralNetworkToMeshRegion.indices, meshIndex + 1ul, memRes);
+    ensureHasSize(dna.machineLearnedBehavior.neuralNetworkToMeshRegion.indices[meshIndex], regionIndex + 1ul, memRes);
+    auto& region = dna.machineLearnedBehavior.neuralNetworkToMeshRegion.indices[meshIndex][regionIndex];
+    region.assign(netIndices, netIndices + count);
+}
+
+template<class TWriterBase>
+inline void WriterImpl<TWriterBase>::deleteNeuralNetwork(std::uint16_t netIndex) {
+    if (netIndex < dna.machineLearnedBehavior.neuralNetworks.size()) {
+        auto it = extd::advanced(dna.machineLearnedBehavior.neuralNetworks.begin(), netIndex);
+        dna.machineLearnedBehavior.neuralNetworks.erase(it);
+    }
+}
+
+template<class TWriterBase>
+inline void WriterImpl<TWriterBase>::setNeuralNetworkInputIndices(std::uint16_t netIndex,
+                                                                  const std::uint16_t* inputIndices,
+                                                                  std::uint16_t count) {
+    ensureHasSize(dna.machineLearnedBehavior.neuralNetworks, netIndex + 1ul, memRes);
+    auto& neuralNet = dna.machineLearnedBehavior.neuralNetworks[netIndex];
+    neuralNet.inputIndices.assign(inputIndices, inputIndices + count);
+}
+
+template<class TWriterBase>
+inline void WriterImpl<TWriterBase>::setNeuralNetworkOutputIndices(std::uint16_t netIndex,
+                                                                   const std::uint16_t* outputIndices,
+                                                                   std::uint16_t count) {
+    ensureHasSize(dna.machineLearnedBehavior.neuralNetworks, netIndex + 1ul, memRes);
+    auto& neuralNet = dna.machineLearnedBehavior.neuralNetworks[netIndex];
+    neuralNet.outputIndices.assign(outputIndices, outputIndices + count);
+}
+
+template<class TWriterBase>
+inline void WriterImpl<TWriterBase>::clearNeuralNetworkLayers(std::uint16_t netIndex) {
+    ensureHasSize(dna.machineLearnedBehavior.neuralNetworks, netIndex + 1ul, memRes);
+    dna.machineLearnedBehavior.neuralNetworks[netIndex].layers.clear();
+}
+
+template<class TWriterBase>
+inline void WriterImpl<TWriterBase>::setNeuralNetworkLayerActivationFunction(std::uint16_t netIndex,
+                                                                             std::uint16_t layerIndex,
+                                                                             ActivationFunction function) {
+    ensureHasSize(dna.machineLearnedBehavior.neuralNetworks, netIndex + 1ul, memRes);
+    ensureHasSize(dna.machineLearnedBehavior.neuralNetworks[netIndex].layers, layerIndex + 1ul, memRes);
+    auto& layer = dna.machineLearnedBehavior.neuralNetworks[netIndex].layers[layerIndex];
+    layer.activationFunction.functionId = static_cast<std::uint16_t>(function);
+}
+
+template<class TWriterBase>
+inline void WriterImpl<TWriterBase>::setNeuralNetworkLayerActivationFunctionParameters(std::uint16_t netIndex,
+                                                                                       std::uint16_t layerIndex,
+                                                                                       const float* activationFunctionParameters,
+                                                                                       std::uint16_t count) {
+    ensureHasSize(dna.machineLearnedBehavior.neuralNetworks, netIndex + 1ul, memRes);
+    ensureHasSize(dna.machineLearnedBehavior.neuralNetworks[netIndex].layers, layerIndex + 1ul, memRes);
+    auto& layer = dna.machineLearnedBehavior.neuralNetworks[netIndex].layers[layerIndex];
+    layer.activationFunction.parameters.assign(activationFunctionParameters, activationFunctionParameters + count);
+}
+
+template<class TWriterBase>
+inline void WriterImpl<TWriterBase>::setNeuralNetworkLayerBiases(std::uint16_t netIndex,
+                                                                 std::uint16_t layerIndex,
+                                                                 const float* biases,
+                                                                 std::uint32_t count) {
+    ensureHasSize(dna.machineLearnedBehavior.neuralNetworks, netIndex + 1ul, memRes);
+    ensureHasSize(dna.machineLearnedBehavior.neuralNetworks[netIndex].layers, layerIndex + 1ul, memRes);
+    auto& layer = dna.machineLearnedBehavior.neuralNetworks[netIndex].layers[layerIndex];
+    layer.biases.assign(biases, biases + count);
+}
+
+template<class TWriterBase>
+inline void WriterImpl<TWriterBase>::setNeuralNetworkLayerWeights(std::uint16_t netIndex,
+                                                                  std::uint16_t layerIndex,
+                                                                  const float* weights,
+                                                                  std::uint32_t count) {
+    ensureHasSize(dna.machineLearnedBehavior.neuralNetworks, netIndex + 1ul, memRes);
+    ensureHasSize(dna.machineLearnedBehavior.neuralNetworks[netIndex].layers, layerIndex + 1ul, memRes);
+    auto& layer = dna.machineLearnedBehavior.neuralNetworks[netIndex].layers[layerIndex];
+    layer.weights.assign(weights, weights + count);
+}
+
 #ifdef _MSC_VER
     #pragma warning(pop)
 #endif
diff --git a/Engine/Plugins/Animation/RigLogic/Source/RigLogicLib/Private/dna/filters/JointFilter.cpp b/Engine/Plugins/Animation/RigLogic/Source/RigLogicLib/Private/dna/filters/JointFilter.cpp
--- a/Engine/Plugins/Animation/RigLogic/Source/RigLogicLib/Private/dna/filters/JointFilter.cpp
+++ b/Engine/Plugins/Animation/RigLogic/Source/RigLogicLib/Private/dna/filters/JointFilter.cpp
@@ -175,7 +175,11 @@ std::uint16_t JointFilter::remapped(std::uint16_t oldIndex) const {
 }
 
 std::uint16_t JointFilter::maxRemappedIndex() const {
-    return (remappedIndices.empty() ? static_cast<std::uint16_t>(0) : extd::maxOf(remappedIndices).second);
+    using IndexPair = std::pair<std::uint16_t, std::uint16_t>;
+    const auto compare = [](const IndexPair& lhs, const IndexPair& rhs) {
+            return lhs.second < rhs.second;
+        };
+    return (remappedIndices.empty() ? static_cast<std::uint16_t>(0) : extd::maxOf(remappedIndices, compare).second);
 }
 
 }  // namespace dna
diff --git a/Engine/Plugins/Animation/RigLogic/Source/RigLogicLib/Private/dna/filters/MeshFilter.cpp b/Engine/Plugins/Animation/RigLogic/Source/RigLogicLib/Private/dna/filters/MeshFilter.cpp
--- a/Engine/Plugins/Animation/RigLogic/Source/RigLogicLib/Private/dna/filters/MeshFilter.cpp
+++ b/Engine/Plugins/Animation/RigLogic/Source/RigLogicLib/Private/dna/filters/MeshFilter.cpp
@@ -37,6 +37,12 @@ void MeshFilter::apply(RawDefinition& dest) {
     dest.meshBlendShapeChannelMapping.updateFrom(remappedIndices);
 }
 
+void MeshFilter::apply(RawMachineLearnedBehavior& dest) {
+    // Delete region names belonging to meshes that are not referenced by the new subset of LODs
+    extd::filter(dest.neuralNetworkToMeshRegion.regionNames, extd::byPosition(passingIndices));
+    extd::filter(dest.neuralNetworkToMeshRegion.indices, extd::byPosition(passingIndices));
+}
+
 bool MeshFilter::passes(std::uint16_t index) const {
     return extd::contains(passingIndices, index);
 }
diff --git a/Engine/Plugins/Animation/RigLogic/Source/RigLogicLib/Private/dna/filters/MeshFilter.h b/Engine/Plugins/Animation/RigLogic/Source/RigLogicLib/Private/dna/filters/MeshFilter.h
--- a/Engine/Plugins/Animation/RigLogic/Source/RigLogicLib/Private/dna/filters/MeshFilter.h
+++ b/Engine/Plugins/Animation/RigLogic/Source/RigLogicLib/Private/dna/filters/MeshFilter.h
@@ -9,12 +9,14 @@
 namespace dna {
 
 struct RawDefinition;
+struct RawMachineLearnedBehavior;
 
 class MeshFilter {
     public:
         explicit MeshFilter(MemoryResource* memRes_);
         void configure(std::uint16_t meshCount, UnorderedSet<std::uint16_t> allowedMeshIndices);
         void apply(RawDefinition& dest);
+        void apply(RawMachineLearnedBehavior& dest);
         bool passes(std::uint16_t index) const;
 
     private:
diff --git a/Engine/Plugins/Animation/RigLogic/Source/RigLogicLib/Private/dna/stream/BinaryStreamReaderImpl.cpp b/Engine/Plugins/Animation/RigLogic/Source/RigLogicLib/Private/dna/stream/BinaryStreamReaderImpl.cpp
--- a/Engine/Plugins/Animation/RigLogic/Source/RigLogicLib/Private/dna/stream/BinaryStreamReaderImpl.cpp
+++ b/Engine/Plugins/Animation/RigLogic/Source/RigLogicLib/Private/dna/stream/BinaryStreamReaderImpl.cpp
@@ -34,28 +34,31 @@ BinaryStreamReader::~BinaryStreamReader() = default;
 
 BinaryStreamReader* BinaryStreamReader::create(BoundedIOStream* stream,
                                                DataLayer layer,
+                                               UnknownLayerPolicy policy,
                                                std::uint16_t maxLOD,
                                                MemoryResource* memRes) {
     PolyAllocator<BinaryStreamReaderImpl> alloc{memRes};
-    return alloc.newObject(stream, layer, maxLOD, LODLimits::min(), memRes);
+    return alloc.newObject(stream, layer, policy, maxLOD, LODLimits::min(), memRes);
 }
 
 BinaryStreamReader* BinaryStreamReader::create(BoundedIOStream* stream,
                                                DataLayer layer,
+                                               UnknownLayerPolicy policy,
                                                std::uint16_t maxLOD,
                                                std::uint16_t minLOD,
                                                MemoryResource* memRes) {
     PolyAllocator<BinaryStreamReaderImpl> alloc{memRes};
-    return alloc.newObject(stream, layer, maxLOD, minLOD, memRes);
+    return alloc.newObject(stream, layer, policy, maxLOD, minLOD, memRes);
 }
 
 BinaryStreamReader* BinaryStreamReader::create(BoundedIOStream* stream,
                                                DataLayer layer,
+                                               UnknownLayerPolicy policy,
                                                std::uint16_t* lods,
                                                std::uint16_t lodCount,
                                                MemoryResource* memRes) {
     PolyAllocator<BinaryStreamReaderImpl> alloc{memRes};
-    return alloc.newObject(stream, layer, ConstArrayView<std::uint16_t>{lods, lodCount}, memRes);
+    return alloc.newObject(stream, layer, policy, ConstArrayView<std::uint16_t>{lods, lodCount}, memRes);
 }
 
 void BinaryStreamReader::destroy(BinaryStreamReader* instance) {
@@ -67,10 +70,11 @@ void BinaryStreamReader::destroy(BinaryStreamReader* instance) {
 
 BinaryStreamReaderImpl::BinaryStreamReaderImpl(BoundedIOStream* stream_,
                                                DataLayer layer_,
+                                               UnknownLayerPolicy policy_,
                                                std::uint16_t maxLOD_,
                                                std::uint16_t minLOD_,
                                                MemoryResource* memRes_) :
-    BaseImpl{memRes_},
+    BaseImpl{policy_, UpgradeFormatPolicy::Disallowed, memRes_},
     ReaderImpl{memRes_},
     stream{stream_},
     archive{stream_, layer_, maxLOD_, minLOD_, memRes_},
@@ -79,9 +83,10 @@ BinaryStreamReaderImpl::BinaryStreamReaderImpl(BoundedIOStream* stream_,
 
 BinaryStreamReaderImpl::BinaryStreamReaderImpl(BoundedIOStream* stream_,
                                                DataLayer layer_,
+                                               UnknownLayerPolicy policy_,
                                                ConstArrayView<std::uint16_t> lods_,
                                                MemoryResource* memRes_) :
-    BaseImpl{memRes_},
+    BaseImpl{policy_, UpgradeFormatPolicy::Disallowed, memRes_},
     ReaderImpl{memRes_},
     stream{stream_},
     archive{stream_, layer_, lods_, memRes_},
@@ -94,14 +99,16 @@ bool BinaryStreamReaderImpl::isLODConstrained() const {
 
 void BinaryStreamReaderImpl::unload(DataLayer layer) {
     if ((layer == DataLayer::All) ||
-        (layer == DataLayer::AllWithoutBlendShapes) ||
         (layer == DataLayer::Descriptor)) {
-        dna = DNA{memRes};
+        dna = DNA{dna.layers.unknownPolicy, dna.layers.upgradePolicy, memRes};
+    } else if (layer == DataLayer::MachineLearnedBehavior) {
+        dna.unloadMachineLearnedBehavior();
     } else if ((layer == DataLayer::Geometry) || (layer == DataLayer::GeometryWithoutBlendShapes)) {
         dna.unloadGeometry();
     } else if (layer == DataLayer::Behavior) {
         dna.unloadBehavior();
     } else if (layer == DataLayer::Definition) {
+        dna.unloadMachineLearnedBehavior();
         dna.unloadGeometry();
         dna.unloadBehavior();
         dna.unloadDefinition();
@@ -123,16 +130,17 @@ void BinaryStreamReaderImpl::read() {
         return;
     }
 
+    if (!archive.isOk()) {
+        status.set(InvalidDataError);
+        return;
+    }
+
     if (!dna.signature.matches()) {
         status.set(SignatureMismatchError, dna.signature.value.expected.data(), dna.signature.value.got.data());
         return;
     }
-    if (!dna.version.matches()) {
-        status.set(VersionMismatchError,
-                   dna.version.generation.expected,
-                   dna.version.version.expected,
-                   dna.version.generation.got,
-                   dna.version.version.got);
+    if (!dna.version.supported()) {
+        status.set(VersionMismatchError, dna.version.generation, dna.version.version);
         return;
     }
 }
diff --git a/Engine/Plugins/Animation/RigLogic/Source/RigLogicLib/Private/dna/stream/BinaryStreamReaderImpl.h b/Engine/Plugins/Animation/RigLogic/Source/RigLogicLib/Private/dna/stream/BinaryStreamReaderImpl.h
--- a/Engine/Plugins/Animation/RigLogic/Source/RigLogicLib/Private/dna/stream/BinaryStreamReaderImpl.h
+++ b/Engine/Plugins/Animation/RigLogic/Source/RigLogicLib/Private/dna/stream/BinaryStreamReaderImpl.h
@@ -5,7 +5,7 @@
 #include "dna/BinaryStreamReader.h"
 #include "dna/ReaderImpl.h"
 #include "dna/TypeDefs.h"
-#include "dna/stream/FilteredInputArchive.h"
+#include "dna/stream/FilteredBinaryInputArchive.h"
 
 #include <status/Provider.h>
 
@@ -15,13 +15,12 @@ class BinaryStreamReaderImpl : public ReaderImpl<BinaryStreamReader> {
     public:
         BinaryStreamReaderImpl(BoundedIOStream* stream_,
                                DataLayer layer_,
+                               UnknownLayerPolicy policy_,
                                std::uint16_t maxLOD_,
                                std::uint16_t minLOD_,
                                MemoryResource* memRes_);
-        BinaryStreamReaderImpl(BoundedIOStream* stream_,
-                               DataLayer layer_,
-                               ConstArrayView<std::uint16_t> lods,
-                               MemoryResource* memRes_);
+        BinaryStreamReaderImpl(BoundedIOStream* stream_, DataLayer layer_, UnknownLayerPolicy policy_,
+                               ConstArrayView<std::uint16_t> lods, MemoryResource* memRes_);
 
         void unload(DataLayer layer) override;
         void read() override;
@@ -31,7 +30,7 @@ class BinaryStreamReaderImpl : public ReaderImpl<BinaryStreamReader> {
         static sc::StatusProvider status;
 
         BoundedIOStream* stream;
-        FilteredInputArchive archive;
+        FilteredBinaryInputArchive archive;
         bool lodConstrained;
 };
 
diff --git a/Engine/Plugins/Animation/RigLogic/Source/RigLogicLib/Private/dna/stream/BinaryStreamWriterImpl.cpp b/Engine/Plugins/Animation/RigLogic/Source/RigLogicLib/Private/dna/stream/BinaryStreamWriterImpl.cpp
--- a/Engine/Plugins/Animation/RigLogic/Source/RigLogicLib/Private/dna/stream/BinaryStreamWriterImpl.cpp
+++ b/Engine/Plugins/Animation/RigLogic/Source/RigLogicLib/Private/dna/stream/BinaryStreamWriterImpl.cpp
@@ -3,12 +3,21 @@
 #include "dna/stream/BinaryStreamWriterImpl.h"
 
 #include "dna/TypeDefs.h"
+#include "dna/stream/BinaryStreamReaderImpl.h"
+#include "dna/stream/JSONStreamReaderImpl.h"
 
+#ifdef _MSC_VER
+    #pragma warning(push)
+    #pragma warning(disable : 4365 4987)
+#endif
 #include <cassert>
 #include <cstddef>
 #include <cstring>
 #include <tuple>
 #include <utility>
+#ifdef _MSC_VER
+    #pragma warning(pop)
+#endif
 
 namespace dna {
 
@@ -27,7 +36,7 @@ void BinaryStreamWriter::destroy(BinaryStreamWriter* instance) {
 }
 
 BinaryStreamWriterImpl::BinaryStreamWriterImpl(BoundedIOStream* stream_, MemoryResource* memRes_) :
-    BaseImpl{memRes_},
+    BaseImpl{UnknownLayerPolicy::Preserve, UpgradeFormatPolicy::Allowed, memRes_},
     WriterImpl{memRes_},
     stream{stream_},
     archive{stream_} {
@@ -40,4 +49,20 @@ void BinaryStreamWriterImpl::write() {
     stream->close();
 }
 
+void BinaryStreamWriterImpl::setFrom(const BinaryStreamReader* source,
+                                     DataLayer layer,
+                                     UnknownLayerPolicy policy,
+                                     MemoryResource* memRes_) {
+    auto reader = const_cast<BinaryStreamReaderImpl*>(static_cast<const BinaryStreamReaderImpl*>(source));
+    reader->rawCopyInto(dna, layer, policy, memRes_);
+}
+
+void BinaryStreamWriterImpl::setFrom(const JSONStreamReader* source,
+                                     DataLayer layer,
+                                     UnknownLayerPolicy policy,
+                                     MemoryResource* memRes_) {
+    auto reader = const_cast<JSONStreamReaderImpl*>(static_cast<const JSONStreamReaderImpl*>(source));
+    reader->rawCopyInto(dna, layer, policy, memRes_);
+}
+
 }  // namespace dna
diff --git a/Engine/Plugins/Animation/RigLogic/Source/RigLogicLib/Private/dna/stream/BinaryStreamWriterImpl.h b/Engine/Plugins/Animation/RigLogic/Source/RigLogicLib/Private/dna/stream/BinaryStreamWriterImpl.h
--- a/Engine/Plugins/Animation/RigLogic/Source/RigLogicLib/Private/dna/stream/BinaryStreamWriterImpl.h
+++ b/Engine/Plugins/Animation/RigLogic/Source/RigLogicLib/Private/dna/stream/BinaryStreamWriterImpl.h
@@ -10,12 +10,21 @@
 
 namespace dna {
 
+class BinaryStreamReader;
+class JSONStreamReader;
+
 class BinaryStreamWriterImpl : public WriterImpl<BinaryStreamWriter> {
     public:
         BinaryStreamWriterImpl(BoundedIOStream* stream_, MemoryResource* memRes_);
 
         void write() override;
 
+        using BinaryStreamWriter::setFrom;
+        void setFrom(const BinaryStreamReader* source, DataLayer layer, UnknownLayerPolicy policy,
+                     MemoryResource* memRes_) override;
+        void setFrom(const JSONStreamReader* source, DataLayer layer, UnknownLayerPolicy policy,
+                     MemoryResource* memRes_) override;
+
     private:
         BoundedIOStream* stream;
         terse::BinaryOutputArchive<BoundedIOStream> archive;
diff --git a/Engine/Plugins/Animation/RigLogic/Source/RigLogicLib/Private/dna/stream/JSONStreamReaderImpl.cpp b/Engine/Plugins/Animation/RigLogic/Source/RigLogicLib/Private/dna/stream/JSONStreamReaderImpl.cpp
--- a/Engine/Plugins/Animation/RigLogic/Source/RigLogicLib/Private/dna/stream/JSONStreamReaderImpl.cpp
+++ b/Engine/Plugins/Animation/RigLogic/Source/RigLogicLib/Private/dna/stream/JSONStreamReaderImpl.cpp
@@ -40,7 +40,7 @@ void JSONStreamReader::destroy(JSONStreamReader* instance) {
 }
 
 JSONStreamReaderImpl::JSONStreamReaderImpl(BoundedIOStream* stream_, MemoryResource* memRes_) :
-    BaseImpl{memRes_},
+    BaseImpl{UnknownLayerPolicy::Preserve, UpgradeFormatPolicy::Disallowed, memRes_},
     ReaderImpl{memRes_},
     stream{stream_},
     archive{stream_} {
@@ -48,14 +48,16 @@ JSONStreamReaderImpl::JSONStreamReaderImpl(BoundedIOStream* stream_, MemoryResou
 
 void JSONStreamReaderImpl::unload(DataLayer layer) {
     if ((layer == DataLayer::All) ||
-        (layer == DataLayer::AllWithoutBlendShapes) ||
         (layer == DataLayer::Descriptor)) {
-        dna = DNA{memRes};
+        dna = DNA{dna.layers.unknownPolicy, dna.layers.upgradePolicy, memRes};
+    } else if (layer == DataLayer::MachineLearnedBehavior) {
+        dna.unloadMachineLearnedBehavior();
     } else if ((layer == DataLayer::Geometry) || (layer == DataLayer::GeometryWithoutBlendShapes)) {
         dna.unloadGeometry();
     } else if (layer == DataLayer::Behavior) {
         dna.unloadBehavior();
     } else if (layer == DataLayer::Definition) {
+        dna.unloadMachineLearnedBehavior();
         dna.unloadGeometry();
         dna.unloadBehavior();
         dna.unloadDefinition();
@@ -77,16 +79,17 @@ void JSONStreamReaderImpl::read() {
         return;
     }
 
+    if (!archive.isOk()) {
+        status.set(InvalidDataError);
+        return;
+    }
+
     if (!dna.signature.matches()) {
         status.set(SignatureMismatchError, dna.signature.value.expected.data(), dna.signature.value.got.data());
         return;
     }
-    if (!dna.version.matches()) {
-        status.set(VersionMismatchError,
-                   dna.version.generation.expected,
-                   dna.version.version.expected,
-                   dna.version.generation.got,
-                   dna.version.version.got);
+    if (!dna.version.supported()) {
+        status.set(VersionMismatchError, dna.version.generation, dna.version.version);
         return;
     }
 }
diff --git a/Engine/Plugins/Animation/RigLogic/Source/RigLogicLib/Private/dna/stream/JSONStreamWriterImpl.cpp b/Engine/Plugins/Animation/RigLogic/Source/RigLogicLib/Private/dna/stream/JSONStreamWriterImpl.cpp
--- a/Engine/Plugins/Animation/RigLogic/Source/RigLogicLib/Private/dna/stream/JSONStreamWriterImpl.cpp
+++ b/Engine/Plugins/Animation/RigLogic/Source/RigLogicLib/Private/dna/stream/JSONStreamWriterImpl.cpp
@@ -3,12 +3,21 @@
 #include "dna/stream/JSONStreamWriterImpl.h"
 
 #include "dna/TypeDefs.h"
+#include "dna/stream/BinaryStreamReaderImpl.h"
+#include "dna/stream/JSONStreamReaderImpl.h"
 
+#ifdef _MSC_VER
+    #pragma warning(push)
+    #pragma warning(disable : 4365 4987)
+#endif
 #include <cassert>
 #include <cstddef>
 #include <cstring>
 #include <tuple>
 #include <utility>
+#ifdef _MSC_VER
+    #pragma warning(pop)
+#endif
 
 namespace dna {
 
@@ -27,7 +36,7 @@ void JSONStreamWriter::destroy(JSONStreamWriter* instance) {
 }
 
 JSONStreamWriterImpl::JSONStreamWriterImpl(BoundedIOStream* stream_, std::uint32_t indentWidth, MemoryResource* memRes_) :
-    BaseImpl{memRes_},
+    BaseImpl{UnknownLayerPolicy::Preserve, UpgradeFormatPolicy::Allowed, memRes_},
     WriterImpl{memRes_},
     stream{stream_},
     archive{stream_, indentWidth} {
@@ -40,4 +49,20 @@ void JSONStreamWriterImpl::write() {
     stream->close();
 }
 
+void JSONStreamWriterImpl::setFrom(const BinaryStreamReader* source,
+                                   DataLayer layer,
+                                   UnknownLayerPolicy policy,
+                                   MemoryResource* memRes_) {
+    auto reader = const_cast<BinaryStreamReaderImpl*>(static_cast<const BinaryStreamReaderImpl*>(source));
+    reader->rawCopyInto(dna, layer, policy, memRes_);
+}
+
+void JSONStreamWriterImpl::setFrom(const JSONStreamReader* source,
+                                   DataLayer layer,
+                                   UnknownLayerPolicy policy,
+                                   MemoryResource* memRes_) {
+    auto reader = const_cast<JSONStreamReaderImpl*>(static_cast<const JSONStreamReaderImpl*>(source));
+    reader->rawCopyInto(dna, layer, policy, memRes_);
+}
+
 }  // namespace dna
diff --git a/Engine/Plugins/Animation/RigLogic/Source/RigLogicLib/Private/dna/stream/JSONStreamWriterImpl.h b/Engine/Plugins/Animation/RigLogic/Source/RigLogicLib/Private/dna/stream/JSONStreamWriterImpl.h
--- a/Engine/Plugins/Animation/RigLogic/Source/RigLogicLib/Private/dna/stream/JSONStreamWriterImpl.h
+++ b/Engine/Plugins/Animation/RigLogic/Source/RigLogicLib/Private/dna/stream/JSONStreamWriterImpl.h
@@ -10,12 +10,21 @@
 
 namespace dna {
 
+class BinaryStreamReader;
+class JSONStreamReader;
+
 class JSONStreamWriterImpl : public WriterImpl<JSONStreamWriter> {
     public:
         JSONStreamWriterImpl(BoundedIOStream* stream_, std::uint32_t indentWidth, MemoryResource* memRes_);
 
         void write() override;
 
+        using JSONStreamWriter::setFrom;
+        void setFrom(const BinaryStreamReader* source, DataLayer layer, UnknownLayerPolicy policy,
+                     MemoryResource* memRes_) override;
+        void setFrom(const JSONStreamReader* source, DataLayer layer, UnknownLayerPolicy policy,
+                     MemoryResource* memRes_) override;
+
     private:
         BoundedIOStream* stream;
         terse::JSONOutputArchive<BoundedIOStream> archive;
diff --git a/Engine/Plugins/Animation/RigLogic/Source/RigLogicLib/Private/dna/stream/StreamReader.cpp b/Engine/Plugins/Animation/RigLogic/Source/RigLogicLib/Private/dna/stream/StreamReader.cpp
--- a/Engine/Plugins/Animation/RigLogic/Source/RigLogicLib/Private/dna/stream/StreamReader.cpp
+++ b/Engine/Plugins/Animation/RigLogic/Source/RigLogicLib/Private/dna/stream/StreamReader.cpp
@@ -7,7 +7,7 @@ namespace dna {
 StreamReader::~StreamReader() = default;
 
 const sc::StatusCode StreamReader::SignatureMismatchError{200, "DNA signature mismatched, expected %.3s, got %.3s"};
-const sc::StatusCode StreamReader::VersionMismatchError{201, "DNA version mismatched, expected %hu.%hu, got %hu.%hu"};
+const sc::StatusCode StreamReader::VersionMismatchError{201, "DNA version mismatched, got %hu.%hu"};
 const sc::StatusCode StreamReader::InvalidDataError{202, "Invalid data in DNA"};
 
 }  // namespace dna
diff --git a/Engine/Plugins/Animation/RigLogic/Source/RigLogicLib/Private/dna/types/Limits.h b/Engine/Plugins/Animation/RigLogic/Source/RigLogicLib/Private/dna/types/Limits.h
--- a/Engine/Plugins/Animation/RigLogic/Source/RigLogicLib/Private/dna/types/Limits.h
+++ b/Engine/Plugins/Animation/RigLogic/Source/RigLogicLib/Private/dna/types/Limits.h
@@ -15,6 +15,11 @@ struct LODLimits {
         return 32u;
     }
 
+    static constexpr std::uint16_t count() {
+        static_assert(min() >= max(), "Min LOD value cannot be lower than Max LOD value.");
+        return static_cast<std::uint16_t>(min() - max() + 1);
+    }
+
 };
 
 }  // namespace dna
diff --git a/Engine/Plugins/Animation/RigLogic/Source/RigLogicLib/Private/dna/utils/Extd.h b/Engine/Plugins/Animation/RigLogic/Source/RigLogicLib/Private/dna/utils/Extd.h
--- a/Engine/Plugins/Animation/RigLogic/Source/RigLogicLib/Private/dna/utils/Extd.h
+++ b/Engine/Plugins/Animation/RigLogic/Source/RigLogicLib/Private/dna/utils/Extd.h
@@ -108,15 +108,20 @@ inline impl::LUTFilter<TLookUpTable, impl::LUTStrategy::ByPosition> byPosition(c
     return impl::LUTFilter<TLookUpTable, impl::LUTStrategy::ByPosition>{lookUpTable};
 }
 
+template<class TContainer, class TComparator>
+inline typename TContainer::value_type maxOf(const TContainer& container, TComparator&& comparator) {
+    using ValueType = typename TContainer::value_type;
+    const auto it = std::max_element(container.begin(), container.end(), comparator);
+    return (it == container.end() ? ValueType{} : *it);
+}
+
 template<class TContainer>
 inline typename TContainer::value_type maxOf(const TContainer& container) {
-    assert(!container.empty());
     using ValueType = typename TContainer::value_type;
     const auto compare = [](const ValueType& lhs, const ValueType& rhs) {
-            return lhs.second < rhs.second;
+            return lhs < rhs;
         };
-    const auto it = std::max_element(container.begin(), container.end(), compare);
-    return (it == container.end() ? ValueType{} : *it);
+    return maxOf(container, compare);
 }
 
 template<class TSource, class TDestination>
