Update RigLogic to v9.1.1 and GeneSplicer to v9.4.0

diff --git a/Engine/Plugins/Animation/RigLogic/Source/RigLogicLib/Private/pma/resources/AlignedMemoryResource.cpp b/Engine/Plugins/Animation/RigLogic/Source/RigLogicLib/Private/pma/resources/AlignedMemoryResource.cpp
--- a/Engine/Plugins/Animation/RigLogic/Source/RigLogicLib/Private/pma/resources/AlignedMemoryResource.cpp
+++ b/Engine/Plugins/Animation/RigLogic/Source/RigLogicLib/Private/pma/resources/AlignedMemoryResource.cpp
@@ -18,7 +18,7 @@
     #define ALIGNED_FREE(ptr) ::free(ptr)
 #elif defined(__cplusplus) && (__cplusplus >= 201703L) && (defined(_GLIBCXX_HAVE_ALIGNED_ALLOC) || \
     defined(_LIBCPP_HAS_ALIGNED_ALLOC) || defined(_LIBCPP_HAS_C11_FEATURES))
-    #define ALIGNED_ALLOC(ptr, alignment, size) ptr = std::aligned_alloc(alignment, size)
+    #define ALIGNED_ALLOC(ptr, alignment, size) ptr = std::aligned_alloc(alignment, (size % alignment == 0) ? size : (size + alignment - (size % alignment)))
     #define ALIGNED_FREE(ptr) std::free(ptr)
 #elif ((defined(_POSIX_VERSION) && (_POSIX_VERSION >= 200112L)) || defined(__linux__) || defined(__APPLE__))
     #define ALIGNED_ALLOC(ptr, alignment, size) if (::posix_memalign(&(ptr), alignment, size)) \
diff --git a/Engine/Plugins/Animation/RigLogic/Source/RigLogicLib/Private/riglogic/TypeDefs.h b/Engine/Plugins/Animation/RigLogic/Source/RigLogicLib/Private/riglogic/TypeDefs.h
--- a/Engine/Plugins/Animation/RigLogic/Source/RigLogicLib/Private/riglogic/TypeDefs.h
+++ b/Engine/Plugins/Animation/RigLogic/Source/RigLogicLib/Private/riglogic/TypeDefs.h
@@ -8,6 +8,10 @@
 #include <pma/TypeDefs.h>
 #include <pma/resources/AlignedMemoryResource.h>
 #include <pma/resources/DefaultMemoryResource.h>
+#include <pma/utils/ManagedInstance.h>
+#include <terse/archives/binary/InputArchive.h>
+#include <terse/archives/binary/OutputArchive.h>
+#include <terse/utils/VirtualSerializerProxy.h>
 
 #include <cstddef>
 
diff --git a/Engine/Plugins/Animation/RigLogic/Source/RigLogicLib/Private/riglogic/animatedmaps/AnimatedMaps.cpp b/Engine/Plugins/Animation/RigLogic/Source/RigLogicLib/Private/riglogic/animatedmaps/AnimatedMaps.cpp
--- a/Engine/Plugins/Animation/RigLogic/Source/RigLogicLib/Private/riglogic/animatedmaps/AnimatedMaps.cpp
+++ b/Engine/Plugins/Animation/RigLogic/Source/RigLogicLib/Private/riglogic/animatedmaps/AnimatedMaps.cpp
@@ -2,22 +2,8 @@
 
 #include "riglogic/animatedmaps/AnimatedMaps.h"
 
-#include "riglogic/TypeDefs.h"
-#include "riglogic/conditionaltable/ConditionalTable.h"
-
-#include <cassert>
-#include <utility>
-
 namespace rl4 {
 
-AnimatedMaps::AnimatedMaps(Vector<std::uint16_t>&& lods_, ConditionalTable&& conditionals_) :
-    lods{std::move(lods_)},
-    conditionals{std::move(conditionals_)} {
-}
-
-void AnimatedMaps::calculate(ConstArrayView<float> inputs, ArrayView<float> outputs, std::uint16_t lod) const {
-    assert(lod < lods.size());
-    conditionals.calculate(inputs.data(), outputs.data(), lods[lod]);
-}
+AnimatedMaps::~AnimatedMaps() = default;
 
 }  // namespace rl4
diff --git a/Engine/Plugins/Animation/RigLogic/Source/RigLogicLib/Private/riglogic/animatedmaps/AnimatedMaps.h b/Engine/Plugins/Animation/RigLogic/Source/RigLogicLib/Private/riglogic/animatedmaps/AnimatedMaps.h
--- a/Engine/Plugins/Animation/RigLogic/Source/RigLogicLib/Private/riglogic/animatedmaps/AnimatedMaps.h
+++ b/Engine/Plugins/Animation/RigLogic/Source/RigLogicLib/Private/riglogic/animatedmaps/AnimatedMaps.h
@@ -3,27 +3,27 @@
 #pragma once
 
 #include "riglogic/TypeDefs.h"
-#include "riglogic/conditionaltable/ConditionalTable.h"
-#include "riglogic/types/Aliases.h"
+#include "riglogic/animatedmaps/AnimatedMapsOutputInstance.h"
 
-#include <cstddef>
+#include <cstdint>
 
 namespace rl4 {
 
+class ControlsInputInstance;
+
 class AnimatedMaps {
     public:
-        AnimatedMaps(Vector<std::uint16_t>&& lods_, ConditionalTable&& conditionals_);
-
-        void calculate(ConstArrayView<float> inputs, ArrayView<float> outputs, std::uint16_t lod) const;
+        using Pointer = UniqueInstance<AnimatedMaps>::PointerType;
 
-        template<class Archive>
-        void serialize(Archive& archive) {
-            archive(lods, conditionals);
-        }
+    protected:
+        virtual ~AnimatedMaps();
 
-    private:
-        Vector<std::uint16_t> lods;
-        ConditionalTable conditionals;
+    public:
+        virtual AnimatedMapsOutputInstance::Pointer createInstance(MemoryResource* memRes) const = 0;
+        virtual void calculate(const ControlsInputInstance* inputs, AnimatedMapsOutputInstance* outputs,
+                               std::uint16_t lod) const = 0;
+        virtual void load(terse::BinaryInputArchive<BoundedIOStream>& archive) = 0;
+        virtual void save(terse::BinaryOutputArchive<BoundedIOStream>& archive) = 0;
 
 };
 
diff --git a/Engine/Plugins/Animation/RigLogic/Source/RigLogicLib/Private/riglogic/animatedmaps/AnimatedMapsFactory.cpp b/Engine/Plugins/Animation/RigLogic/Source/RigLogicLib/Private/riglogic/animatedmaps/AnimatedMapsFactory.cpp
--- a/Engine/Plugins/Animation/RigLogic/Source/RigLogicLib/Private/riglogic/animatedmaps/AnimatedMapsFactory.cpp
+++ b/Engine/Plugins/Animation/RigLogic/Source/RigLogicLib/Private/riglogic/animatedmaps/AnimatedMapsFactory.cpp
@@ -2,21 +2,47 @@
 
 #include "riglogic/animatedmaps/AnimatedMapsFactory.h"
 
+#include "riglogic/TypeDefs.h"
+#include "riglogic/animatedmaps/AnimatedMapsImpl.h"
+#include "riglogic/animatedmaps/AnimatedMapsImplOutputInstance.h"
+#include "riglogic/animatedmaps/AnimatedMapsNull.h"
+#include "riglogic/animatedmaps/AnimatedMapsOutputInstance.h"
 #include "riglogic/conditionaltable/ConditionalTable.h"
+#include "riglogic/riglogic/Configuration.h"
+#include "riglogic/riglogic/RigMetrics.h"
 #include "riglogic/utils/Extd.h"
 
-#include <dna/layers/BehaviorReader.h>
-
 #include <cstddef>
+#include <cstdint>
 
 namespace rl4 {
 
-AnimatedMapsFactory::AnimatedMapsPtr AnimatedMapsFactory::create(MemoryResource* memRes) {
-    return ManagedAnimatedMaps::with(memRes).create(Vector<std::uint16_t>{memRes},
-                                                    ConditionalTable{memRes});
+static AnimatedMapsOutputInstance::Factory createAnimatedMapsOutputInstanceFactory(const Configuration&  /*unused*/,
+                                                                                   std::uint16_t animatedMapCount) {
+    return [ = ](MemoryResource* memRes) {
+               return UniqueInstance<AnimatedMapsImplOutputInstance, AnimatedMapsOutputInstance>::with(memRes).create(
+                   animatedMapCount,
+                   memRes);
+    };
+}
+
+AnimatedMaps::Pointer AnimatedMapsFactory::create(const Configuration& config, const RigMetrics& metrics,
+                                                  MemoryResource* memRes) {
+    if (!config.loadAnimatedMaps) {
+        return UniqueInstance<AnimatedMapsNull, AnimatedMaps>::with(memRes).create();
+    }
+    auto instanceFactory = createAnimatedMapsOutputInstanceFactory(config, metrics.animatedMapCount);
+    auto moduleFactory = UniqueInstance<AnimatedMapsImpl, AnimatedMaps>::with(memRes);
+    return moduleFactory.create(Vector<std::uint16_t>{memRes}, ConditionalTable{memRes}, instanceFactory);
 }
 
-AnimatedMapsFactory::AnimatedMapsPtr AnimatedMapsFactory::create(const dna::BehaviorReader* reader, MemoryResource* memRes) {
+AnimatedMaps::Pointer AnimatedMapsFactory::create(const Configuration& config,
+                                                  const dna::BehaviorReader* reader,
+                                                  MemoryResource* memRes) {
+    if (!config.loadAnimatedMaps) {
+        return UniqueInstance<AnimatedMapsNull, AnimatedMaps>::with(memRes).create();
+    }
+
     Vector<std::uint16_t> lods{memRes};
     Vector<std::uint16_t> inputIndices{memRes};
     Vector<std::uint16_t> outputIndices{memRes};
@@ -52,7 +78,10 @@ AnimatedMapsFactory::AnimatedMapsPtr AnimatedMapsFactory::create(const dna::Beha
                                   inputCount,
                                   outputCount,
                                   memRes};
-    return ManagedAnimatedMaps::with(memRes).create(std::move(lods), std::move(conditionals));
+
+    auto instanceFactory = createAnimatedMapsOutputInstanceFactory(config, outputCount);
+    auto moduleFactory = UniqueInstance<AnimatedMapsImpl, AnimatedMaps>::with(memRes);
+    return moduleFactory.create(std::move(lods), std::move(conditionals), instanceFactory);
 }
 
 }  // namespace rl4
diff --git a/Engine/Plugins/Animation/RigLogic/Source/RigLogicLib/Private/riglogic/animatedmaps/AnimatedMapsFactory.h b/Engine/Plugins/Animation/RigLogic/Source/RigLogicLib/Private/riglogic/animatedmaps/AnimatedMapsFactory.h
--- a/Engine/Plugins/Animation/RigLogic/Source/RigLogicLib/Private/riglogic/animatedmaps/AnimatedMapsFactory.h
+++ b/Engine/Plugins/Animation/RigLogic/Source/RigLogicLib/Private/riglogic/animatedmaps/AnimatedMapsFactory.h
@@ -2,24 +2,17 @@
 
 #pragma once
 
+#include "riglogic/TypeDefs.h"
 #include "riglogic/animatedmaps/AnimatedMaps.h"
 
-#include <pma/utils/ManagedInstance.h>
-
-namespace dna {
-
-class BehaviorReader;
-
-}  // namespace dna
-
 namespace rl4 {
 
-struct AnimatedMapsFactory {
-    using ManagedAnimatedMaps = pma::UniqueInstance<AnimatedMaps>;
-    using AnimatedMapsPtr = typename ManagedAnimatedMaps::PointerType;
+struct Configuration;
+struct RigMetrics;
 
-    static AnimatedMapsPtr create(MemoryResource* memRes);
-    static AnimatedMapsPtr create(const dna::BehaviorReader* reader, MemoryResource* memRes);
+struct AnimatedMapsFactory {
+    static AnimatedMaps::Pointer create(const Configuration& config, const dna::BehaviorReader* reader, MemoryResource* memRes);
+    static AnimatedMaps::Pointer create(const Configuration& config, const RigMetrics& metrics, MemoryResource* memRes);
 };
 
 }  // namespace rl4
diff --git a/Engine/Plugins/Animation/RigLogic/Source/RigLogicLib/Private/riglogic/animatedmaps/AnimatedMapsNull.cpp b/Engine/Plugins/Animation/RigLogic/Source/RigLogicLib/Private/riglogic/animatedmaps/AnimatedMapsNull.cpp
new file mode 100644
--- /dev/null
+++ b/Engine/Plugins/Animation/RigLogic/Source/RigLogicLib/Private/riglogic/animatedmaps/AnimatedMapsNull.cpp
@@ -0,0 +1,23 @@
+// Copyright Epic Games, Inc. All Rights Reserved.
+
+#include "riglogic/animatedmaps/AnimatedMapsNull.h"
+
+#include "riglogic/animatedmaps/AnimatedMapsNullOutputInstance.h"
+
+namespace rl4 {
+
+AnimatedMapsOutputInstance::Pointer AnimatedMapsNull::createInstance(MemoryResource* instanceMemRes) const {
+    return UniqueInstance<AnimatedMapsNullOutputInstance, AnimatedMapsOutputInstance>::with(instanceMemRes).create();
+}
+
+void AnimatedMapsNull::calculate(const ControlsInputInstance*  /*unused*/, AnimatedMapsOutputInstance*  /*unused*/,
+                                 std::uint16_t  /*unused*/) const {
+}
+
+void AnimatedMapsNull::load(terse::BinaryInputArchive<BoundedIOStream>&  /*unused*/) {
+}
+
+void AnimatedMapsNull::save(terse::BinaryOutputArchive<BoundedIOStream>&  /*unused*/) {
+}
+
+}  // namespace rl4
diff --git a/Engine/Plugins/Animation/RigLogic/Source/RigLogicLib/Private/riglogic/animatedmaps/AnimatedMapsNull.h b/Engine/Plugins/Animation/RigLogic/Source/RigLogicLib/Private/riglogic/animatedmaps/AnimatedMapsNull.h
new file mode 100644
--- /dev/null
+++ b/Engine/Plugins/Animation/RigLogic/Source/RigLogicLib/Private/riglogic/animatedmaps/AnimatedMapsNull.h
@@ -0,0 +1,25 @@
+// Copyright Epic Games, Inc. All Rights Reserved.
+
+#pragma once
+
+#include "riglogic/TypeDefs.h"
+#include "riglogic/animatedmaps/AnimatedMaps.h"
+#include "riglogic/animatedmaps/AnimatedMapsOutputInstance.h"
+#include "riglogic/conditionaltable/ConditionalTable.h"
+
+#include <cstdint>
+
+namespace rl4 {
+
+class ControlsInputInstance;
+
+class AnimatedMapsNull : public AnimatedMaps {
+    public:
+        AnimatedMapsOutputInstance::Pointer createInstance(MemoryResource* instanceMemRes) const override;
+        void calculate(const ControlsInputInstance*  /*unused*/, AnimatedMapsOutputInstance*  /*unused*/,
+                       std::uint16_t  /*unused*/) const override;
+        void load(terse::BinaryInputArchive<BoundedIOStream>&  /*unused*/) override;
+        void save(terse::BinaryOutputArchive<BoundedIOStream>&  /*unused*/) override;
+};
+
+}  // namespace rl4
diff --git a/Engine/Plugins/Animation/RigLogic/Source/RigLogicLib/Private/riglogic/blendshapes/BlendShapes.cpp b/Engine/Plugins/Animation/RigLogic/Source/RigLogicLib/Private/riglogic/blendshapes/BlendShapes.cpp
--- a/Engine/Plugins/Animation/RigLogic/Source/RigLogicLib/Private/riglogic/blendshapes/BlendShapes.cpp
+++ b/Engine/Plugins/Animation/RigLogic/Source/RigLogicLib/Private/riglogic/blendshapes/BlendShapes.cpp
@@ -2,27 +2,8 @@
 
 #include "riglogic/blendshapes/BlendShapes.h"
 
-#include "riglogic/TypeDefs.h"
-
-#include <cassert>
-#include <cstddef>
-
 namespace rl4 {
 
-BlendShapes::BlendShapes(Vector<std::uint16_t>&& lods_,
-                         Vector<std::uint16_t>&& inputIndices_,
-                         Vector<std::uint16_t>&& outputIndices_) :
-    lods{std::move(lods_)},
-    inputIndices{std::move(inputIndices_)},
-    outputIndices{std::move(outputIndices_)} {
-}
-
-void BlendShapes::calculate(ConstArrayView<float> inputs, ArrayView<float> outputs, std::uint16_t lod) const {
-    assert(lod < lods.size());
-    std::fill(outputs.begin(), outputs.end(), 0.0f);
-    for (std::uint16_t i = 0u; i < lods[lod]; ++i) {
-        outputs[outputIndices[i]] = inputs[inputIndices[i]];
-    }
-}
+BlendShapes::~BlendShapes() = default;
 
 }  // namespace rl4
diff --git a/Engine/Plugins/Animation/RigLogic/Source/RigLogicLib/Private/riglogic/blendshapes/BlendShapes.h b/Engine/Plugins/Animation/RigLogic/Source/RigLogicLib/Private/riglogic/blendshapes/BlendShapes.h
--- a/Engine/Plugins/Animation/RigLogic/Source/RigLogicLib/Private/riglogic/blendshapes/BlendShapes.h
+++ b/Engine/Plugins/Animation/RigLogic/Source/RigLogicLib/Private/riglogic/blendshapes/BlendShapes.h
@@ -3,27 +3,27 @@
 #pragma once
 
 #include "riglogic/TypeDefs.h"
-#include "riglogic/types/Aliases.h"
+#include "riglogic/blendshapes/BlendShapesOutputInstance.h"
 
-#include <cstddef>
+#include <cstdint>
 
 namespace rl4 {
 
+class ControlsInputInstance;
+
 class BlendShapes {
     public:
-        BlendShapes(Vector<std::uint16_t>&& lods_, Vector<std::uint16_t>&& inputIndices_, Vector<std::uint16_t>&& outputIndices_);
-
-        void calculate(ConstArrayView<float> inputs, ArrayView<float> outputs, std::uint16_t lod) const;
+        using Pointer = UniqueInstance<BlendShapes>::PointerType;
 
-        template<class Archive>
-        void serialize(Archive& archive) {
-            archive(lods, inputIndices, outputIndices);
-        }
+    protected:
+        virtual ~BlendShapes();
 
-    private:
-        Vector<std::uint16_t> lods;
-        Vector<std::uint16_t> inputIndices;
-        Vector<std::uint16_t> outputIndices;
+    public:
+        virtual BlendShapesOutputInstance::Pointer createInstance(MemoryResource* instanceMemRes) const = 0;
+        virtual void calculate(const ControlsInputInstance* inputs, BlendShapesOutputInstance* outputs,
+                               std::uint16_t lod) const = 0;
+        virtual void load(terse::BinaryInputArchive<BoundedIOStream>& archive) = 0;
+        virtual void save(terse::BinaryOutputArchive<BoundedIOStream>& archive) = 0;
 
 };
 
diff --git a/Engine/Plugins/Animation/RigLogic/Source/RigLogicLib/Private/riglogic/blendshapes/BlendShapesFactory.cpp b/Engine/Plugins/Animation/RigLogic/Source/RigLogicLib/Private/riglogic/blendshapes/BlendShapesFactory.cpp
--- a/Engine/Plugins/Animation/RigLogic/Source/RigLogicLib/Private/riglogic/blendshapes/BlendShapesFactory.cpp
+++ b/Engine/Plugins/Animation/RigLogic/Source/RigLogicLib/Private/riglogic/blendshapes/BlendShapesFactory.cpp
@@ -2,23 +2,55 @@
 
 #include "riglogic/blendshapes/BlendShapesFactory.h"
 
+#include "riglogic/TypeDefs.h"
 #include "riglogic/utils/Extd.h"
+#include "riglogic/blendshapes/BlendShapesImpl.h"
+#include "riglogic/blendshapes/BlendShapesImplOutputInstance.h"
+#include "riglogic/blendshapes/BlendShapesNull.h"
+#include "riglogic/blendshapes/BlendShapesOutputInstance.h"
+#include "riglogic/riglogic/Configuration.h"
+#include "riglogic/riglogic/RigMetrics.h"
 
-#include <dna/layers/BehaviorReader.h>
-
+#ifdef _MSC_VER
+    #pragma warning(push)
+    #pragma warning(disable : 4365 4987)
+#endif
 #include <cstdint>
 #include <utility>
-#include <cstring>
+#ifdef _MSC_VER
+    #pragma warning(pop)
+#endif
 
 namespace rl4 {
 
-BlendShapesFactory::BlendShapesPtr BlendShapesFactory::create(MemoryResource* memRes) {
-    return ManagedBlendShapes::with(memRes).create(Vector<std::uint16_t>{memRes},
-                                                   Vector<std::uint16_t>{memRes},
-                                                   Vector<std::uint16_t>{memRes});
+static BlendShapesOutputInstance::Factory createBlendShapesOutputInstanceFactory(const Configuration&  /*unused*/,
+                                                                                 std::uint16_t blendShapeCount) {
+    return [ = ](MemoryResource* memRes) {
+               return UniqueInstance<BlendShapesImplOutputInstance, BlendShapesOutputInstance>::with(memRes).create(
+                   blendShapeCount,
+                   memRes);
+    };
+}
+
+BlendShapes::Pointer BlendShapesFactory::create(const Configuration& config, const RigMetrics& metrics, MemoryResource* memRes) {
+    if (!config.loadBlendShapes) {
+        return UniqueInstance<BlendShapesNull, BlendShapes>::with(memRes).create();
+    }
+    auto instanceFactory = createBlendShapesOutputInstanceFactory(config, metrics.blendShapeCount);
+    auto moduleFactory = UniqueInstance<BlendShapesImpl, BlendShapes>::with(memRes);
+    return moduleFactory.create(Vector<std::uint16_t>{memRes},
+                                Vector<std::uint16_t>{memRes},
+                                Vector<std::uint16_t>{memRes},
+                                instanceFactory);
 }
 
-BlendShapesFactory::BlendShapesPtr BlendShapesFactory::create(const dna::BehaviorReader* reader, MemoryResource* memRes) {
+BlendShapes::Pointer BlendShapesFactory::create(const Configuration& config,
+                                                const dna::BehaviorReader* reader,
+                                                MemoryResource* memRes) {
+    if (!config.loadBlendShapes) {
+        return UniqueInstance<BlendShapesNull, BlendShapes>::with(memRes).create();
+    }
+
     Vector<std::uint16_t> lods{memRes};
     Vector<std::uint16_t> inputIndices{memRes};
     Vector<std::uint16_t> outputIndices{memRes};
@@ -27,9 +59,9 @@ BlendShapesFactory::BlendShapesPtr BlendShapesFactory::create(const dna::Behavio
     extd::copy(reader->getBlendShapeChannelInputIndices(), inputIndices);
     extd::copy(reader->getBlendShapeChannelOutputIndices(), outputIndices);
 
-    return ManagedBlendShapes::with(memRes).create(std::move(lods),
-                                                   std::move(inputIndices),
-                                                   std::move(outputIndices));
+    auto instanceFactory = createBlendShapesOutputInstanceFactory(config, reader->getBlendShapeChannelCount());
+    auto moduleFactory = UniqueInstance<BlendShapesImpl, BlendShapes>::with(memRes);
+    return moduleFactory.create(std::move(lods), std::move(inputIndices), std::move(outputIndices), instanceFactory);
 }
 
 }  // namespace rl4
diff --git a/Engine/Plugins/Animation/RigLogic/Source/RigLogicLib/Private/riglogic/blendshapes/BlendShapesFactory.h b/Engine/Plugins/Animation/RigLogic/Source/RigLogicLib/Private/riglogic/blendshapes/BlendShapesFactory.h
--- a/Engine/Plugins/Animation/RigLogic/Source/RigLogicLib/Private/riglogic/blendshapes/BlendShapesFactory.h
+++ b/Engine/Plugins/Animation/RigLogic/Source/RigLogicLib/Private/riglogic/blendshapes/BlendShapesFactory.h
@@ -2,24 +2,17 @@
 
 #pragma once
 
+#include "riglogic/TypeDefs.h"
 #include "riglogic/blendshapes/BlendShapes.h"
 
-#include <pma/utils/ManagedInstance.h>
-
-namespace dna {
-
-class BehaviorReader;
-
-}  // namespace dna
-
 namespace rl4 {
 
-struct BlendShapesFactory {
-    using ManagedBlendShapes = pma::UniqueInstance<BlendShapes>;
-    using BlendShapesPtr = typename ManagedBlendShapes::PointerType;
+struct Configuration;
+struct RigMetrics;
 
-    static BlendShapesPtr create(MemoryResource* memRes);
-    static BlendShapesPtr create(const dna::BehaviorReader* reader, MemoryResource* memRes);
+struct BlendShapesFactory {
+    static BlendShapes::Pointer create(const Configuration& config, const dna::BehaviorReader* reader, MemoryResource* memRes);
+    static BlendShapes::Pointer create(const Configuration& config, const RigMetrics& metrics, MemoryResource* memRes);
 };
 
 }  // namespace rl4
diff --git a/Engine/Plugins/Animation/RigLogic/Source/RigLogicLib/Private/riglogic/blendshapes/BlendShapesNull.cpp b/Engine/Plugins/Animation/RigLogic/Source/RigLogicLib/Private/riglogic/blendshapes/BlendShapesNull.cpp
new file mode 100644
--- /dev/null
+++ b/Engine/Plugins/Animation/RigLogic/Source/RigLogicLib/Private/riglogic/blendshapes/BlendShapesNull.cpp
@@ -0,0 +1,23 @@
+// Copyright Epic Games, Inc. All Rights Reserved.
+
+#include "riglogic/blendshapes/BlendShapesNull.h"
+
+#include "riglogic/blendshapes/BlendShapesNullOutputInstance.h"
+
+namespace rl4 {
+
+BlendShapesOutputInstance::Pointer BlendShapesNull::createInstance(MemoryResource* instanceMemRes) const {
+    return UniqueInstance<BlendShapesNullOutputInstance, BlendShapesOutputInstance>::with(instanceMemRes).create();
+}
+
+void BlendShapesNull::calculate(const ControlsInputInstance*  /*unused*/, BlendShapesOutputInstance*  /*unused*/,
+                                std::uint16_t  /*unused*/) const {
+}
+
+void BlendShapesNull::load(terse::BinaryInputArchive<BoundedIOStream>&  /*unused*/) {
+}
+
+void BlendShapesNull::save(terse::BinaryOutputArchive<BoundedIOStream>&  /*unused*/) {
+}
+
+}  // namespace rl4
diff --git a/Engine/Plugins/Animation/RigLogic/Source/RigLogicLib/Private/riglogic/blendshapes/BlendShapesNull.h b/Engine/Plugins/Animation/RigLogic/Source/RigLogicLib/Private/riglogic/blendshapes/BlendShapesNull.h
new file mode 100644
--- /dev/null
+++ b/Engine/Plugins/Animation/RigLogic/Source/RigLogicLib/Private/riglogic/blendshapes/BlendShapesNull.h
@@ -0,0 +1,25 @@
+// Copyright Epic Games, Inc. All Rights Reserved.
+
+#pragma once
+
+#include "riglogic/TypeDefs.h"
+#include "riglogic/blendshapes/BlendShapes.h"
+#include "riglogic/blendshapes/BlendShapesOutputInstance.h"
+
+#include <cstdint>
+
+namespace rl4 {
+
+class ControlsInputInstance;
+
+class BlendShapesNull : public BlendShapes {
+    public:
+        BlendShapesOutputInstance::Pointer createInstance(MemoryResource* instanceMemRes) const override;
+        void calculate(const ControlsInputInstance*  /*unused*/, BlendShapesOutputInstance*  /*unused*/,
+                       std::uint16_t  /*unused*/) const override;
+        void load(terse::BinaryInputArchive<BoundedIOStream>&  /*unused*/) override;
+        void save(terse::BinaryOutputArchive<BoundedIOStream>&  /*unused*/) override;
+
+};
+
+}  // namespace rl4
diff --git a/Engine/Plugins/Animation/RigLogic/Source/RigLogicLib/Private/riglogic/conditionaltable/ConditionalTable.cpp b/Engine/Plugins/Animation/RigLogic/Source/RigLogicLib/Private/riglogic/conditionaltable/ConditionalTable.cpp
--- a/Engine/Plugins/Animation/RigLogic/Source/RigLogicLib/Private/riglogic/conditionaltable/ConditionalTable.cpp
+++ b/Engine/Plugins/Animation/RigLogic/Source/RigLogicLib/Private/riglogic/conditionaltable/ConditionalTable.cpp
@@ -24,8 +24,8 @@ namespace {
 const float clampMin = 0.0f;
 const float clampMax = 1.0f;
 
-Vector<std::uint16_t> buildIntervalSkipMap(const Vector<std::uint16_t>& inputIndices,
-                                           const Vector<std::uint16_t>& outputIndices,
+Vector<std::uint16_t> buildIntervalSkipMap(ConstArrayView<std::uint16_t> inputIndices,
+                                           ConstArrayView<std::uint16_t> outputIndices,
                                            MemoryResource* memRes) {
     assert(inputIndices.size() == outputIndices.size());
     Vector<std::uint16_t> intervalsRemaining{inputIndices.size(), {}, memRes};
@@ -45,9 +45,28 @@ Vector<std::uint16_t> buildIntervalSkipMap(const Vector<std::uint16_t>& inputInd
     return intervalsRemaining;
 }
 
+Vector<RangeMap> buildRangeMap(ConstArrayView<std::uint16_t> inputIndices,
+                               ConstArrayView<float> fromValues,
+                               ConstArrayView<float> toValues,
+                               MemoryResource* memRes) {
+    const std::size_t maxIndex = extd::maxOf(inputIndices);
+    const std::size_t rangeMapCount = (maxIndex + 1ul);
+    Vector<RangeMap> rangeMaps{rangeMapCount, RangeMap{memRes}, memRes};
+
+    for (std::size_t i = {}; i < inputIndices.size(); ++i) {
+        const std::size_t ri = inputIndices[i];
+        auto& map = rangeMaps[ri];
+        auto range = map.addRange(fromValues[i], toValues[i]);
+        range->rows.push_back(static_cast<std::uint16_t>(i));
+    }
+
+    return rangeMaps;
+}
+
 }  // namespace
 
 ConditionalTable::ConditionalTable(MemoryResource* memRes) :
+    rangeMaps{memRes},
     intervalsRemaining{memRes},
     inputIndices{memRes},
     outputIndices{memRes},
@@ -68,7 +87,13 @@ ConditionalTable::ConditionalTable(Vector<std::uint16_t>&& inputIndices_,
                                    std::uint16_t inputCount_,
                                    std::uint16_t outputCount_,
                                    MemoryResource* memRes) :
-    intervalsRemaining{buildIntervalSkipMap(inputIndices_, outputIndices_, memRes)},
+    rangeMaps{buildRangeMap(ConstArrayView<std::uint16_t> {inputIndices_},
+                            ConstArrayView<float> {fromValues_},
+                            ConstArrayView<float> {toValues_},
+                            memRes)},
+    intervalsRemaining{buildIntervalSkipMap(ConstArrayView<std::uint16_t> {inputIndices_},
+                                            ConstArrayView<std::uint16_t> {outputIndices_},
+                                            memRes)},
     inputIndices{std::move(inputIndices_)},
     outputIndices{std::move(outputIndices_)},
     fromValues{std::move(fromValues_)},
@@ -79,6 +104,10 @@ ConditionalTable::ConditionalTable(Vector<std::uint16_t>&& inputIndices_,
     outputCount{outputCount_} {
 }
 
+std::uint16_t ConditionalTable::getRowCount() const {
+    return static_cast<std::uint16_t>(inputIndices.size());
+}
+
 std::uint16_t ConditionalTable::getInputCount() const {
     return inputCount;
 }
@@ -87,19 +116,19 @@ std::uint16_t ConditionalTable::getOutputCount() const {
     return outputCount;
 }
 
-void ConditionalTable::calculate(const float* inputs, float* outputs, std::uint16_t chunkSize) const {
+void ConditionalTable::calculateForward(const float* inputs, float* outputs, std::uint16_t rowCount) const {
     std::fill_n(outputs, outputCount, 0.0f);
 
-    for (std::uint16_t i = {}; i < chunkSize; ++i) {
-        const float inValue = inputs[inputIndices[i]];
-        const float from = fromValues[i];
-        const float to = toValues[i];
+    for (std::uint16_t row = {}; row < rowCount; ++row) {
+        const float inValue = inputs[inputIndices[row]];
+        const float from = fromValues[row];
+        const float to = toValues[row];
         if ((from <= inValue) && (inValue <= to)) {
-            const std::uint16_t outIndex = outputIndices[i];
-            const float slope = slopeValues[i];
-            const float cut = cutValues[i];
+            const std::uint16_t outIndex = outputIndices[row];
+            const float slope = slopeValues[row];
+            const float cut = cutValues[row];
             outputs[outIndex] += (slope * inValue + cut);
-            i = static_cast<std::uint16_t>(i + intervalsRemaining[i]);
+            row = static_cast<std::uint16_t>(row + intervalsRemaining[row]);
         }
     }
 
@@ -108,8 +137,83 @@ void ConditionalTable::calculate(const float* inputs, float* outputs, std::uint1
     }
 }
 
-void ConditionalTable::calculate(const float* inputs, float* outputs) const {
-    calculate(inputs, outputs, static_cast<std::uint16_t>(outputIndices.size()));
+void ConditionalTable::calculateForward(const float* inputs, float* outputs) const {
+    calculateForward(inputs, outputs, static_cast<std::uint16_t>(outputIndices.size()));
+}
+
+void ConditionalTable::calculateReverse(float* inputs, const float* outputs, std::uint16_t rowCount) const {
+    std::fill_n(inputs, inputCount, 0.0f);
+
+    auto isValidOutput = [this, outputs, rowCount](std::uint16_t row) {
+            if (row >= rowCount) {
+                return false;
+            }
+            const std::uint16_t outIndex = outputIndices[row];
+            const float from = fromValues[row];
+            const float to = toValues[row];
+            const float slope = slopeValues[row];
+            const float cut = cutValues[row];
+            const float outValue = outputs[outIndex];
+            const float inValue = (outValue - cut) / slope;
+            return (outValue != 0.0f) && ((from <= inValue) && (inValue <= to));
+        };
+
+    auto findRangeWithMostSolutions = [isValidOutput](ConstArrayView<Range> ranges) {
+            // In reverse mapping, there is some ambiguity about finding out which row was utilized to calculate an output,
+            // as by looking purely at the values, the same output value can be mapped back to different input values through
+            // multiple rows in some cases.
+            // To resolve this ambiguity and get the original input values back, the whole table is partitioned into groups,
+            // where a group is made up of all the rows that rely on the same input index.
+            // Within a single group, rows are further partitioned into even smaller groups based on the range (from, to) of
+            // input values that they accept.
+            // When a single group (containing all rows that map back to the same input index) is being reverse mapped, each
+            // from-to range, within the group is checked for how many valid solutions they generate by the rows they contain.
+            // The from-to range with the largest number of valid solutions is chosen as the most likely candidate that was
+            // used in the original forward mapping, and so the reverse calculation is performed on the first row that gives
+            // a valid input value back within this winning range (rows retain their relative order within these groups, so
+            // because the first matching row is picked as the winner in forward calculations, the same is done in reverse).
+            std::size_t maxSolutionIndex = {};
+            std::size_t maxSolutionCount = {};
+            for (std::size_t i = {}; i < ranges.size(); ++i) {
+                const auto& range = ranges[i];
+                const auto solutionCount =
+                    static_cast<std::size_t>(std::count_if(range.rows.begin(), range.rows.end(), isValidOutput));
+                if (solutionCount > maxSolutionCount) {
+                    maxSolutionCount = solutionCount;
+                    maxSolutionIndex = i;
+                }
+            }
+            return maxSolutionIndex;
+        };
+
+    auto solveInput = [this, inputs, outputs, rowCount](ConstArrayView<std::uint16_t> rows) {
+            for (auto row : rows) {
+                if (row < rowCount) {
+                    const std::uint16_t inIndex = inputIndices[row];
+                    const std::uint16_t outIndex = outputIndices[row];
+                    const float from = fromValues[row];
+                    const float to = toValues[row];
+                    const float slope = slopeValues[row];
+                    const float cut = cutValues[row];
+                    const float outValue = outputs[outIndex];
+                    const float inValue = (outValue - cut) / slope;
+                    if ((from <= inValue) && (inValue <= to)) {
+                        inputs[inIndex] = inValue;
+                        return true;
+                    }
+                }
+            }
+            return false;
+        };
+
+    for (const auto& map : rangeMaps) {
+        const auto rangeIndex = findRangeWithMostSolutions(ConstArrayView<Range>{map.ranges});
+        solveInput(ConstArrayView<std::uint16_t>{map.ranges[rangeIndex].rows});
+    }
+}
+
+void ConditionalTable::calculateReverse(float* inputs, const float* outputs) const {
+    calculateReverse(inputs, outputs, static_cast<std::uint16_t>(outputIndices.size()));
 }
 
 }  // namespace rl4
diff --git a/Engine/Plugins/Animation/RigLogic/Source/RigLogicLib/Private/riglogic/conditionaltable/ConditionalTable.h b/Engine/Plugins/Animation/RigLogic/Source/RigLogicLib/Private/riglogic/conditionaltable/ConditionalTable.h
--- a/Engine/Plugins/Animation/RigLogic/Source/RigLogicLib/Private/riglogic/conditionaltable/ConditionalTable.h
+++ b/Engine/Plugins/Animation/RigLogic/Source/RigLogicLib/Private/riglogic/conditionaltable/ConditionalTable.h
@@ -8,6 +8,57 @@
 
 namespace rl4 {
 
+struct Range {
+    float from;
+    float to;
+    Vector<std::uint16_t> rows;
+
+    explicit Range(MemoryResource* memRes) : from{}, to{}, rows{memRes} {
+    }
+
+    Range(float from_, float to_, MemoryResource* memRes) : from{from_}, to{to_}, rows{memRes} {
+    }
+
+    std::size_t size() const {
+        return rows.size();
+    }
+
+    template<class Archive>
+    void serialize(Archive& archive) {
+        archive(from, to, rows);
+    }
+
+};
+
+struct RangeMap {
+    Vector<Range> ranges;
+
+    explicit RangeMap(MemoryResource* memRes) : ranges{memRes} {
+    }
+
+    Range* findRange(float from, float to) {
+        auto it = std::find_if(ranges.begin(), ranges.end(), [from, to](const Range& range) {
+                return ((range.from == from) && (range.to == to));
+            });
+        return (it == ranges.end() ? nullptr : &(*it));
+    }
+
+    Range* addRange(float from, float to) {
+        Range* range = findRange(from, to);
+        if (range != nullptr) {
+            return range;
+        }
+        ranges.emplace_back(from, to, ranges.get_allocator().getMemoryResource());
+        return &ranges.back();
+    }
+
+    template<class Archive>
+    void serialize(Archive& archive) {
+        archive(ranges);
+    }
+
+};
+
 class ConditionalTable {
     public:
         explicit ConditionalTable(MemoryResource* memRes);
@@ -21,14 +72,18 @@ class ConditionalTable {
                          std::uint16_t outputCount_,
                          MemoryResource* memRes);
 
+        std::uint16_t getRowCount() const;
         std::uint16_t getInputCount() const;
         std::uint16_t getOutputCount() const;
-        void calculate(const float* inputs, float* outputs) const;
-        void calculate(const float* inputs, float* outputs, std::uint16_t chunkSize) const;
+        void calculateForward(const float* inputs, float* outputs) const;
+        void calculateForward(const float* inputs, float* outputs, std::uint16_t rowCount) const;
+        void calculateReverse(float* inputs, const float* outputs) const;
+        void calculateReverse(float* inputs, const float* outputs, std::uint16_t rowCount) const;
 
         template<class Archive>
         void serialize(Archive& archive) {
-            archive(intervalsRemaining,
+            archive(rangeMaps,
+                    intervalsRemaining,
                     inputIndices,
                     outputIndices,
                     fromValues,
@@ -40,6 +95,7 @@ class ConditionalTable {
         }
 
     private:
+        Vector<RangeMap> rangeMaps;
         Vector<std::uint16_t> intervalsRemaining;
         Vector<std::uint16_t> inputIndices;
         Vector<std::uint16_t> outputIndices;
diff --git a/Engine/Plugins/Animation/RigLogic/Source/RigLogicLib/Private/riglogic/controls/Controls.cpp b/Engine/Plugins/Animation/RigLogic/Source/RigLogicLib/Private/riglogic/controls/Controls.cpp
--- a/Engine/Plugins/Animation/RigLogic/Source/RigLogicLib/Private/riglogic/controls/Controls.cpp
+++ b/Engine/Plugins/Animation/RigLogic/Source/RigLogicLib/Private/riglogic/controls/Controls.cpp
@@ -3,6 +3,7 @@
 #include "riglogic/controls/Controls.h"
 
 #include "riglogic/conditionaltable/ConditionalTable.h"
+#include "riglogic/controls/ControlsInputInstance.h"
 #include "riglogic/psdmatrix/PSDMatrix.h"
 
 #include <cassert>
@@ -10,18 +11,32 @@
 
 namespace rl4 {
 
-Controls::Controls(ConditionalTable&& guiToRawMapping_, PSDMatrix&& psds_) :
+Controls::Controls(ConditionalTable&& guiToRawMapping_, PSDMatrix&& psds_, ControlsInputInstance::Factory instanceFactory_) :
     guiToRawMapping{std::move(guiToRawMapping_)},
-    psds{std::move(psds_)} {
+    psds{std::move(psds_)},
+    instanceFactory{instanceFactory_} {
 }
 
-void Controls::mapGUIToRaw(ConstArrayView<float> guiValues, ArrayView<float> inputs) const {
-    assert(guiValues.size() == guiToRawMapping.getInputCount());
-    guiToRawMapping.calculate(guiValues.data(), inputs.data());
+ControlsInputInstance::Pointer Controls::createInstance(MemoryResource* instanceMemRes) const {
+    return instanceFactory(instanceMemRes);
 }
 
-void Controls::calculate(ArrayView<float> inputs) const {
-    psds.calculate(inputs.data(), guiToRawMapping.getOutputCount());
+void Controls::mapGUIToRaw(ControlsInputInstance* instance) const {
+    auto guiControlBuffer = instance->getGUIControlBuffer();
+    assert(guiControlBuffer.size() == guiToRawMapping.getInputCount());
+    auto inputBuffer = instance->getInputBuffer();
+    guiToRawMapping.calculateForward(guiControlBuffer.data(), inputBuffer.data());
+}
+
+void Controls::mapRawToGUI(ControlsInputInstance* instance) const {
+    auto guiControlBuffer = instance->getGUIControlBuffer();
+    assert(guiControlBuffer.size() == guiToRawMapping.getInputCount());
+    auto inputBuffer = instance->getInputBuffer();
+    guiToRawMapping.calculateReverse(guiControlBuffer.data(), inputBuffer.data());
+}
+
+void Controls::calculate(ControlsInputInstance* instance) const {
+    psds.calculate(instance->getInputBuffer().data(), guiToRawMapping.getOutputCount());
 }
 
 }  // namespace rl4
diff --git a/Engine/Plugins/Animation/RigLogic/Source/RigLogicLib/Private/riglogic/controls/Controls.h b/Engine/Plugins/Animation/RigLogic/Source/RigLogicLib/Private/riglogic/controls/Controls.h
--- a/Engine/Plugins/Animation/RigLogic/Source/RigLogicLib/Private/riglogic/controls/Controls.h
+++ b/Engine/Plugins/Animation/RigLogic/Source/RigLogicLib/Private/riglogic/controls/Controls.h
@@ -2,9 +2,10 @@
 
 #pragma once
 
+#include "riglogic/TypeDefs.h"
 #include "riglogic/conditionaltable/ConditionalTable.h"
+#include "riglogic/controls/ControlsInputInstance.h"
 #include "riglogic/psdmatrix/PSDMatrix.h"
-#include "riglogic/types/Aliases.h"
 
 #include <cstdint>
 
@@ -12,10 +13,15 @@ namespace rl4 {
 
 class Controls {
     public:
-        Controls(ConditionalTable&& guiToRawMapping_, PSDMatrix&& psds_);
+        using Pointer = UniqueInstance<Controls>::PointerType;
 
-        void mapGUIToRaw(ConstArrayView<float> guiValues, ArrayView<float> inputs) const;
-        void calculate(ArrayView<float> inputs) const;
+    public:
+        Controls(ConditionalTable&& guiToRawMapping_, PSDMatrix&& psds_, ControlsInputInstance::Factory instanceFactory_);
+
+        ControlsInputInstance::Pointer createInstance(MemoryResource* instanceMemRes) const;
+        void mapGUIToRaw(ControlsInputInstance* instance) const;
+        void mapRawToGUI(ControlsInputInstance* instance) const;
+        void calculate(ControlsInputInstance* instance) const;
 
         template<class Archive>
         void serialize(Archive& archive) {
@@ -25,6 +31,7 @@ class Controls {
     private:
         ConditionalTable guiToRawMapping;
         PSDMatrix psds;
+        ControlsInputInstance::Factory instanceFactory;
 
 };
 
diff --git a/Engine/Plugins/Animation/RigLogic/Source/RigLogicLib/Private/riglogic/controls/ControlsFactory.cpp b/Engine/Plugins/Animation/RigLogic/Source/RigLogicLib/Private/riglogic/controls/ControlsFactory.cpp
--- a/Engine/Plugins/Animation/RigLogic/Source/RigLogicLib/Private/riglogic/controls/ControlsFactory.cpp
+++ b/Engine/Plugins/Animation/RigLogic/Source/RigLogicLib/Private/riglogic/controls/ControlsFactory.cpp
@@ -2,23 +2,40 @@
 
 #include "riglogic/controls/ControlsFactory.h"
 
+#include "riglogic/TypeDefs.h"
 #include "riglogic/conditionaltable/ConditionalTable.h"
 #include "riglogic/controls/Controls.h"
+#include "riglogic/controls/instances/StandardControlsInputInstance.h"
 #include "riglogic/psdmatrix/PSDMatrix.h"
+#include "riglogic/riglogic/Configuration.h"
+#include "riglogic/riglogic/RigMetrics.h"
 #include "riglogic/utils/Extd.h"
 
-#include <dna/layers/BehaviorReader.h>
-
-#include <cstddef>
-#include <cstring>
+#include <cstdint>
 
 namespace rl4 {
 
-ControlsFactory::ControlsPtr ControlsFactory::create(MemoryResource* memRes) {
-    return ManagedControls::with(memRes).create(ConditionalTable{memRes}, PSDMatrix{memRes});
+static ControlsInputInstance::Factory createInstanceFactory(const Configuration&  /*unused*/,
+                                                            std::uint16_t guiControlCount,
+                                                            std::uint16_t rawControlCount,
+                                                            std::uint16_t psdControlCount,
+                                                            std::uint16_t mlControlCount) {
+    return [ = ](MemoryResource* memRes) {
+               auto factory = UniqueInstance<StandardControlsInputInstance, ControlsInputInstance>::with(memRes);
+               return factory.create(guiControlCount, rawControlCount, psdControlCount, mlControlCount, memRes);
+    };
+}
+
+Controls::Pointer ControlsFactory::create(const Configuration& config, const RigMetrics& metrics, MemoryResource* memRes) {
+    auto instanceFactory = createInstanceFactory(config,
+                                                 metrics.guiControlCount,
+                                                 metrics.rawControlCount,
+                                                 metrics.psdControlCount,
+                                                 metrics.mlControlCount);
+    return UniqueInstance<Controls>::with(memRes).create(ConditionalTable{memRes}, PSDMatrix{memRes}, instanceFactory);
 }
 
-ControlsFactory::ControlsPtr ControlsFactory::create(const dna::BehaviorReader* reader, MemoryResource* memRes) {
+Controls::Pointer ControlsFactory::create(const Configuration& config, const dna::Reader* reader, MemoryResource* memRes) {
     Vector<std::uint16_t> inputIndices{memRes};
     Vector<std::uint16_t> outputIndices{memRes};
     Vector<float> fromValues{memRes};
@@ -61,7 +78,13 @@ ControlsFactory::ControlsPtr ControlsFactory::create(const dna::BehaviorReader*
                    std::move(psdColumns),
                    std::move(psdValues)};
 
-    return ManagedControls::with(memRes).create(std::move(conditionals), std::move(psds));
+    auto instanceFactory = createInstanceFactory(config,
+                                                 conditionals.getInputCount(),
+                                                 conditionals.getOutputCount(),
+                                                 psds.getDistinctPSDCount(),
+                                                 reader->getMLControlCount());
+
+    return UniqueInstance<Controls>::with(memRes).create(std::move(conditionals), std::move(psds), instanceFactory);
 }
 
 }  // namespace rl4
diff --git a/Engine/Plugins/Animation/RigLogic/Source/RigLogicLib/Private/riglogic/controls/ControlsFactory.h b/Engine/Plugins/Animation/RigLogic/Source/RigLogicLib/Private/riglogic/controls/ControlsFactory.h
--- a/Engine/Plugins/Animation/RigLogic/Source/RigLogicLib/Private/riglogic/controls/ControlsFactory.h
+++ b/Engine/Plugins/Animation/RigLogic/Source/RigLogicLib/Private/riglogic/controls/ControlsFactory.h
@@ -2,24 +2,18 @@
 
 #pragma once
 
+#include "riglogic/TypeDefs.h"
 #include "riglogic/controls/Controls.h"
 
-#include <pma/utils/ManagedInstance.h>
-
-namespace dna {
-
-class BehaviorReader;
-
-}  // namespace dna
-
 namespace rl4 {
 
+struct Configuration;
+struct RigMetrics;
+
 struct ControlsFactory {
-    using ManagedControls = pma::UniqueInstance<Controls>;
-    using ControlsPtr = typename ManagedControls::PointerType;
+    static Controls::Pointer create(const Configuration& config, const dna::Reader* reader, MemoryResource* memRes);
+    static Controls::Pointer create(const Configuration& config, const RigMetrics& metrics, MemoryResource* memRes);
 
-    static ControlsPtr create(MemoryResource* memRes);
-    static ControlsPtr create(const dna::BehaviorReader* reader, MemoryResource* memRes);
 };
 
 }  // namespace rl4
diff --git a/Engine/Plugins/Animation/RigLogic/Source/RigLogicLib/Private/riglogic/controls/instances/StandardControlsInputInstance.cpp b/Engine/Plugins/Animation/RigLogic/Source/RigLogicLib/Private/riglogic/controls/instances/StandardControlsInputInstance.cpp
new file mode 100644
--- /dev/null
+++ b/Engine/Plugins/Animation/RigLogic/Source/RigLogicLib/Private/riglogic/controls/instances/StandardControlsInputInstance.cpp
@@ -0,0 +1,35 @@
+// Copyright Epic Games, Inc. All Rights Reserved.
+
+#include "riglogic/controls/instances/StandardControlsInputInstance.h"
+
+#include <cstdint>
+#include <cstddef>
+
+namespace rl4 {
+
+StandardControlsInputInstance::StandardControlsInputInstance(std::uint16_t guiControlCount,
+                                                             std::uint16_t rawControlCount,
+                                                             std::uint16_t psdControlCount,
+                                                             std::uint16_t mlControlCount,
+                                                             MemoryResource* memRes) :
+    guiControlBuffer{guiControlCount, {}, memRes},
+    inputBuffer{static_cast<std::size_t>(rawControlCount + psdControlCount + mlControlCount), {}, memRes} {
+}
+
+ArrayView<float> StandardControlsInputInstance::getGUIControlBuffer() {
+    return ArrayView<float>{guiControlBuffer};
+}
+
+ArrayView<float> StandardControlsInputInstance::getInputBuffer() {
+    return ArrayView<float>{inputBuffer};
+}
+
+ConstArrayView<float> StandardControlsInputInstance::getGUIControlBuffer() const {
+    return ConstArrayView<float>{guiControlBuffer};
+}
+
+ConstArrayView<float> StandardControlsInputInstance::getInputBuffer() const {
+    return ConstArrayView<float>{inputBuffer};
+}
+
+}  // namespace rl4
diff --git a/Engine/Plugins/Animation/RigLogic/Source/RigLogicLib/Private/riglogic/joints/Joints.cpp b/Engine/Plugins/Animation/RigLogic/Source/RigLogicLib/Private/riglogic/joints/Joints.cpp
--- a/Engine/Plugins/Animation/RigLogic/Source/RigLogicLib/Private/riglogic/joints/Joints.cpp
+++ b/Engine/Plugins/Animation/RigLogic/Source/RigLogicLib/Private/riglogic/joints/Joints.cpp
@@ -2,6 +2,7 @@
 
 #include "riglogic/joints/Joints.h"
 
+#include "riglogic/joints/JointsEvaluator.h"
 #include "riglogic/riglogic/RigInstance.h"
 #include "riglogic/transformation/Transformation.h"
 #include "riglogic/types/Aliases.h"
@@ -10,14 +11,14 @@
 
 namespace rl4 {
 
-Joints::Joints(JointsEvaluatorPtr evaluator_, MemoryResource* memRes) :
+Joints::Joints(JointsEvaluator::Pointer evaluator_, MemoryResource* memRes) :
     evaluator{std::move(evaluator_)},
     neutralValues{memRes},
     variableAttributeIndices{memRes},
     jointGroupCount{} {
 }
 
-Joints::Joints(JointsEvaluatorPtr evaluator_,
+Joints::Joints(JointsEvaluator::Pointer evaluator_,
                Vector<float>&& neutralValues_,
                Matrix<std::uint16_t>&& variableAttributeIndices_,
                std::uint16_t jointGroupCount_) :
@@ -27,11 +28,17 @@ Joints::Joints(JointsEvaluatorPtr evaluator_,
     jointGroupCount{jointGroupCount_} {
 }
 
-void Joints::calculate(ConstArrayView<float> inputs, ArrayView<float> outputs, std::uint16_t lod) const {
+JointsOutputInstance::Pointer Joints::createInstance(MemoryResource* instanceMemRes) const {
+    return evaluator->createInstance(instanceMemRes);
+}
+
+void Joints::calculate(const ControlsInputInstance* inputs, JointsOutputInstance* outputs, std::uint16_t lod) const {
     evaluator->calculate(inputs, outputs, lod);
 }
 
-void Joints::calculate(ConstArrayView<float> inputs, ArrayView<float> outputs, std::uint16_t lod,
+void Joints::calculate(const ControlsInputInstance* inputs,
+                       JointsOutputInstance* outputs,
+                       std::uint16_t lod,
                        std::uint16_t jointGroupIndex) const {
     evaluator->calculate(inputs, outputs, lod, jointGroupIndex);
 }
diff --git a/Engine/Plugins/Animation/RigLogic/Source/RigLogicLib/Private/riglogic/joints/Joints.h b/Engine/Plugins/Animation/RigLogic/Source/RigLogicLib/Private/riglogic/joints/Joints.h
--- a/Engine/Plugins/Animation/RigLogic/Source/RigLogicLib/Private/riglogic/joints/Joints.h
+++ b/Engine/Plugins/Animation/RigLogic/Source/RigLogicLib/Private/riglogic/joints/Joints.h
@@ -4,34 +4,32 @@
 
 #include "riglogic/TypeDefs.h"
 #include "riglogic/joints/JointsEvaluator.h"
+#include "riglogic/joints/JointsOutputInstance.h"
 #include "riglogic/transformation/Transformation.h"
-#include "riglogic/types/Aliases.h"
 
-#ifdef _MSC_VER
-    #pragma warning(push)
-    #pragma warning(disable : 4365 4987)
-#endif
 #include <cstdint>
-#include <functional>
-#include <memory>
-#ifdef _MSC_VER
-    #pragma warning(pop)
-#endif
 
 namespace rl4 {
 
-class RigInstance;
+class ControlsInputInstance;
+class JointsOutputInstance;
 
 class Joints {
     public:
-        Joints(JointsEvaluatorPtr evaluator_, MemoryResource* memRes);
-        Joints(JointsEvaluatorPtr evaluator_,
+        using Pointer = UniqueInstance<Joints>::PointerType;
+
+    public:
+        Joints(JointsEvaluator::Pointer evaluator_, MemoryResource* memRes);
+        Joints(JointsEvaluator::Pointer evaluator_,
                Vector<float>&& neutralValues_,
                Matrix<std::uint16_t>&& variableAttributeIndices_,
                std::uint16_t jointGroupCount_);
 
-        void calculate(ConstArrayView<float> inputs, ArrayView<float> outputs, std::uint16_t lod) const;
-        void calculate(ConstArrayView<float> inputs, ArrayView<float> outputs, std::uint16_t lod,
+        JointsOutputInstance::Pointer createInstance(MemoryResource* instanceMemRes) const;
+        void calculate(const ControlsInputInstance* inputs, JointsOutputInstance* outputs, std::uint16_t lod) const;
+        void calculate(const ControlsInputInstance* inputs,
+                       JointsOutputInstance* outputs,
+                       std::uint16_t lod,
                        std::uint16_t jointGroupIndex) const;
 
         template<class Archive>
@@ -52,7 +50,7 @@ class Joints {
         ConstArrayView<std::uint16_t> getVariableAttributeIndices(std::uint16_t lod) const;
 
     private:
-        JointsEvaluatorPtr evaluator;
+        JointsEvaluator::Pointer evaluator;
         Vector<float> neutralValues;
         Matrix<std::uint16_t> variableAttributeIndices;
         std::uint16_t jointGroupCount;
diff --git a/Engine/Plugins/Animation/RigLogic/Source/RigLogicLib/Private/riglogic/joints/JointsBuilder.cpp b/Engine/Plugins/Animation/RigLogic/Source/RigLogicLib/Private/riglogic/joints/JointsBuilder.cpp
--- a/Engine/Plugins/Animation/RigLogic/Source/RigLogicLib/Private/riglogic/joints/JointsBuilder.cpp
+++ b/Engine/Plugins/Animation/RigLogic/Source/RigLogicLib/Private/riglogic/joints/JointsBuilder.cpp
@@ -2,7 +2,9 @@
 
 #include "riglogic/joints/JointsBuilder.h"
 
+#include "riglogic/TypeDefs.h"
 #include "riglogic/joints/bpcm/builders/Scalar.h"
+#include "riglogic/system/simd/Detect.h"
 #ifdef RL_BUILD_WITH_AVX
     #ifdef RL_USE_HALF_FLOATS
         #include "riglogic/joints/bpcm/builders/AVXHalfFloat.h"
@@ -18,26 +20,24 @@
     #endif  // RL_USE_HALF_FLOATS
 #endif  // RL_BUILD_WITH_SSE
 
-#include <pma/utils/ManagedInstance.h>
-
 namespace rl4 {
 
 JointsBuilder::~JointsBuilder() = default;
 
-JointsBuilder::JointsBuilderPtr JointsBuilder::create(Configuration config, MemoryResource* memRes) {
+JointsBuilder::Pointer JointsBuilder::create(Configuration config, MemoryResource* memRes) {
     // Work around unused parameter warning when building without SSE and AVX
     static_cast<void>(config);
     #ifdef RL_BUILD_WITH_SSE
         if (config.calculationType == CalculationType::SSE) {
-            return pma::UniqueInstance<bpcm::SSEJointsBuilder, JointsBuilder>::with(memRes).create(memRes);
+            return UniqueInstance<bpcm::SSEJointsBuilder, JointsBuilder>::with(memRes).create(memRes);
         }
     #endif  // RL_BUILD_WITH_SSE
     #ifdef RL_BUILD_WITH_AVX
         if (config.calculationType == CalculationType::AVX) {
-            return pma::UniqueInstance<bpcm::AVXJointsBuilder, JointsBuilder>::with(memRes).create(memRes);
+            return UniqueInstance<bpcm::AVXJointsBuilder, JointsBuilder>::with(memRes).create(memRes);
         }
     #endif  // RL_BUILD_WITH_AVX
-    return pma::UniqueInstance<bpcm::ScalarJointsBuilder, JointsBuilder>::with(memRes).create(memRes);
+    return UniqueInstance<bpcm::ScalarJointsBuilder, JointsBuilder>::with(memRes).create(memRes);
 }
 
 }  // namespace rl4
diff --git a/Engine/Plugins/Animation/RigLogic/Source/RigLogicLib/Private/riglogic/joints/JointsBuilder.h b/Engine/Plugins/Animation/RigLogic/Source/RigLogicLib/Private/riglogic/joints/JointsBuilder.h
--- a/Engine/Plugins/Animation/RigLogic/Source/RigLogicLib/Private/riglogic/joints/JointsBuilder.h
+++ b/Engine/Plugins/Animation/RigLogic/Source/RigLogicLib/Private/riglogic/joints/JointsBuilder.h
@@ -2,35 +2,28 @@
 
 #pragma once
 
+#include "riglogic/TypeDefs.h"
 #include "riglogic/joints/JointsEvaluator.h"
 #include "riglogic/riglogic/Configuration.h"
-#include "riglogic/types/Aliases.h"
-
-#ifdef _MSC_VER
-    #pragma warning(push)
-    #pragma warning(disable : 4365 4987)
-#endif
-#include <functional>
-#include <memory>
-#ifdef _MSC_VER
-    #pragma warning(pop)
-#endif
 
 namespace rl4 {
 
+struct RigMetrics;
+
 class JointsBuilder {
     public:
-        using JointsEvaluatorPtr = std::unique_ptr<JointsEvaluator, std::function<void (JointsEvaluator*)> >;
-        using JointsBuilderPtr = std::unique_ptr<JointsBuilder, std::function<void (JointsBuilder*)> >;
+        using Pointer = UniqueInstance<JointsBuilder>::PointerType;
 
     public:
         virtual ~JointsBuilder();
 
-        static JointsBuilderPtr create(Configuration config, MemoryResource* memRes);
+        static Pointer create(Configuration config, MemoryResource* memRes);
 
+        virtual void computeStorageRequirements(const RigMetrics& source) = 0;
+        virtual void computeStorageRequirements(const dna::BehaviorReader* source) = 0;
         virtual void allocateStorage(const dna::BehaviorReader* source) = 0;
         virtual void fillStorage(const dna::BehaviorReader* source) = 0;
-        virtual JointsEvaluatorPtr build() = 0;
+        virtual JointsEvaluator::Pointer build() = 0;
 };
 
 }  // namespace rl4
diff --git a/Engine/Plugins/Animation/RigLogic/Source/RigLogicLib/Private/riglogic/joints/JointsEvaluator.h b/Engine/Plugins/Animation/RigLogic/Source/RigLogicLib/Private/riglogic/joints/JointsEvaluator.h
--- a/Engine/Plugins/Animation/RigLogic/Source/RigLogicLib/Private/riglogic/joints/JointsEvaluator.h
+++ b/Engine/Plugins/Animation/RigLogic/Source/RigLogicLib/Private/riglogic/joints/JointsEvaluator.h
@@ -2,41 +2,31 @@
 
 #pragma once
 
-#include "riglogic/types/Aliases.h"
+#include "riglogic/TypeDefs.h"
+#include "riglogic/joints/JointsOutputInstance.h"
 
-#include <terse/archives/binary/InputArchive.h>
-#include <terse/archives/binary/OutputArchive.h>
-#include <trio/Stream.h>
-
-#ifdef _MSC_VER
-    #pragma warning(push)
-    #pragma warning(disable : 4365 4987)
-#endif
 #include <cstdint>
-#include <functional>
-#include <memory>
-#ifdef _MSC_VER
-    #pragma warning(pop)
-#endif
 
 namespace rl4 {
 
-class RigInstance;
+class ControlsInputInstance;
 
 class JointsEvaluator {
+    public:
+        using Pointer = UniqueInstance<JointsEvaluator>::PointerType;
+
     protected:
         virtual ~JointsEvaluator();
 
     public:
-        virtual void calculate(ConstArrayView<float> inputs, ArrayView<float> outputs, std::uint16_t lod) const = 0;
-        virtual void calculate(ConstArrayView<float> inputs,
-                               ArrayView<float> outputs,
+        virtual JointsOutputInstance::Pointer createInstance(MemoryResource* instanceMemRes) const = 0;
+        virtual void calculate(const ControlsInputInstance* inputs, JointsOutputInstance* outputs, std::uint16_t lod) const = 0;
+        virtual void calculate(const ControlsInputInstance* inputs,
+                               JointsOutputInstance* outputs,
                                std::uint16_t lod,
                                std::uint16_t jointGroupIndex) const = 0;
         virtual void load(terse::BinaryInputArchive<BoundedIOStream>& archive) = 0;
         virtual void save(terse::BinaryOutputArchive<BoundedIOStream>& archive) = 0;
 };
 
-using JointsEvaluatorPtr = std::unique_ptr<JointsEvaluator, std::function<void (JointsEvaluator*)> >;
-
 }  // namespace rl4
diff --git a/Engine/Plugins/Animation/RigLogic/Source/RigLogicLib/Private/riglogic/joints/JointsFactory.cpp b/Engine/Plugins/Animation/RigLogic/Source/RigLogicLib/Private/riglogic/joints/JointsFactory.cpp
--- a/Engine/Plugins/Animation/RigLogic/Source/RigLogicLib/Private/riglogic/joints/JointsFactory.cpp
+++ b/Engine/Plugins/Animation/RigLogic/Source/RigLogicLib/Private/riglogic/joints/JointsFactory.cpp
@@ -2,14 +2,13 @@
 
 #include "riglogic/joints/JointsFactory.h"
 
+#include "riglogic/TypeDefs.h"
 #include "riglogic/joints/JointsBuilder.h"
 #include "riglogic/joints/JointsEvaluator.h"
+#include "riglogic/joints/JointsNullEvaluator.h"
 #include "riglogic/riglogic/Configuration.h"
+#include "riglogic/riglogic/RigMetrics.h"
 #include "riglogic/transformation/Transformation.h"
-#include "riglogic/types/Aliases.h"
-
-#include <dna/layers/BehaviorReader.h>
-#include <pma/utils/ManagedInstance.h>
 
 #ifdef _MSC_VER
     #pragma warning(push)
@@ -68,26 +67,35 @@ static Matrix<std::uint16_t> copyVariableAttributeIndices(const dna::BehaviorRea
     return variableAttributeIndices;
 }
 
-JointsFactory::JointsPtr JointsFactory::create(const Configuration& config,
-                                               const dna::BehaviorReader* reader,
-                                               MemoryResource* memRes) {
+Joints::Pointer JointsFactory::create(const Configuration& config, const dna::BehaviorReader* reader, MemoryResource* memRes) {
+    if (!config.loadJoints) {
+        auto evaluator = UniqueInstance<JointsNullEvaluator, JointsEvaluator>::with(memRes).create();
+        return UniqueInstance<Joints>::with(memRes).create(std::move(evaluator), memRes);
+    }
+
     auto neutralValues = copyNeutralValues(reader, memRes);
     auto variableAttributeIndices = copyVariableAttributeIndices(reader, memRes);
     auto builder = JointsBuilder::create(config, memRes);
+    builder->computeStorageRequirements(reader);
     builder->allocateStorage(reader);
     builder->fillStorage(reader);
     auto evaluator = builder->build();
-    return pma::UniqueInstance<Joints>::with(memRes).create(std::move(evaluator),
-                                                            std::move(neutralValues),
-                                                            std::move(variableAttributeIndices),
-                                                            reader->getJointGroupCount());
+    return UniqueInstance<Joints>::with(memRes).create(std::move(evaluator),
+                                                       std::move(neutralValues),
+                                                       std::move(variableAttributeIndices),
+                                                       reader->getJointGroupCount());
 }
 
-JointsFactory::JointsPtr JointsFactory::create(const Configuration& config, const RigMetrics&  /*unused*/,
-                                               MemoryResource* memRes) {
+Joints::Pointer JointsFactory::create(const Configuration& config, const RigMetrics& metrics, MemoryResource* memRes) {
+    if (!config.loadJoints) {
+        auto evaluator = UniqueInstance<JointsNullEvaluator, JointsEvaluator>::with(memRes).create();
+        return UniqueInstance<Joints>::with(memRes).create(std::move(evaluator), memRes);
+    }
+
     auto builder = JointsBuilder::create(config, memRes);
+    builder->computeStorageRequirements(metrics);
     auto evaluator = builder->build();
-    return pma::UniqueInstance<Joints>::with(memRes).create(std::move(evaluator), memRes);
+    return UniqueInstance<Joints>::with(memRes).create(std::move(evaluator), memRes);
 }
 
 }  // namespace rl4
diff --git a/Engine/Plugins/Animation/RigLogic/Source/RigLogicLib/Private/riglogic/joints/JointsFactory.h b/Engine/Plugins/Animation/RigLogic/Source/RigLogicLib/Private/riglogic/joints/JointsFactory.h
--- a/Engine/Plugins/Animation/RigLogic/Source/RigLogicLib/Private/riglogic/joints/JointsFactory.h
+++ b/Engine/Plugins/Animation/RigLogic/Source/RigLogicLib/Private/riglogic/joints/JointsFactory.h
@@ -2,28 +2,17 @@
 
 #pragma once
 
+#include "riglogic/TypeDefs.h"
 #include "riglogic/joints/Joints.h"
-#include "riglogic/riglogic/Configuration.h"
-#include "riglogic/riglogic/RigMetrics.h"
-#include "riglogic/types/Aliases.h"
-
-#ifdef _MSC_VER
-    #pragma warning(push)
-    #pragma warning(disable : 4365 4987)
-#endif
-#include <functional>
-#include <memory>
-#ifdef _MSC_VER
-    #pragma warning(pop)
-#endif
 
 namespace rl4 {
 
-struct JointsFactory {
-    using JointsPtr = std::unique_ptr<Joints, std::function<void (Joints*)> >;
+struct Configuration;
+struct RigMetrics;
 
-    static JointsPtr create(const Configuration& config, const dna::BehaviorReader* reader, MemoryResource* memRes);
-    static JointsPtr create(const Configuration& config, const RigMetrics& metrics, MemoryResource* memRes);
+struct JointsFactory {
+    static Joints::Pointer create(const Configuration& config, const dna::BehaviorReader* reader, MemoryResource* memRes);
+    static Joints::Pointer create(const Configuration& config, const RigMetrics& metrics, MemoryResource* memRes);
 
 };
 
diff --git a/Engine/Plugins/Animation/RigLogic/Source/RigLogicLib/Private/riglogic/joints/JointsNullEvaluator.cpp b/Engine/Plugins/Animation/RigLogic/Source/RigLogicLib/Private/riglogic/joints/JointsNullEvaluator.cpp
new file mode 100644
--- /dev/null
+++ b/Engine/Plugins/Animation/RigLogic/Source/RigLogicLib/Private/riglogic/joints/JointsNullEvaluator.cpp
@@ -0,0 +1,27 @@
+// Copyright Epic Games, Inc. All Rights Reserved.
+
+#include "riglogic/joints/JointsNullEvaluator.h"
+
+namespace rl4 {
+
+JointsOutputInstance::Pointer JointsNullEvaluator::createInstance(MemoryResource* instanceMemRes) const {
+    return UniqueInstance<JointsNullOutputInstance, JointsOutputInstance>::with(instanceMemRes).create();
+}
+
+void JointsNullEvaluator::calculate(const ControlsInputInstance*  /*unused*/, JointsOutputInstance*  /*unused*/,
+                                    std::uint16_t  /*unused*/) const {
+}
+
+void JointsNullEvaluator::calculate(const ControlsInputInstance*  /*unused*/,
+                                    JointsOutputInstance*  /*unused*/,
+                                    std::uint16_t  /*unused*/,
+                                    std::uint16_t  /*unused*/) const {
+}
+
+void JointsNullEvaluator::load(terse::BinaryInputArchive<BoundedIOStream>&  /*unused*/) {
+}
+
+void JointsNullEvaluator::save(terse::BinaryOutputArchive<BoundedIOStream>&  /*unused*/) {
+}
+
+}  // namespace rl4
diff --git a/Engine/Plugins/Animation/RigLogic/Source/RigLogicLib/Private/riglogic/joints/JointsNullEvaluator.h b/Engine/Plugins/Animation/RigLogic/Source/RigLogicLib/Private/riglogic/joints/JointsNullEvaluator.h
new file mode 100644
--- /dev/null
+++ b/Engine/Plugins/Animation/RigLogic/Source/RigLogicLib/Private/riglogic/joints/JointsNullEvaluator.h
@@ -0,0 +1,28 @@
+// Copyright Epic Games, Inc. All Rights Reserved.
+
+#pragma once
+
+#include "riglogic/joints/JointsEvaluator.h"
+#include "riglogic/joints/JointsNullOutputInstance.h"
+
+#include <cstdint>
+
+namespace rl4 {
+
+class ControlsInputInstance;
+
+class JointsNullEvaluator : public JointsEvaluator {
+    public:
+        JointsOutputInstance::Pointer createInstance(MemoryResource* instanceMemRes) const override;
+        void calculate(const ControlsInputInstance*  /*unused*/, JointsOutputInstance*  /*unused*/,
+                       std::uint16_t  /*unused*/) const override;
+        void calculate(const ControlsInputInstance*  /*unused*/,
+                       JointsOutputInstance*  /*unused*/,
+                       std::uint16_t  /*unused*/,
+                       std::uint16_t  /*unused*/) const override;
+        void load(terse::BinaryInputArchive<BoundedIOStream>&  /*unused*/) override;
+        void save(terse::BinaryOutputArchive<BoundedIOStream>&  /*unused*/) override;
+
+};
+
+}  // namespace rl4
diff --git a/Engine/Plugins/Animation/RigLogic/Source/RigLogicLib/Private/riglogic/joints/bpcm/BPCMOutputInstance.cpp b/Engine/Plugins/Animation/RigLogic/Source/RigLogicLib/Private/riglogic/joints/bpcm/BPCMOutputInstance.cpp
new file mode 100644
--- /dev/null
+++ b/Engine/Plugins/Animation/RigLogic/Source/RigLogicLib/Private/riglogic/joints/bpcm/BPCMOutputInstance.cpp
@@ -0,0 +1,24 @@
+// Copyright Epic Games, Inc. All Rights Reserved.
+
+#include "riglogic/joints/bpcm/BPCMOutputInstance.h"
+
+#include "riglogic/TypeDefs.h"
+
+#include <cstdint>
+#include <cstddef>
+
+namespace rl4 {
+
+namespace bpcm {
+
+OutputInstance::OutputInstance(std::uint16_t jointAttributeCount, MemoryResource* memRes) :
+    outputBuffer{jointAttributeCount, {}, memRes} {
+}
+
+ArrayView<float> OutputInstance::getOutputBuffer() {
+    return ArrayView<float>{outputBuffer};
+}
+
+}  // namespace bpcm
+
+}  // namespace rl4
diff --git a/Engine/Plugins/Animation/RigLogic/Source/RigLogicLib/Private/riglogic/joints/bpcm/BPCMOutputInstance.h b/Engine/Plugins/Animation/RigLogic/Source/RigLogicLib/Private/riglogic/joints/bpcm/BPCMOutputInstance.h
new file mode 100644
--- /dev/null
+++ b/Engine/Plugins/Animation/RigLogic/Source/RigLogicLib/Private/riglogic/joints/bpcm/BPCMOutputInstance.h
@@ -0,0 +1,27 @@
+// Copyright Epic Games, Inc. All Rights Reserved.
+
+#pragma once
+
+#include "riglogic/TypeDefs.h"
+#include "riglogic/joints/JointsOutputInstance.h"
+
+#include <cstdint>
+
+namespace rl4 {
+
+namespace bpcm {
+
+class OutputInstance : public JointsOutputInstance {
+    public:
+        OutputInstance(std::uint16_t jointAttributeCount, MemoryResource* memRes);
+
+        ArrayView<float> getOutputBuffer() override;
+
+    private:
+        AlignedVector<float> outputBuffer;
+
+};
+
+}  // namespace bpcm
+
+}  // namespace rl4
diff --git a/Engine/Plugins/Animation/RigLogic/Source/RigLogicLib/Private/riglogic/joints/bpcm/BuilderCommon.h b/Engine/Plugins/Animation/RigLogic/Source/RigLogicLib/Private/riglogic/joints/bpcm/BuilderCommon.h
new file mode 100644
--- /dev/null
+++ b/Engine/Plugins/Animation/RigLogic/Source/RigLogicLib/Private/riglogic/joints/bpcm/BuilderCommon.h
@@ -0,0 +1,173 @@
+// Copyright Epic Games, Inc. All Rights Reserved.
+
+#pragma once
+
+#include "riglogic/TypeDefs.h"
+#include "riglogic/joints/JointsBuilder.h"
+#include "riglogic/joints/bpcm/CalculationStrategy.h"
+#include "riglogic/joints/bpcm/Evaluator.h"
+#include "riglogic/joints/bpcm/BPCMOutputInstance.h"
+#include "riglogic/joints/bpcm/Storage.h"
+#include "riglogic/joints/bpcm/StorageSize.h"
+#include "riglogic/riglogic/RigMetrics.h"
+#include "riglogic/types/Aliases.h"
+#include "riglogic/types/bpcm/Optimizer.h"
+#include "riglogic/utils/Extd.h"
+
+#ifdef _MSC_VER
+    #pragma warning(push)
+    #pragma warning(disable : 4365 4987)
+#endif
+#include <algorithm>
+#include <cstddef>
+#include <cstdint>
+#include <functional>
+#include <memory>
+#include <utility>
+#ifdef _MSC_VER
+    #pragma warning(pop)
+#endif
+
+namespace rl4 {
+
+namespace bpcm {
+
+template<typename TValue, std::uint32_t BlockHeight, std::uint32_t PadTo, typename TFVec>
+class JointsBuilderCommon : public JointsBuilder {
+    protected:
+        using CalculationStrategy = JointCalculationStrategy<TValue>;
+        using CalculationStrategyPtr = std::unique_ptr<CalculationStrategy, std::function<void (CalculationStrategy*)> >;
+
+    protected:
+        JointsBuilderCommon(CalculationStrategyPtr strategy_, MemoryResource* memRes_);
+
+    public:
+        void computeStorageRequirements(const RigMetrics& source) override;
+        void computeStorageRequirements(const dna::BehaviorReader* source) override;
+        void allocateStorage(const dna::BehaviorReader* source) override;
+        void fillStorage(const dna::BehaviorReader* source) override;
+        JointsEvaluator::Pointer build() override;
+
+    protected:
+        void setValues(const dna::BehaviorReader* source);
+        void setInputIndices(const dna::BehaviorReader* source);
+        void setOutputIndices(const dna::BehaviorReader* source);
+        void setLODs(const dna::BehaviorReader* source);
+
+    protected:
+        MemoryResource* memRes;
+        StorageSize sizeReqs;
+
+        JointStorage<TValue> storage;
+        CalculationStrategyPtr strategy;
+};
+
+template<typename TValue, std::uint32_t BlockHeight, std::uint32_t PadTo, typename TFVec>
+JointsBuilderCommon<TValue, BlockHeight, PadTo, TFVec>::JointsBuilderCommon(CalculationStrategyPtr strategy_,
+                                                                            MemoryResource* memRes_) :
+    memRes{memRes_},
+    sizeReqs{memRes},
+    storage{memRes},
+    strategy{std::move(strategy_)} {
+}
+
+template<typename TValue, std::uint32_t BlockHeight, std::uint32_t PadTo, typename TFVec>
+void JointsBuilderCommon<TValue, BlockHeight, PadTo, TFVec>::computeStorageRequirements(const RigMetrics& source) {
+    // This is incomplete, but enough to create valid joint output instances when restoring from a dump.
+    // Joint attribute count is the only needed data by the instance factory.
+    sizeReqs.attributeCount = source.jointAttributeCount;
+}
+
+template<typename TValue, std::uint32_t BlockHeight, std::uint32_t PadTo, typename TFVec>
+void JointsBuilderCommon<TValue, BlockHeight, PadTo, TFVec>::computeStorageRequirements(const dna::BehaviorReader* source) {
+    sizeReqs.computeFrom(source, PadTo);
+}
+
+template<typename TValue, std::uint32_t BlockHeight, std::uint32_t PadTo, typename TFVec>
+void JointsBuilderCommon<TValue, BlockHeight, PadTo, TFVec>::allocateStorage(const dna::BehaviorReader*  /*unused*/) {
+    storage.values.resize(sizeReqs.valueCount);
+    storage.inputIndices.resize(sizeReqs.inputIndexCount);
+    storage.outputIndices.resize(sizeReqs.outputIndexCount);
+    storage.lodRegions.reserve(sizeReqs.lodRegionCount);
+    storage.jointGroups.resize(sizeReqs.jointGroups.size());
+}
+
+template<typename TValue, std::uint32_t BlockHeight, std::uint32_t PadTo, typename TFVec>
+void JointsBuilderCommon<TValue, BlockHeight, PadTo, TFVec>::fillStorage(const dna::BehaviorReader* source) {
+    setValues(source);
+    setInputIndices(source);
+    setOutputIndices(source);
+    setLODs(source);
+}
+
+template<typename TValue, std::uint32_t BlockHeight, std::uint32_t PadTo, typename TFVec>
+void JointsBuilderCommon<TValue, BlockHeight, PadTo, TFVec>::setValues(const dna::BehaviorReader* source) {
+    std::uint32_t offset = 0ul;
+    for (std::uint16_t i = 0u; i < source->getJointGroupCount(); ++i) {
+        const auto values = source->getJointGroupValues(i);
+        const auto jointGroupSize = sizeReqs.getJointGroupSize(i);
+        storage.jointGroups[i].valuesOffset = offset;
+        storage.jointGroups[i].valuesSize = jointGroupSize.padded.size();
+        offset += Optimizer<TFVec, BlockHeight, PadTo>::optimize(storage.values.data() + offset,
+                                                                 values.data(),
+                                                                 jointGroupSize.original);
+    }
+}
+
+template<typename TValue, std::uint32_t BlockHeight, std::uint32_t PadTo, typename TFVec>
+void JointsBuilderCommon<TValue, BlockHeight, PadTo, TFVec>::setInputIndices(const dna::BehaviorReader* source) {
+    std::uint32_t offset = 0ul;
+    for (std::uint16_t i = 0u; i < source->getJointGroupCount(); ++i) {
+        const auto jointGroupSize = sizeReqs.getJointGroupSize(i);
+        const auto indices = source->getJointGroupInputIndices(i);
+        std::copy(indices.begin(), indices.end(), extd::advanced(storage.inputIndices.begin(), offset));
+        storage.jointGroups[i].inputIndicesOffset = offset;
+        storage.jointGroups[i].inputIndicesSize = jointGroupSize.padded.cols;
+        storage.jointGroups[i].inputIndicesSizeAlignedTo4 = jointGroupSize.padded.cols - (jointGroupSize.padded.cols % 4u);
+        storage.jointGroups[i].inputIndicesSizeAlignedTo8 = jointGroupSize.padded.cols - (jointGroupSize.padded.cols % 8u);
+        offset += jointGroupSize.padded.cols;
+    }
+}
+
+template<typename TValue, std::uint32_t BlockHeight, std::uint32_t PadTo, typename TFVec>
+void JointsBuilderCommon<TValue, BlockHeight, PadTo, TFVec>::setOutputIndices(const dna::BehaviorReader* source) {
+    std::uint32_t offset = 0ul;
+    for (std::uint16_t i = 0u; i < source->getJointGroupCount(); ++i) {
+        const auto jointGroupSize = sizeReqs.getJointGroupSize(i);
+        const auto indices = source->getJointGroupOutputIndices(i);
+        std::copy(indices.begin(), indices.end(), extd::advanced(storage.outputIndices.begin(), offset));
+        storage.jointGroups[i].outputIndicesOffset = offset;
+        offset += jointGroupSize.padded.rows;
+    }
+}
+
+template<typename TValue, std::uint32_t BlockHeight, std::uint32_t PadTo, typename TFVec>
+void JointsBuilderCommon<TValue, BlockHeight, PadTo, TFVec>::setLODs(const dna::BehaviorReader* source) {
+    std::uint32_t offset = 0ul;
+    for (std::uint16_t i = 0u; i < source->getJointGroupCount(); ++i) {
+        const auto jointGroupSize = sizeReqs.getJointGroupSize(i);
+        const auto dstRowCount = jointGroupSize.padded.rows;
+        auto makeLODRegion = [dstRowCount](std::uint16_t lodRowCount) {
+                return LODRegion{lodRowCount, dstRowCount, BlockHeight, PadTo};
+            };
+        const auto lods = source->getJointGroupLODs(i);
+        std::transform(lods.begin(), lods.end(), std::back_inserter(storage.lodRegions), makeLODRegion);
+        storage.jointGroups[i].lodsOffset = offset;
+        offset += source->getLODCount();
+    }
+}
+
+template<typename TValue, std::uint32_t BlockHeight, std::uint32_t PadTo, typename TFVec>
+JointsEvaluator::Pointer JointsBuilderCommon<TValue, BlockHeight, PadTo, TFVec>::build() {
+    const auto attributeCount = static_cast<std::uint16_t>(sizeReqs.attributeCount);
+    auto instanceFactory = [attributeCount](MemoryResource* instanceMemRes) {
+            return UniqueInstance<OutputInstance, JointsOutputInstance>::with(instanceMemRes).create(attributeCount,
+                                                                                                     instanceMemRes);
+        };
+    auto factory = UniqueInstance<Evaluator<TValue>, JointsEvaluator>::with(memRes);
+    return factory.create(std::move(storage), std::move(strategy), instanceFactory, memRes);
+}
+
+}  // namespace bpcm
+
+}  // namespace rl4
diff --git a/Engine/Plugins/Animation/RigLogic/Source/RigLogicLib/Private/riglogic/joints/bpcm/CalculationStrategy.h b/Engine/Plugins/Animation/RigLogic/Source/RigLogicLib/Private/riglogic/joints/bpcm/CalculationStrategy.h
new file mode 100644
--- /dev/null
+++ b/Engine/Plugins/Animation/RigLogic/Source/RigLogicLib/Private/riglogic/joints/bpcm/CalculationStrategy.h
@@ -0,0 +1,34 @@
+// Copyright Epic Games, Inc. All Rights Reserved.
+
+#pragma once
+
+#include "riglogic/joints/bpcm/StorageView.h"
+#include "riglogic/types/Aliases.h"
+
+#include <cstdint>
+
+namespace rl4 {
+
+namespace bpcm {
+
+template<typename TValue>
+struct JointCalculationStrategy {
+    using JointGroupArrayView = ConstArrayView<JointGroupView<TValue> >;
+
+    virtual ~JointCalculationStrategy() = default;
+
+    virtual void calculate(const JointGroupArrayView& jointGroups,
+                           ConstArrayView<float> inputs,
+                           ArrayView<float> outputs,
+                           std::uint16_t lod) const = 0;
+
+    virtual void calculate(const JointGroupArrayView& jointGroups,
+                           ConstArrayView<float> inputs,
+                           ArrayView<float> outputs,
+                           std::uint16_t lod,
+                           std::uint16_t jointGroupIndex) const = 0;
+};
+
+}  // namespace bpcm
+
+}  // namespace rl4
diff --git a/Engine/Plugins/Animation/RigLogic/Source/RigLogicLib/Private/riglogic/joints/bpcm/Evaluator.h b/Engine/Plugins/Animation/RigLogic/Source/RigLogicLib/Private/riglogic/joints/bpcm/Evaluator.h
new file mode 100644
--- /dev/null
+++ b/Engine/Plugins/Animation/RigLogic/Source/RigLogicLib/Private/riglogic/joints/bpcm/Evaluator.h
@@ -0,0 +1,87 @@
+// Copyright Epic Games, Inc. All Rights Reserved.
+
+#pragma once
+
+#include "riglogic/TypeDefs.h"
+#include "riglogic/controls/ControlsInputInstance.h"
+#include "riglogic/joints/JointsEvaluator.h"
+#include "riglogic/joints/JointsOutputInstance.h"
+#include "riglogic/joints/bpcm/CalculationStrategy.h"
+#include "riglogic/joints/bpcm/Storage.h"
+#include "riglogic/joints/bpcm/StorageView.h"
+#include "riglogic/riglogic/RigInstanceImpl.h"
+
+#include <cstdint>
+
+namespace rl4 {
+
+namespace bpcm {
+
+template<typename TValue>
+class Evaluator : public JointsEvaluator {
+    public:
+        using CalculationStrategy = JointCalculationStrategy<TValue>;
+        using CalculationStrategyPtr = typename UniqueInstance<CalculationStrategy>::PointerType;
+
+        struct Accessor;
+        friend Accessor;
+
+    private:
+        using JointGroupArrayView = typename CalculationStrategy::JointGroupArrayView;
+
+    public:
+        Evaluator(JointStorage<TValue>&& storage_,
+                  CalculationStrategyPtr strategy_,
+                  JointsOutputInstance::Factory instanceFactory_,
+                  MemoryResource* memRes_) :
+            memRes{memRes_},
+            storage{std::move(storage_)},
+            storageView{takeStorageSnapshot(storage, memRes)},
+            jointGroupsView{storageView},
+            strategy{std::move(strategy_)},
+            instanceFactory{instanceFactory_} {
+        }
+
+        JointsOutputInstance::Pointer createInstance(MemoryResource* instanceMemRes) const override {
+            return instanceFactory(instanceMemRes);
+        }
+
+        void calculate(const ControlsInputInstance* inputs, JointsOutputInstance* outputs, std::uint16_t lod) const override {
+            assert(strategy != nullptr);
+            strategy->calculate(jointGroupsView, inputs->getInputBuffer(), outputs->getOutputBuffer(), lod);
+        }
+
+        void calculate(const ControlsInputInstance* inputs,
+                       JointsOutputInstance* outputs,
+                       std::uint16_t lod,
+                       std::uint16_t jointGroupIndex) const override {
+            assert(strategy != nullptr);
+            strategy->calculate(jointGroupsView,
+                                inputs->getInputBuffer(),
+                                outputs->getOutputBuffer(),
+                                lod,
+                                jointGroupIndex);
+        }
+
+        void load(terse::BinaryInputArchive<BoundedIOStream>& archive) override {
+            archive(storage);
+            storageView = takeStorageSnapshot(storage, memRes);
+            jointGroupsView = JointGroupArrayView{storageView};
+        }
+
+        void save(terse::BinaryOutputArchive<BoundedIOStream>& archive) override {
+            archive(storage);
+        }
+
+    private:
+        MemoryResource* memRes;
+        JointStorage<TValue> storage;
+        Vector<JointGroupView<TValue> > storageView;
+        JointGroupArrayView jointGroupsView;
+        CalculationStrategyPtr strategy;
+        JointsOutputInstance::Factory instanceFactory;
+};
+
+}  // namespace bpcm
+
+}  // namespace rl4
diff --git a/Engine/Plugins/Animation/RigLogic/Source/RigLogicLib/Private/riglogic/joints/bpcm/LODRegion.h b/Engine/Plugins/Animation/RigLogic/Source/RigLogicLib/Private/riglogic/joints/bpcm/LODRegion.h
--- a/Engine/Plugins/Animation/RigLogic/Source/RigLogicLib/Private/riglogic/joints/bpcm/LODRegion.h
+++ b/Engine/Plugins/Animation/RigLogic/Source/RigLogicLib/Private/riglogic/joints/bpcm/LODRegion.h
@@ -2,34 +2,15 @@
 
 #pragma once
 
+#include "riglogic/types/AlignedBlockView.h"
+
 #include <cstdint>
 
 namespace rl4 {
 
 namespace bpcm {
 
-struct LODRegion {
-    // Unaligned number of output indices
-    std::uint32_t size;
-    // If the last full block needs to be masked-off because of LODs,
-    // this view will include that last full block as well
-    std::uint32_t sizeAlignedToLastFullBlock;
-    // If the last full block needs to be masked-off because of LODs,
-    // that full block won't be included in this boundary as it needs special care
-    std::uint32_t sizeAlignedToSecondLastFullBlock;
-
-    LODRegion() = default;
-    LODRegion(std::uint32_t lodRowCount, std::uint32_t rowCount, std::uint32_t blockHeight, std::uint32_t padTo);
-    LODRegion(std::uint32_t lodRowCount,
-              std::uint32_t lodRowCountAlignedToLastFullBlock,
-              std::uint32_t lodRowCountAlignedToSecondLastFullBlock);
-
-    template<class Archive>
-    void serialize(Archive& archive) {
-        archive(size, sizeAlignedToLastFullBlock, sizeAlignedToSecondLastFullBlock);
-    }
-
-};
+using LODRegion = AlignedBlockView;
 
 }  // namespace bpcm
 
diff --git a/Engine/Plugins/Animation/RigLogic/Source/RigLogicLib/Private/riglogic/joints/bpcm/StorageSize.cpp b/Engine/Plugins/Animation/RigLogic/Source/RigLogicLib/Private/riglogic/joints/bpcm/StorageSize.cpp
--- a/Engine/Plugins/Animation/RigLogic/Source/RigLogicLib/Private/riglogic/joints/bpcm/StorageSize.cpp
+++ b/Engine/Plugins/Animation/RigLogic/Source/RigLogicLib/Private/riglogic/joints/bpcm/StorageSize.cpp
@@ -2,11 +2,10 @@
 
 #include "riglogic/joints/bpcm/StorageSize.h"
 
-#include "riglogic/joints/bpcm/Extent.h"
+#include "riglogic/TypeDefs.h"
+#include "riglogic/types/Extent.h"
 #include "riglogic/utils/Extd.h"
 
-#include <dna/layers/BehaviorReader.h>
-
 #include <cassert>
 #include <cstddef>
 
@@ -15,6 +14,7 @@ namespace rl4 {
 namespace bpcm {
 
 StorageSize::StorageSize(MemoryResource* memRes) :
+    attributeCount{},
     valueCount{},
     inputIndexCount{},
     outputIndexCount{},
@@ -27,6 +27,7 @@ void StorageSize::computeFrom(const dna::BehaviorReader* src, std::size_t padTo)
     const auto jointGroupCount = src->getJointGroupCount();
     jointGroups.clear();
     jointGroups.reserve(jointGroupCount);
+    attributeCount = src->getJointRowCount();
     lodCount = src->getLODCount();
     for (std::uint16_t i = 0u; i < jointGroupCount; ++i) {
         assert(src->getJointGroupLODs(i).size() == lodCount);
@@ -38,8 +39,8 @@ void StorageSize::computeFrom(const dna::BehaviorReader* src, std::size_t padTo)
         outputIndexCount += (rowCount + padding);
         lodRegionCount += lodCount;
         jointGroups.push_back({
-                    {static_cast<std::uint16_t>(rowCount), static_cast<std::uint16_t>(columnCount)},
-                    {static_cast<std::uint16_t>(rowCount + padding), static_cast<std::uint16_t>(columnCount)}
+                    {static_cast<std::uint32_t>(rowCount), static_cast<std::uint32_t>(columnCount)},
+                    {static_cast<std::uint32_t>(rowCount + padding), static_cast<std::uint32_t>(columnCount)}
                 });
     }
 }
diff --git a/Engine/Plugins/Animation/RigLogic/Source/RigLogicLib/Private/riglogic/joints/bpcm/StorageSize.h b/Engine/Plugins/Animation/RigLogic/Source/RigLogicLib/Private/riglogic/joints/bpcm/StorageSize.h
--- a/Engine/Plugins/Animation/RigLogic/Source/RigLogicLib/Private/riglogic/joints/bpcm/StorageSize.h
+++ b/Engine/Plugins/Animation/RigLogic/Source/RigLogicLib/Private/riglogic/joints/bpcm/StorageSize.h
@@ -3,7 +3,7 @@
 #pragma once
 
 #include "riglogic/TypeDefs.h"
-#include "riglogic/joints/bpcm/Extent.h"
+#include "riglogic/types/Extent.h"
 
 #include <cstddef>
 
@@ -17,6 +17,7 @@ struct JointGroupSize {
 };
 
 struct StorageSize {
+    std::size_t attributeCount;
     std::size_t valueCount;
     std::size_t inputIndexCount;
     std::size_t outputIndexCount;
diff --git a/Engine/Plugins/Animation/RigLogic/Source/RigLogicLib/Private/riglogic/joints/bpcm/StorageView.h b/Engine/Plugins/Animation/RigLogic/Source/RigLogicLib/Private/riglogic/joints/bpcm/StorageView.h
new file mode 100644
--- /dev/null
+++ b/Engine/Plugins/Animation/RigLogic/Source/RigLogicLib/Private/riglogic/joints/bpcm/StorageView.h
@@ -0,0 +1,43 @@
+// Copyright Epic Games, Inc. All Rights Reserved.
+
+#pragma once
+
+#include "riglogic/TypeDefs.h"
+#include "riglogic/joints/bpcm/Storage.h"
+
+#include <cstdint>
+
+namespace rl4 {
+
+namespace bpcm {
+
+template<typename TValue>
+struct JointGroupView {
+    TValue* values;
+    std::uint16_t* inputIndices;
+    std::uint16_t* inputIndicesEnd;
+    std::uint16_t* inputIndicesEndAlignedTo4;
+    std::uint16_t* inputIndicesEndAlignedTo8;
+    std::uint16_t* outputIndices;
+    LODRegion* lods;
+};
+
+template<typename TValue>
+Vector<JointGroupView<TValue> > takeStorageSnapshot(JointStorage<TValue>& storage, MemoryResource* memRes) {
+    Vector<JointGroupView<TValue> > snapshot{storage.jointGroups.size(), {}, memRes};
+    for (std::size_t i = 0ul; i < storage.jointGroups.size(); ++i) {
+        const auto& jointGroup = storage.jointGroups[i];
+        snapshot[i].values = storage.values.data() + jointGroup.valuesOffset;
+        snapshot[i].inputIndices = storage.inputIndices.data() + jointGroup.inputIndicesOffset;
+        snapshot[i].inputIndicesEnd = snapshot[i].inputIndices + jointGroup.inputIndicesSize;
+        snapshot[i].inputIndicesEndAlignedTo4 = snapshot[i].inputIndices + jointGroup.inputIndicesSizeAlignedTo4;
+        snapshot[i].inputIndicesEndAlignedTo8 = snapshot[i].inputIndices + jointGroup.inputIndicesSizeAlignedTo8;
+        snapshot[i].outputIndices = storage.outputIndices.data() + jointGroup.outputIndicesOffset;
+        snapshot[i].lods = storage.lodRegions.data() + jointGroup.lodsOffset;
+    }
+    return snapshot;
+}
+
+}  // namespace bpcm
+
+}  // namespace rl4
diff --git a/Engine/Plugins/Animation/RigLogic/Source/RigLogicLib/Private/riglogic/joints/bpcm/builders/AVX.cpp b/Engine/Plugins/Animation/RigLogic/Source/RigLogicLib/Private/riglogic/joints/bpcm/builders/AVX.cpp
--- a/Engine/Plugins/Animation/RigLogic/Source/RigLogicLib/Private/riglogic/joints/bpcm/builders/AVX.cpp
+++ b/Engine/Plugins/Animation/RigLogic/Source/RigLogicLib/Private/riglogic/joints/bpcm/builders/AVX.cpp
@@ -1,25 +1,14 @@
 // Copyright Epic Games, Inc. All Rights Reserved.
 
-#if !defined(RL_USE_HALF_FLOATS) && defined(RL_BUILD_WITH_AVX)
-
 #include "riglogic/joints/bpcm/builders/AVX.h"
 
-#include "riglogic/joints/bpcm/Consts.h"
-#include "riglogic/joints/bpcm/strategies/AVX.h"
-
-#include <pma/utils/ManagedInstance.h>
-
-#include <cstddef>
-#include <cstdint>
+// *INDENT-OFF*
+#if !defined(RL_USE_HALF_FLOATS) && defined(RL_BUILD_WITH_AVX)
 
 namespace rl4 {
 
 namespace bpcm {
 
-AVXJointsBuilder::AVXJointsBuilder(MemoryResource* memRes_) : FloatStorageBuilder{block16Height, block8Height, memRes_} {
-    strategy = pma::UniqueInstance<AVXJointCalculationStrategy<float>, CalculationStrategy>::with(memRes_).create();
-}
-
 AVXJointsBuilder::~AVXJointsBuilder() = default;
 
 }  // namespace bpcm
@@ -27,3 +16,4 @@ AVXJointsBuilder::~AVXJointsBuilder() = default;
 }  // namespace rl4
 
 #endif  // !defined(RL_USE_HALF_FLOATS) && defined(RL_BUILD_WITH_AVX)
+// *INDENT-ON*
diff --git a/Engine/Plugins/Animation/RigLogic/Source/RigLogicLib/Private/riglogic/joints/bpcm/builders/AVX.h b/Engine/Plugins/Animation/RigLogic/Source/RigLogicLib/Private/riglogic/joints/bpcm/builders/AVX.h
--- a/Engine/Plugins/Animation/RigLogic/Source/RigLogicLib/Private/riglogic/joints/bpcm/builders/AVX.h
+++ b/Engine/Plugins/Animation/RigLogic/Source/RigLogicLib/Private/riglogic/joints/bpcm/builders/AVX.h
@@ -5,16 +5,23 @@
 // *INDENT-OFF*
 #if !defined(RL_USE_HALF_FLOATS) && defined(RL_BUILD_WITH_AVX)
 
-#include "riglogic/joints/bpcm/builders/Float.h"
+#include "riglogic/joints/bpcm/BuilderCommon.h"
+#include "riglogic/joints/bpcm/Consts.h"
+#include "riglogic/joints/bpcm/strategies/AVX.h"
 #include "riglogic/types/Aliases.h"
 
 namespace rl4 {
 
 namespace bpcm {
 
-class AVXJointsBuilder : public FloatStorageBuilder {
+class AVXJointsBuilder : public JointsBuilderCommon<float, block16Height, block8Height, trimd::avx::F256> {
+    private:
+        using Super = JointsBuilderCommon<float, block16Height, block8Height, trimd::avx::F256>;
+
     public:
-        explicit AVXJointsBuilder(MemoryResource* memRes_);
+        explicit AVXJointsBuilder(MemoryResource* memRes_) : Super{UniqueInstance<AVXJointCalculationStrategy<float>, CalculationStrategy>::with(memRes_).create(), memRes_} {
+        }
+
         ~AVXJointsBuilder();
 
 };
diff --git a/Engine/Plugins/Animation/RigLogic/Source/RigLogicLib/Private/riglogic/joints/bpcm/builders/AVXHalfFloat.cpp b/Engine/Plugins/Animation/RigLogic/Source/RigLogicLib/Private/riglogic/joints/bpcm/builders/AVXHalfFloat.cpp
--- a/Engine/Plugins/Animation/RigLogic/Source/RigLogicLib/Private/riglogic/joints/bpcm/builders/AVXHalfFloat.cpp
+++ b/Engine/Plugins/Animation/RigLogic/Source/RigLogicLib/Private/riglogic/joints/bpcm/builders/AVXHalfFloat.cpp
@@ -1,100 +1,19 @@
 // Copyright Epic Games, Inc. All Rights Reserved.
 
-#if defined(RL_USE_HALF_FLOATS) && defined(RL_BUILD_WITH_AVX)
-
 #include "riglogic/joints/bpcm/builders/AVXHalfFloat.h"
 
-#include "riglogic/joints/bpcm/Consts.h"
-#include "riglogic/joints/bpcm/strategies/AVX.h"
-#include "riglogic/utils/Extd.h"
-
-#include <dna/layers/BehaviorReader.h>
-#include <pma/utils/ManagedInstance.h>
-
-#include <immintrin.h>
-
-#include <cstddef>
-#include <cstdint>
+// *INDENT-OFF*
+#if defined(RL_USE_HALF_FLOATS) && defined(RL_BUILD_WITH_AVX)
 
 namespace rl4 {
 
 namespace bpcm {
 
-inline void toHalfFloat256(const float* source, std::uint16_t* destination) {
-    #if defined(__clang__) || defined(__GNUC__)
-        #pragma GCC diagnostic push
-        #pragma GCC diagnostic ignored "-Wold-style-cast"
-        #if !defined(__clang__)
-            #pragma GCC diagnostic ignored "-Wuseless-cast"
-        #endif
-    #endif
-    __m256 block = _mm256_load_ps(source);
-    __m128i half = _mm256_cvtps_ph(block, _MM_FROUND_CUR_DIRECTION);
-    _mm_store_si128(reinterpret_cast<__m128i*>(destination), half);
-    #if defined(__clang__) || defined(__GNUC__)
-        #pragma GCC diagnostic pop
-    #endif
-}
-
-AVXJointsBuilder::AVXJointsBuilder(MemoryResource* memRes_) : JointsBuilderCommon{block16Height, block8Height, memRes_} {
-    strategy = pma::UniqueInstance<AVXJointCalculationStrategy<std::uint16_t>, CalculationStrategy>::with(memRes_).create();
-}
-
 AVXJointsBuilder::~AVXJointsBuilder() = default;
 
-void AVXJointsBuilder::setValues(const dna::BehaviorReader* reader) {
-    std::uint32_t offset = 0ul;
-    for (std::uint16_t i = 0u; i < reader->getJointGroupCount(); ++i) {
-        const auto source = reader->getJointGroupValues(i);
-        const auto jointGroupSize = sizeReqs.getJointGroupSize(i);
-        storage.jointGroups[i].valuesOffset = offset;
-        storage.jointGroups[i].valuesSize = jointGroupSize.padded.size();
-
-        std::uint32_t remainder = jointGroupSize.original.rows % block16Height;
-        std::uint32_t target = jointGroupSize.original.rows - remainder;
-        // Copy the portion of the source matrix that is partitionable into block16Height chunks
-        for (std::size_t row = 0ul; row < target; row += block16Height) {
-            for (std::size_t col = 0ul; col < jointGroupSize.original.cols; ++col) {
-                alignas(32) float fBlock[block16Height];
-                alignas(16) std::uint16_t hfBlock[block16Height];
-
-                for (std::size_t blkIdx = 0ul; blkIdx < block16Height; ++blkIdx) {
-                    fBlock[blkIdx] = source[(row + blkIdx) * jointGroupSize.original.cols + col];
-                }
-
-                toHalfFloat256(fBlock, hfBlock);
-                toHalfFloat256(fBlock + block8Height, hfBlock + block8Height);
-
-                for (std::size_t blkIdx = 0ul; blkIdx < block16Height; ++blkIdx, ++offset) {
-                    storage.values[offset] = hfBlock[blkIdx];
-                }
-            }
-        }
-        // Copy the non-blockHeight sized remainder portion of the source matrix
-        for (std::size_t row = target; row < jointGroupSize.original.rows; row += remainder) {
-            for (std::size_t col = 0ul; col < jointGroupSize.original.cols; ++col) {
-                alignas(32) float fBlock[block16Height] = {};
-                alignas(16) std::uint16_t hfBlock[block16Height] = {};
-
-                for (std::size_t blkIdx = 0ul; blkIdx < remainder; ++blkIdx) {
-                    fBlock[blkIdx] = source[(row + blkIdx) * jointGroupSize.original.cols + col];
-                }
-
-                toHalfFloat256(fBlock, hfBlock);
-                toHalfFloat256(fBlock + block8Height, hfBlock + block8Height);
-
-                // Copy even the padded zero values as the representation of half-floats might
-                // differ from the storage's zero-initialized state
-                for (std::size_t blkIdx = 0ul; blkIdx < extd::roundUp(remainder, block8Height); ++blkIdx, ++offset) {
-                    storage.values[offset] = hfBlock[blkIdx];
-                }
-            }
-        }
-    }
-}
-
 }  // namespace bpcm
 
 }  // namespace rl4
 
 #endif  // defined(RL_USE_HALF_FLOATS) && defined(RL_BUILD_WITH_AVX)
+// *INDENT-ON*
diff --git a/Engine/Plugins/Animation/RigLogic/Source/RigLogicLib/Private/riglogic/joints/bpcm/builders/AVXHalfFloat.h b/Engine/Plugins/Animation/RigLogic/Source/RigLogicLib/Private/riglogic/joints/bpcm/builders/AVXHalfFloat.h
--- a/Engine/Plugins/Animation/RigLogic/Source/RigLogicLib/Private/riglogic/joints/bpcm/builders/AVXHalfFloat.h
+++ b/Engine/Plugins/Animation/RigLogic/Source/RigLogicLib/Private/riglogic/joints/bpcm/builders/AVXHalfFloat.h
@@ -5,7 +5,9 @@
 // *INDENT-OFF*
 #if defined(RL_USE_HALF_FLOATS) && defined(RL_BUILD_WITH_AVX)
 
-#include "riglogic/joints/bpcm/JointsBuilderCommon.h"
+#include "riglogic/joints/bpcm/BuilderCommon.h"
+#include "riglogic/joints/bpcm/Consts.h"
+#include "riglogic/joints/bpcm/strategies/AVX.h"
 #include "riglogic/types/Aliases.h"
 
 #include <cstdint>
@@ -14,14 +16,15 @@ namespace rl4 {
 
 namespace bpcm {
 
-class AVXJointsBuilder : public JointsBuilderCommon<std::uint16_t> {
-    public:
-        explicit AVXJointsBuilder(MemoryResource* memRes_);
-        ~AVXJointsBuilder();
+class AVXJointsBuilder : public JointsBuilderCommon<std::uint16_t, block16Height, block8Height, trimd::avx::F256> {
+    private:
+        using Super = JointsBuilderCommon<std::uint16_t, block16Height, block8Height, trimd::avx::F256>;
 
-    protected:
-        void setValues(const dna::BehaviorReader* reader) override;
+    public:
+        explicit AVXJointsBuilder(MemoryResource* memRes_) : Super{UniqueInstance<AVXJointCalculationStrategy<std::uint16_t>, CalculationStrategy>::with(memRes_).create(), memRes_} {
+        }
 
+        ~AVXJointsBuilder();
 };
 
 }  // namespace bpcm
diff --git a/Engine/Plugins/Animation/RigLogic/Source/RigLogicLib/Private/riglogic/joints/bpcm/builders/SSE.cpp b/Engine/Plugins/Animation/RigLogic/Source/RigLogicLib/Private/riglogic/joints/bpcm/builders/SSE.cpp
--- a/Engine/Plugins/Animation/RigLogic/Source/RigLogicLib/Private/riglogic/joints/bpcm/builders/SSE.cpp
+++ b/Engine/Plugins/Animation/RigLogic/Source/RigLogicLib/Private/riglogic/joints/bpcm/builders/SSE.cpp
@@ -1,25 +1,14 @@
 // Copyright Epic Games, Inc. All Rights Reserved.
 
-#if !defined(RL_USE_HALF_FLOATS) && defined(RL_BUILD_WITH_SSE)
-
 #include "riglogic/joints/bpcm/builders/SSE.h"
 
-#include "riglogic/joints/bpcm/Consts.h"
-#include "riglogic/joints/bpcm/strategies/SSE.h"
-
-#include <pma/utils/ManagedInstance.h>
-
-#include <cstddef>
-#include <cstdint>
+// *INDENT-OFF*
+#if !defined(RL_USE_HALF_FLOATS) && defined(RL_BUILD_WITH_SSE)
 
 namespace rl4 {
 
 namespace bpcm {
 
-SSEJointsBuilder::SSEJointsBuilder(MemoryResource* memRes_) : FloatStorageBuilder{block8Height, block4Height, memRes_} {
-    strategy = pma::UniqueInstance<SSEJointCalculationStrategy<float>, CalculationStrategy>::with(memRes).create();
-}
-
 SSEJointsBuilder::~SSEJointsBuilder() = default;
 
 }  // namespace bpcm
@@ -27,3 +16,4 @@ SSEJointsBuilder::~SSEJointsBuilder() = default;
 }  // namespace rl4
 
 #endif  // !defined(RL_USE_HALF_FLOATS) && defined(RL_BUILD_WITH_SSE)
+// *INDENT-ON*
diff --git a/Engine/Plugins/Animation/RigLogic/Source/RigLogicLib/Private/riglogic/joints/bpcm/builders/SSE.h b/Engine/Plugins/Animation/RigLogic/Source/RigLogicLib/Private/riglogic/joints/bpcm/builders/SSE.h
--- a/Engine/Plugins/Animation/RigLogic/Source/RigLogicLib/Private/riglogic/joints/bpcm/builders/SSE.h
+++ b/Engine/Plugins/Animation/RigLogic/Source/RigLogicLib/Private/riglogic/joints/bpcm/builders/SSE.h
@@ -2,20 +2,30 @@
 
 #pragma once
 
+#include "riglogic/system/simd/Detect.h"
+
 // *INDENT-OFF*
 #if !defined(RL_USE_HALF_FLOATS) && defined(RL_BUILD_WITH_SSE)
 
-#include "riglogic/joints/bpcm/builders/Float.h"
+#include "riglogic/joints/bpcm/BuilderCommon.h"
+#include "riglogic/joints/bpcm/Consts.h"
+#include "riglogic/joints/bpcm/strategies/SSE.h"
 #include "riglogic/types/Aliases.h"
 
 namespace rl4 {
 
 namespace bpcm {
 
-class SSEJointsBuilder : public FloatStorageBuilder {
+class SSEJointsBuilder : public JointsBuilderCommon<float, block8Height, block4Height, trimd::sse::F128> {
+    private:
+        using Super = JointsBuilderCommon<float, block8Height, block4Height, trimd::sse::F128>;
+
     public:
-        explicit SSEJointsBuilder(MemoryResource* memRes_);
+        explicit SSEJointsBuilder(MemoryResource* memRes_) : Super{UniqueInstance<SSEJointCalculationStrategy<float>, CalculationStrategy>::with(memRes_).create(), memRes_} {
+        }
+
         ~SSEJointsBuilder();
+
 };
 
 }  // namespace bpcm
