Update RigLogic to v9.1.1 and GeneSplicer to v9.4.0

diff --git a/Engine/Plugins/Animation/RigLogic/Source/RigLogicLib/Private/riglogic/joints/bpcm/builders/SSEHalfFloat.cpp b/Engine/Plugins/Animation/RigLogic/Source/RigLogicLib/Private/riglogic/joints/bpcm/builders/SSEHalfFloat.cpp
--- a/Engine/Plugins/Animation/RigLogic/Source/RigLogicLib/Private/riglogic/joints/bpcm/builders/SSEHalfFloat.cpp
+++ b/Engine/Plugins/Animation/RigLogic/Source/RigLogicLib/Private/riglogic/joints/bpcm/builders/SSEHalfFloat.cpp
@@ -1,100 +1,19 @@
 // Copyright Epic Games, Inc. All Rights Reserved.
 
-#if defined(RL_USE_HALF_FLOATS) && defined(RL_BUILD_WITH_SSE)
-
 #include "riglogic/joints/bpcm/builders/SSEHalfFloat.h"
 
-#include "riglogic/joints/bpcm/Consts.h"
-#include "riglogic/joints/bpcm/strategies/SSE.h"
-#include "riglogic/utils/Extd.h"
-
-#include <dna/layers/BehaviorReader.h>
-#include <pma/utils/ManagedInstance.h>
-
-#include <immintrin.h>
-
-#include <cstddef>
-#include <cstdint>
+// *INDENT-OFF*
+#if defined(RL_USE_HALF_FLOATS) && defined(RL_BUILD_WITH_SSE)
 
 namespace rl4 {
 
 namespace bpcm {
 
-inline void toHalfFloat128(const float* source, std::uint16_t* destination) {
-    #if defined(__clang__) || defined(__GNUC__)
-        #pragma GCC diagnostic push
-        #pragma GCC diagnostic ignored "-Wold-style-cast"
-        #if !defined(__clang__)
-            #pragma GCC diagnostic ignored "-Wuseless-cast"
-        #endif
-    #endif
-    __m128 block = _mm_load_ps(source);
-    __m128i half = _mm_cvtps_ph(block, _MM_FROUND_CUR_DIRECTION);
-    _mm_storel_epi64(reinterpret_cast<__m128i*>(destination), half);
-    #if defined(__clang__) || defined(__GNUC__)
-        #pragma GCC diagnostic pop
-    #endif
-}
-
 SSEJointsBuilder::~SSEJointsBuilder() = default;
 
-SSEJointsBuilder::SSEJointsBuilder(MemoryResource* memRes_) : JointsBuilderCommon{block8Height, block4Height, memRes_} {
-    strategy = pma::UniqueInstance<SSEJointCalculationStrategy<std::uint16_t>, CalculationStrategy>::with(memRes_).create();
-}
-
-void SSEJointsBuilder::setValues(const dna::BehaviorReader* reader) {
-    std::uint32_t offset = 0ul;
-    for (std::uint16_t i = 0u; i < reader->getJointGroupCount(); ++i) {
-        const auto source = reader->getJointGroupValues(i);
-        const auto jointGroupSize = sizeReqs.jointGroups[i];
-        storage.jointGroups[i].valuesOffset = offset;
-        storage.jointGroups[i].valuesSize = jointGroupSize.padded.size();
-
-        std::uint32_t remainder = jointGroupSize.original.rows % block8Height;
-        std::uint32_t target = jointGroupSize.original.rows - remainder;
-        // Copy the portion of the source matrix that is partitionable into block8Height chunks
-        for (std::size_t row = 0ul; row < target; row += block8Height) {
-            for (std::size_t col = 0ul; col < jointGroupSize.original.cols; ++col) {
-                alignas(16) float fBlock[block8Height];
-                alignas(16) std::uint16_t hfBlock[block8Height];
-
-                for (std::size_t blkIdx = 0ul; blkIdx < block8Height; ++blkIdx) {
-                    fBlock[blkIdx] = source[(row + blkIdx) * jointGroupSize.original.cols + col];
-                }
-
-                toHalfFloat128(fBlock, hfBlock);
-                toHalfFloat128(fBlock + block4Height, hfBlock + block4Height);
-
-                for (std::size_t blkIdx = 0ul; blkIdx < block8Height; ++blkIdx, ++offset) {
-                    storage.values[offset] = hfBlock[blkIdx];
-                }
-            }
-        }
-        // Copy the non-block8Height sized remainder portion of the source matrix
-        for (std::size_t row = target; row < jointGroupSize.original.rows; row += remainder) {
-            for (std::size_t col = 0ul; col < jointGroupSize.original.cols; ++col) {
-                alignas(16) float fBlock[block8Height] = {};
-                alignas(16) std::uint16_t hfBlock[block8Height] = {};
-
-                for (std::size_t blkIdx = 0ul; blkIdx < remainder; ++blkIdx) {
-                    fBlock[blkIdx] = source[(row + blkIdx) * jointGroupSize.original.cols + col];
-                }
-
-                toHalfFloat128(fBlock, hfBlock);
-                toHalfFloat128(fBlock + block4Height, hfBlock + block4Height);
-
-                // Copy even the padded zero values as the representation of half-floats might
-                // differ from the storage's zero-initialized state
-                for (std::size_t blkIdx = 0ul; blkIdx < extd::roundUp(remainder, block4Height); ++blkIdx, ++offset) {
-                    storage.values[offset] = hfBlock[blkIdx];
-                }
-            }
-        }
-    }
-}
-
 }  // namespace bpcm
 
 }  // namespace rl4
 
 #endif  // defined(RL_USE_HALF_FLOATS) && defined(RL_BUILD_WITH_SSE)
+// *INDENT-ON*
diff --git a/Engine/Plugins/Animation/RigLogic/Source/RigLogicLib/Private/riglogic/joints/bpcm/builders/SSEHalfFloat.h b/Engine/Plugins/Animation/RigLogic/Source/RigLogicLib/Private/riglogic/joints/bpcm/builders/SSEHalfFloat.h
--- a/Engine/Plugins/Animation/RigLogic/Source/RigLogicLib/Private/riglogic/joints/bpcm/builders/SSEHalfFloat.h
+++ b/Engine/Plugins/Animation/RigLogic/Source/RigLogicLib/Private/riglogic/joints/bpcm/builders/SSEHalfFloat.h
@@ -2,10 +2,14 @@
 
 #pragma once
 
+#include "riglogic/system/simd/Detect.h"
+
 // *INDENT-OFF*
 #if defined(RL_USE_HALF_FLOATS) && defined(RL_BUILD_WITH_SSE)
 
-#include "riglogic/joints/bpcm/JointsBuilderCommon.h"
+#include "riglogic/joints/bpcm/BuilderCommon.h"
+#include "riglogic/joints/bpcm/Consts.h"
+#include "riglogic/joints/bpcm/strategies/SSE.h"
 #include "riglogic/types/Aliases.h"
 
 #include <cstdint>
@@ -14,13 +18,15 @@ namespace rl4 {
 
 namespace bpcm {
 
-class SSEJointsBuilder : public JointsBuilderCommon<std::uint16_t> {
+class SSEJointsBuilder : public JointsBuilderCommon<std::uint16_t, block8Height, block4Height, trimd::sse::F128> {
+    private:
+        using Super = JointsBuilderCommon<std::uint16_t, block8Height, block4Height, trimd::sse::F128>;
+
     public:
-        explicit SSEJointsBuilder(MemoryResource* memRes_);
-        ~SSEJointsBuilder();
+        explicit SSEJointsBuilder(MemoryResource* memRes_) : Super{UniqueInstance<SSEJointCalculationStrategy<std::uint16_t>, CalculationStrategy>::with(memRes_).create(), memRes_} {
+        }
 
-    protected:
-        void setValues(const dna::BehaviorReader* reader) override;
+        ~SSEJointsBuilder();
 
 };
 
diff --git a/Engine/Plugins/Animation/RigLogic/Source/RigLogicLib/Private/riglogic/joints/bpcm/builders/Scalar.cpp b/Engine/Plugins/Animation/RigLogic/Source/RigLogicLib/Private/riglogic/joints/bpcm/builders/Scalar.cpp
--- a/Engine/Plugins/Animation/RigLogic/Source/RigLogicLib/Private/riglogic/joints/bpcm/builders/Scalar.cpp
+++ b/Engine/Plugins/Animation/RigLogic/Source/RigLogicLib/Private/riglogic/joints/bpcm/builders/Scalar.cpp
@@ -2,23 +2,10 @@
 
 #include "riglogic/joints/bpcm/builders/Scalar.h"
 
-#include "riglogic/joints/bpcm/Consts.h"
-#include "riglogic/joints/bpcm/strategies/Scalar.h"
-
-#include <pma/utils/ManagedInstance.h>
-
-#include <cstddef>
-#include <cstdint>
-
 namespace rl4 {
 
 namespace bpcm {
 
-ScalarJointsBuilder::ScalarJointsBuilder(MemoryResource* memRes_) : FloatStorageBuilder{block8Height, block4Height,
-                                                                                        memRes_} {
-    strategy = pma::UniqueInstance<ScalarJointCalculationStrategy, CalculationStrategy>::with(memRes).create();
-}
-
 ScalarJointsBuilder::~ScalarJointsBuilder() = default;
 
 }  // namespace bpcm
diff --git a/Engine/Plugins/Animation/RigLogic/Source/RigLogicLib/Private/riglogic/joints/bpcm/builders/Scalar.h b/Engine/Plugins/Animation/RigLogic/Source/RigLogicLib/Private/riglogic/joints/bpcm/builders/Scalar.h
--- a/Engine/Plugins/Animation/RigLogic/Source/RigLogicLib/Private/riglogic/joints/bpcm/builders/Scalar.h
+++ b/Engine/Plugins/Animation/RigLogic/Source/RigLogicLib/Private/riglogic/joints/bpcm/builders/Scalar.h
@@ -2,17 +2,27 @@
 
 #pragma once
 
-#include "riglogic/joints/bpcm/builders/Float.h"
+#include "riglogic/joints/bpcm/BuilderCommon.h"
+#include "riglogic/joints/bpcm/Consts.h"
+#include "riglogic/joints/bpcm/strategies/Scalar.h"
 #include "riglogic/types/Aliases.h"
 
 namespace rl4 {
 
 namespace bpcm {
 
-class ScalarJointsBuilder : public FloatStorageBuilder {
+class ScalarJointsBuilder : public JointsBuilderCommon<float, block8Height, block4Height, trimd::scalar::F128>  {
+    private:
+        using Super = JointsBuilderCommon<float, block8Height, block4Height, trimd::scalar::F128>;
+
     public:
-        explicit ScalarJointsBuilder(MemoryResource* memRes_);
+        explicit ScalarJointsBuilder(MemoryResource* memRes_) : Super{UniqueInstance<ScalarJointCalculationStrategy,
+                                                                                     CalculationStrategy>::with(
+                                                                          memRes_).create(), memRes_} {
+        }
+
         ~ScalarJointsBuilder();
+
 };
 
 }  // namespace bpcm
diff --git a/Engine/Plugins/Animation/RigLogic/Source/RigLogicLib/Private/riglogic/joints/bpcm/strategies/Block4.h b/Engine/Plugins/Animation/RigLogic/Source/RigLogicLib/Private/riglogic/joints/bpcm/strategies/Block4.h
--- a/Engine/Plugins/Animation/RigLogic/Source/RigLogicLib/Private/riglogic/joints/bpcm/strategies/Block4.h
+++ b/Engine/Plugins/Animation/RigLogic/Source/RigLogicLib/Private/riglogic/joints/bpcm/strategies/Block4.h
@@ -3,9 +3,9 @@
 #pragma once
 
 #include "riglogic/joints/bpcm/Consts.h"
-#include "riglogic/joints/bpcm/JointCalculationStrategy.h"
+#include "riglogic/joints/bpcm/CalculationStrategy.h"
 #include "riglogic/joints/bpcm/JointGroup.h"
-#include "riglogic/joints/bpcm/JointStorageView.h"
+#include "riglogic/joints/bpcm/StorageView.h"
 #include "riglogic/joints/bpcm/LODRegion.h"
 #include "riglogic/utils/Macros.h"
 
diff --git a/Engine/Plugins/Animation/RigLogic/Source/RigLogicLib/Private/riglogic/joints/bpcm/strategies/Block8.h b/Engine/Plugins/Animation/RigLogic/Source/RigLogicLib/Private/riglogic/joints/bpcm/strategies/Block8.h
--- a/Engine/Plugins/Animation/RigLogic/Source/RigLogicLib/Private/riglogic/joints/bpcm/strategies/Block8.h
+++ b/Engine/Plugins/Animation/RigLogic/Source/RigLogicLib/Private/riglogic/joints/bpcm/strategies/Block8.h
@@ -3,9 +3,9 @@
 #pragma once
 
 #include "riglogic/joints/bpcm/Consts.h"
-#include "riglogic/joints/bpcm/JointCalculationStrategy.h"
+#include "riglogic/joints/bpcm/CalculationStrategy.h"
 #include "riglogic/joints/bpcm/JointGroup.h"
-#include "riglogic/joints/bpcm/JointStorageView.h"
+#include "riglogic/joints/bpcm/StorageView.h"
 #include "riglogic/joints/bpcm/LODRegion.h"
 #include "riglogic/utils/Macros.h"
 
diff --git a/Engine/Plugins/Animation/RigLogic/Source/RigLogicLib/Private/riglogic/ml/LODSpec.h b/Engine/Plugins/Animation/RigLogic/Source/RigLogicLib/Private/riglogic/ml/LODSpec.h
new file mode 100644
--- /dev/null
+++ b/Engine/Plugins/Animation/RigLogic/Source/RigLogicLib/Private/riglogic/ml/LODSpec.h
@@ -0,0 +1,28 @@
+// Copyright Epic Games, Inc. All Rights Reserved.
+
+#pragma once
+
+#include "riglogic/TypeDefs.h"
+
+namespace rl4 {
+
+struct LODSpec {
+    Matrix<std::uint32_t> netIndicesPerLOD;
+    std::uint32_t netCount;
+
+    explicit LODSpec(MemoryResource* memRes) : netIndicesPerLOD{memRes}, netCount{} {
+    }
+
+    LODSpec(Matrix<std::uint32_t>&& netIndicesPerLOD_, std::uint32_t netCount_) :
+        netIndicesPerLOD{std::move(netIndicesPerLOD_)},
+        netCount{netCount_} {
+    }
+
+    template<class Archive>
+    void serialize(Archive& archive) {
+        archive(netIndicesPerLOD, netCount);
+    }
+
+};
+
+}  // namespace rl4
diff --git a/Engine/Plugins/Animation/RigLogic/Source/RigLogicLib/Private/riglogic/ml/MachineLearnedBehavior.cpp b/Engine/Plugins/Animation/RigLogic/Source/RigLogicLib/Private/riglogic/ml/MachineLearnedBehavior.cpp
new file mode 100644
--- /dev/null
+++ b/Engine/Plugins/Animation/RigLogic/Source/RigLogicLib/Private/riglogic/ml/MachineLearnedBehavior.cpp
@@ -0,0 +1,58 @@
+// Copyright Epic Games, Inc. All Rights Reserved.
+
+#include "riglogic/ml/MachineLearnedBehavior.h"
+
+#include "riglogic/TypeDefs.h"
+#include "riglogic/controls/ControlsInputInstance.h"
+#include "riglogic/ml/MachineLearnedBehaviorOutputInstance.h"
+
+#include <cassert>
+#include <cstddef>
+
+namespace rl4 {
+
+MachineLearnedBehavior::MachineLearnedBehavior(MachineLearnedBehaviorEvaluator::Pointer evaluator_, MemoryResource* memRes) :
+    evaluator{std::move(evaluator_)},
+    neuralNetworkIndicesPerMeshRegion{memRes} {
+}
+
+MachineLearnedBehavior::MachineLearnedBehavior(MachineLearnedBehaviorEvaluator::Pointer evaluator_,
+                                               Vector<Matrix<std::uint16_t> >&& neuralNetworkIndicesPerMeshRegion_) :
+    evaluator{std::move(evaluator_)},
+    neuralNetworkIndicesPerMeshRegion{std::move(neuralNetworkIndicesPerMeshRegion_)} {
+}
+
+MachineLearnedBehaviorOutputInstance::Pointer MachineLearnedBehavior::createInstance(MemoryResource* instanceMemRes) const {
+    return evaluator->createInstance(instanceMemRes);
+}
+
+void MachineLearnedBehavior::calculate(ControlsInputInstance* inputs,
+                                       MachineLearnedBehaviorOutputInstance* intermediateOutputs,
+                                       std::uint16_t lod) const {
+    evaluator->calculate(inputs, intermediateOutputs, lod);
+}
+
+void MachineLearnedBehavior::calculate(ControlsInputInstance* inputs,
+                                       MachineLearnedBehaviorOutputInstance* intermediateOutputs,
+                                       std::uint16_t lod,
+                                       std::uint16_t neuralNetIndex) const {
+    evaluator->calculate(inputs, intermediateOutputs, lod, neuralNetIndex);
+}
+
+std::uint16_t MachineLearnedBehavior::getMeshCount() const {
+    return static_cast<std::uint16_t>(neuralNetworkIndicesPerMeshRegion.size());
+}
+
+std::uint16_t MachineLearnedBehavior::getMeshRegionCount(std::uint16_t meshIndex) const {
+    assert(meshIndex < neuralNetworkIndicesPerMeshRegion.size());
+    return static_cast<std::uint16_t>(neuralNetworkIndicesPerMeshRegion[meshIndex].size());
+}
+
+ConstArrayView<std::uint16_t> MachineLearnedBehavior::getNeuralNetworkIndices(std::uint16_t meshIndex,
+                                                                              std::uint16_t regionIndex) const {
+    assert(meshIndex < getMeshCount());
+    assert(regionIndex < getMeshRegionCount(meshIndex));
+    return ConstArrayView<std::uint16_t>{neuralNetworkIndicesPerMeshRegion[meshIndex][regionIndex]};
+}
+
+}  // namespace rl4
diff --git a/Engine/Plugins/Animation/RigLogic/Source/RigLogicLib/Private/riglogic/ml/MachineLearnedBehavior.h b/Engine/Plugins/Animation/RigLogic/Source/RigLogicLib/Private/riglogic/ml/MachineLearnedBehavior.h
new file mode 100644
--- /dev/null
+++ b/Engine/Plugins/Animation/RigLogic/Source/RigLogicLib/Private/riglogic/ml/MachineLearnedBehavior.h
@@ -0,0 +1,55 @@
+// Copyright Epic Games, Inc. All Rights Reserved.
+
+#pragma once
+
+#include "riglogic/TypeDefs.h"
+#include "riglogic/types/Aliases.h"
+#include "riglogic/ml/MachineLearnedBehaviorEvaluator.h"
+#include "riglogic/ml/MachineLearnedBehaviorOutputInstance.h"
+
+#include <cstddef>
+
+namespace rl4 {
+
+class ControlsInputInstance;
+
+class MachineLearnedBehavior {
+    public:
+        using Pointer = UniqueInstance<MachineLearnedBehavior>::PointerType;
+
+    public:
+        MachineLearnedBehavior(MachineLearnedBehaviorEvaluator::Pointer evaluator_, MemoryResource* memRes);
+        MachineLearnedBehavior(MachineLearnedBehaviorEvaluator::Pointer evaluator_,
+                               Vector<Matrix<std::uint16_t> >&& neuralNetworkIndicesPerMeshRegion_);
+
+        MachineLearnedBehaviorOutputInstance::Pointer createInstance(MemoryResource* instanceMemRes) const;
+        void calculate(ControlsInputInstance* inputs, MachineLearnedBehaviorOutputInstance* intermediateOutputs,
+                       std::uint16_t lod) const;
+        void calculate(ControlsInputInstance* inputs,
+                       MachineLearnedBehaviorOutputInstance* intermediateOutputs,
+                       std::uint16_t lod,
+                       std::uint16_t neuralNetIndex) const;
+
+        template<class Archive>
+        void load(Archive& archive) {
+            evaluator->load(archive);
+            archive >> neuralNetworkIndicesPerMeshRegion;
+        }
+
+        template<class Archive>
+        void save(Archive& archive) {
+            evaluator->save(archive);
+            archive << neuralNetworkIndicesPerMeshRegion;
+        }
+
+        std::uint16_t getMeshCount() const;
+        std::uint16_t getMeshRegionCount(std::uint16_t meshIndex) const;
+        ConstArrayView<std::uint16_t> getNeuralNetworkIndices(std::uint16_t meshIndex, std::uint16_t regionIndex) const;
+
+    private:
+        MachineLearnedBehaviorEvaluator::Pointer evaluator;
+        Vector<Matrix<std::uint16_t> > neuralNetworkIndicesPerMeshRegion;
+
+};
+
+}  // namespace rl4
diff --git a/Engine/Plugins/Animation/RigLogic/Source/RigLogicLib/Private/riglogic/ml/MachineLearnedBehaviorFactory.cpp b/Engine/Plugins/Animation/RigLogic/Source/RigLogicLib/Private/riglogic/ml/MachineLearnedBehaviorFactory.cpp
new file mode 100644
--- /dev/null
+++ b/Engine/Plugins/Animation/RigLogic/Source/RigLogicLib/Private/riglogic/ml/MachineLearnedBehaviorFactory.cpp
@@ -0,0 +1,77 @@
+// Copyright Epic Games, Inc. All Rights Reserved.
+
+#include "riglogic/ml/MachineLearnedBehaviorFactory.h"
+
+#include "riglogic/TypeDefs.h"
+#include "riglogic/utils/Extd.h"
+#include "riglogic/ml/MachineLearnedBehaviorEvaluator.h"
+#include "riglogic/ml/MachineLearnedBehaviorNullEvaluator.h"
+#include "riglogic/ml/cpu/Factory.h"
+#include "riglogic/riglogic/Configuration.h"
+#include "riglogic/riglogic/RigMetrics.h"
+#include "riglogic/system/simd/Detect.h"
+#include "riglogic/system/simd/SIMD.h"
+
+#include <cstdint>
+
+namespace rl4 {
+
+#ifdef RL_USE_HALF_FLOATS
+    using StorageValueType = std::uint16_t;
+#else
+    using StorageValueType = float;
+#endif  // RL_USE_HALF_FLOATS
+
+MachineLearnedBehaviorEvaluator::Pointer createEvaluator(const Configuration& config,
+                                                         const dna::MachineLearnedBehaviorReader* reader,
+                                                         MemoryResource* memRes) {
+    // Work around unused parameter warning when building without SSE and AVX
+    static_cast<void>(config);
+    #ifdef RL_BUILD_WITH_SSE
+        if (config.calculationType == CalculationType::SSE) {
+            return ml::cpu::Factory<StorageValueType, trimd::sse::F256, trimd::sse::F128>::create(reader, memRes);
+        }
+    #endif  // RL_BUILD_WITH_SSE
+    #ifdef RL_BUILD_WITH_AVX
+        if (config.calculationType == CalculationType::AVX) {
+            // Use 256-bit AVX registers and whatever 128-bit width type is available
+            return ml::cpu::Factory<StorageValueType, trimd::avx::F256, trimd::F128>::create(reader, memRes);
+        }
+    #endif  // RL_BUILD_WITH_AVX
+    return ml::cpu::Factory<float, trimd::scalar::F256, trimd::scalar::F128>::create(reader, memRes);
+}
+
+MachineLearnedBehavior::Pointer MachineLearnedBehaviorFactory::create(const Configuration& config,
+                                                                      const dna::MachineLearnedBehaviorReader* reader,
+                                                                      MemoryResource* memRes) {
+    auto moduleFactory = UniqueInstance<MachineLearnedBehavior>::with(memRes);
+    if (!config.loadMachineLearnedBehavior || (reader->getNeuralNetworkCount() == 0u)) {
+        auto evaluator =
+            UniqueInstance<MachineLearnedBehaviorNullEvaluator, MachineLearnedBehaviorEvaluator>::with(memRes).create();
+        return moduleFactory.create(std::move(evaluator), memRes);
+    }
+
+    Vector<Matrix<std::uint16_t> > neuralNetworkIndicesPerMeshRegion{memRes};
+    neuralNetworkIndicesPerMeshRegion.resize(reader->getMeshCount());
+    for (std::uint16_t meshIdx = {}; meshIdx < neuralNetworkIndicesPerMeshRegion.size(); ++meshIdx) {
+        neuralNetworkIndicesPerMeshRegion[meshIdx].resize(reader->getMeshRegionCount(meshIdx));
+        for (std::uint16_t regionIdx = {}; regionIdx < neuralNetworkIndicesPerMeshRegion[meshIdx].size(); ++regionIdx) {
+            const auto netIndices = reader->getNeuralNetworkIndicesForMeshRegion(meshIdx, regionIdx);
+            neuralNetworkIndicesPerMeshRegion[meshIdx][regionIdx].assign(netIndices.begin(), netIndices.end());
+        }
+    }
+
+    return moduleFactory.create(createEvaluator(config, reader, memRes), std::move(neuralNetworkIndicesPerMeshRegion));
+}
+
+MachineLearnedBehavior::Pointer MachineLearnedBehaviorFactory::create(const Configuration& config, MemoryResource* memRes) {
+    auto moduleFactory = UniqueInstance<MachineLearnedBehavior>::with(memRes);
+    if (!config.loadMachineLearnedBehavior) {
+        auto evaluator =
+            UniqueInstance<MachineLearnedBehaviorNullEvaluator, MachineLearnedBehaviorEvaluator>::with(memRes).create();
+        return moduleFactory.create(std::move(evaluator), memRes);
+    }
+    return moduleFactory.create(createEvaluator(config, nullptr, memRes), memRes);
+}
+
+}  // namespace rl4
diff --git a/Engine/Plugins/Animation/RigLogic/Source/RigLogicLib/Private/riglogic/ml/MachineLearnedBehaviorFactory.h b/Engine/Plugins/Animation/RigLogic/Source/RigLogicLib/Private/riglogic/ml/MachineLearnedBehaviorFactory.h
new file mode 100644
--- /dev/null
+++ b/Engine/Plugins/Animation/RigLogic/Source/RigLogicLib/Private/riglogic/ml/MachineLearnedBehaviorFactory.h
@@ -0,0 +1,21 @@
+// Copyright Epic Games, Inc. All Rights Reserved.
+
+#pragma once
+
+#include "riglogic/TypeDefs.h"
+#include "riglogic/ml/MachineLearnedBehavior.h"
+
+namespace rl4 {
+
+struct Configuration;
+struct RigMetrics;
+
+struct MachineLearnedBehaviorFactory {
+    static MachineLearnedBehavior::Pointer create(const Configuration& config,
+                                                  const dna::MachineLearnedBehaviorReader* reader,
+                                                  MemoryResource* memRes);
+    static MachineLearnedBehavior::Pointer create(const Configuration& config, MemoryResource* memRes);
+
+};
+
+}  // namespace rl4
diff --git a/Engine/Plugins/Animation/RigLogic/Source/RigLogicLib/Private/riglogic/ml/MachineLearnedBehaviorNullEvaluator.cpp b/Engine/Plugins/Animation/RigLogic/Source/RigLogicLib/Private/riglogic/ml/MachineLearnedBehaviorNullEvaluator.cpp
new file mode 100644
--- /dev/null
+++ b/Engine/Plugins/Animation/RigLogic/Source/RigLogicLib/Private/riglogic/ml/MachineLearnedBehaviorNullEvaluator.cpp
@@ -0,0 +1,32 @@
+// Copyright Epic Games, Inc. All Rights Reserved.
+
+#include "riglogic/ml/MachineLearnedBehaviorNullEvaluator.h"
+
+#include "riglogic/ml/MachineLearnedBehaviorNullOutputInstance.h"
+
+namespace rl4 {
+
+MachineLearnedBehaviorOutputInstance::Pointer MachineLearnedBehaviorNullEvaluator::createInstance(MemoryResource* instanceMemRes)
+const {
+    return UniqueInstance<MachineLearnedBehaviorNullOutputInstance,
+                          MachineLearnedBehaviorOutputInstance>::with(instanceMemRes).create();
+}
+
+void MachineLearnedBehaviorNullEvaluator::calculate(ControlsInputInstance*  /*unused*/,
+                                                    MachineLearnedBehaviorOutputInstance*  /*unused*/,
+                                                    std::uint16_t  /*unused*/) const {
+}
+
+void MachineLearnedBehaviorNullEvaluator::calculate(ControlsInputInstance*  /*unused*/,
+                                                    MachineLearnedBehaviorOutputInstance*  /*unused*/,
+                                                    std::uint16_t  /*unused*/,
+                                                    std::uint16_t  /*unused*/) const {
+}
+
+void MachineLearnedBehaviorNullEvaluator::load(terse::BinaryInputArchive<BoundedIOStream>&  /*unused*/) {
+}
+
+void MachineLearnedBehaviorNullEvaluator::save(terse::BinaryOutputArchive<BoundedIOStream>&  /*unused*/) {
+}
+
+}  // namespace rl4
diff --git a/Engine/Plugins/Animation/RigLogic/Source/RigLogicLib/Private/riglogic/ml/MachineLearnedBehaviorNullEvaluator.h b/Engine/Plugins/Animation/RigLogic/Source/RigLogicLib/Private/riglogic/ml/MachineLearnedBehaviorNullEvaluator.h
new file mode 100644
--- /dev/null
+++ b/Engine/Plugins/Animation/RigLogic/Source/RigLogicLib/Private/riglogic/ml/MachineLearnedBehaviorNullEvaluator.h
@@ -0,0 +1,27 @@
+// Copyright Epic Games, Inc. All Rights Reserved.
+
+#pragma once
+
+#include "riglogic/TypeDefs.h"
+#include "riglogic/ml/MachineLearnedBehaviorEvaluator.h"
+#include "riglogic/ml/MachineLearnedBehaviorOutputInstance.h"
+
+namespace rl4 {
+
+class ControlsInputInstance;
+
+class MachineLearnedBehaviorNullEvaluator : public MachineLearnedBehaviorEvaluator {
+    public:
+        MachineLearnedBehaviorOutputInstance::Pointer createInstance(MemoryResource* instanceMemRes) const override;
+        void calculate(ControlsInputInstance*  /*unused*/, MachineLearnedBehaviorOutputInstance*  /*unused*/,
+                       std::uint16_t  /*unused*/) const override;
+        void calculate(ControlsInputInstance*  /*unused*/,
+                       MachineLearnedBehaviorOutputInstance*  /*unused*/,
+                       std::uint16_t  /*unused*/,
+                       std::uint16_t  /*unused*/) const override;
+        void load(terse::BinaryInputArchive<BoundedIOStream>&  /*unused*/) override;
+        void save(terse::BinaryOutputArchive<BoundedIOStream>&  /*unused*/) override;
+
+};
+
+}  // namespace rl4
diff --git a/Engine/Plugins/Animation/RigLogic/Source/RigLogicLib/Private/riglogic/ml/cpu/Block4.h b/Engine/Plugins/Animation/RigLogic/Source/RigLogicLib/Private/riglogic/ml/cpu/Block4.h
new file mode 100644
--- /dev/null
+++ b/Engine/Plugins/Animation/RigLogic/Source/RigLogicLib/Private/riglogic/ml/cpu/Block4.h
@@ -0,0 +1,216 @@
+// Copyright Epic Games, Inc. All Rights Reserved.
+
+#pragma once
+
+#include "riglogic/TypeDefs.h"
+#include "riglogic/ml/cpu/Consts.h"
+#include "riglogic/utils/Macros.h"
+
+namespace rl4 {
+
+namespace ml {
+
+namespace cpu {
+
+template<typename TF256, typename T>
+static FORCE_INLINE void processBlocks8x1(const float* inputVectorEndAlignedTo4,
+                                          const float* inputVectorEnd,
+                                          const T*& weights,
+                                          TF256& sum) {
+    TF256 remainder{};
+    for (const float* inputVector = inputVectorEndAlignedTo4;
+         inputVector < inputVectorEnd;
+         ++inputVector, weights += block8Height) {
+        const TF256 input{*inputVector};
+        const TF256 blk = TF256::fromAlignedSource(weights);
+        remainder += (blk * input);
+    }
+    sum += remainder;
+}
+
+template<typename TF256, typename TActivationFunction, typename T>
+static FORCE_INLINE void processBlocks8x4(const float* inputVectorStart,
+                                          const float* inputVectorEndAlignedTo4,
+                                          const float* inputVectorEnd,
+                                          const T*& weights,
+                                          const T* biases,
+                                          const float* activationParams,
+                                          float* outbuf) {
+    TF256 sum1{};
+    TF256 sum2{};
+    TF256 sum3{};
+    TF256 sum4{};
+    for (const float* inputVector = inputVectorStart;
+         inputVector < inputVectorEndAlignedTo4;
+         inputVector += block8Width, weights += block8Height * block8Width) {
+        const TF256 input1{inputVector[0]};
+        const TF256 input2{inputVector[1]};
+        const TF256 input3{inputVector[2]};
+        const TF256 input4{inputVector[3]};
+        const TF256 blk1 = TF256::fromAlignedSource(weights + TF256::size() * 0);
+        const TF256 blk2 = TF256::fromAlignedSource(weights + TF256::size() * 1);
+        const TF256 blk3 = TF256::fromAlignedSource(weights + TF256::size() * 2);
+        const TF256 blk4 = TF256::fromAlignedSource(weights + TF256::size() * 3);
+        /*
+        #ifndef RL_NO_MANUAL_PREFETCH
+            #ifdef RL_USE_NON_TEMPORAL_PREFETCH
+                TF256::prefetchNTA(values + (lookAheadOffset * 6ul));
+                TF256::prefetchNTA(values + (lookAheadOffset * 7ul));
+                TF256::prefetchNTA(values + (lookAheadOffset * 8ul));
+                TF256::prefetchNTA(values + (lookAheadOffset * 9ul));
+                TF256::prefetchNTA(values + (lookAheadOffset * 10ul));
+                TF256::prefetchNTA(values + (lookAheadOffset * 11ul));
+            #else
+                TF256::prefetchT0(values + (lookAheadOffset * 6ul));
+                TF256::prefetchT0(values + (lookAheadOffset * 7ul));
+                TF256::prefetchT0(values + (lookAheadOffset * 8ul));
+                TF256::prefetchT0(values + (lookAheadOffset * 9ul));
+                TF256::prefetchT0(values + (lookAheadOffset * 10ul));
+                TF256::prefetchT0(values + (lookAheadOffset * 11ul));
+            #endif
+        #endif
+        */
+        sum1 += (blk1 * input1);
+        sum2 += (blk2 * input2);
+        sum3 += (blk3 * input3);
+        sum4 += (blk4 * input4);
+    }
+    // Process 8x1 horizontal remainder portion after 8x4 blocks are consumed
+    processBlocks8x1(inputVectorEndAlignedTo4, inputVectorEnd, weights, sum1);
+
+    const TF256 bias = TF256::fromAlignedSource(biases);
+
+    sum1 += sum3;
+    sum2 += sum4;
+    sum1 += sum2;
+
+    sum1 += bias;
+
+    TActivationFunction{} (sum1, activationParams);
+
+    sum1.alignedStore(outbuf);
+}
+
+template<typename TF128, typename T>
+static FORCE_INLINE void processBlocks4x1(const float* inputVectorEndAlignedTo8,
+                                          const float* inputVectorEnd,
+                                          const T*& weights,
+                                          TF128& sum1) {
+    TF128 remainder{};
+    for (const float* inputVector = inputVectorEndAlignedTo8;
+         inputVector < inputVectorEnd;
+         ++inputVector, weights += block4Height) {
+        const TF128 input{*inputVector};
+        const TF128 blk = TF128::fromAlignedSource(weights);
+        remainder += (blk * input);
+    }
+    sum1 += remainder;
+}
+
+template<typename TF128, typename TActivationFunction, typename T>
+static FORCE_INLINE void processBlocks4x8(const float* inputVectorStart,
+                                          const float* inputVectorEndAlignedTo8,
+                                          const float* inputVectorEnd,
+                                          const T*& weights,
+                                          const T* biases,
+                                          const float* activationParams,
+                                          float* outbuf) {
+    TF128 sum1{};
+    TF128 sum2{};
+    TF128 sum3{};
+    TF128 sum4{};
+    TF128 sum5{};
+    TF128 sum6{};
+    TF128 sum7{};
+    TF128 sum8{};
+    for (const float* inputVector = inputVectorStart;
+         inputVector < inputVectorEndAlignedTo8;
+         inputVector += block4Width, weights += block4Height * block4Width) {
+        const TF128 input1{inputVector[0]};
+        const TF128 input2{inputVector[1]};
+        const TF128 input3{inputVector[2]};
+        const TF128 input4{inputVector[3]};
+        const TF128 input5{inputVector[4]};
+        const TF128 input6{inputVector[5]};
+        const TF128 input7{inputVector[6]};
+        const TF128 input8{inputVector[7]};
+        const TF128 blk1 = TF128::fromAlignedSource(weights);
+        const TF128 blk2 = TF128::fromAlignedSource(weights + TF128::size());
+        const TF128 blk3 = TF128::fromAlignedSource(weights + TF128::size() * 2);
+        const TF128 blk4 = TF128::fromAlignedSource(weights + TF128::size() * 3);
+        const TF128 blk5 = TF128::fromAlignedSource(weights + TF128::size() * 4);
+        const TF128 blk6 = TF128::fromAlignedSource(weights + TF128::size() * 5);
+        const TF128 blk7 = TF128::fromAlignedSource(weights + TF128::size() * 6);
+        const TF128 blk8 = TF128::fromAlignedSource(weights + TF128::size() * 7);
+        sum1 += (blk1 * input1);
+        sum2 += (blk2 * input2);
+        sum3 += (blk3 * input3);
+        sum4 += (blk4 * input4);
+        sum5 += (blk5 * input5);
+        sum6 += (blk6 * input6);
+        sum7 += (blk7 * input7);
+        sum8 += (blk8 * input8);
+    }
+    // Process 4x1 horizontal remainder portion after 4x8 blocks are consumed
+    processBlocks4x1(inputVectorEndAlignedTo8, inputVectorEnd, weights, sum1);
+
+    const TF128 bias1 = TF128::fromAlignedSource(biases);
+
+    sum1 += sum2;
+    sum3 += sum4;
+    sum5 += sum6;
+    sum7 += sum8;
+    sum1 += sum3;
+    sum5 += sum7;
+    sum1 += sum5;
+
+    sum1 += bias1;
+
+    TActivationFunction{} (sum1, activationParams);
+
+    sum1.alignedStore(outbuf);
+}
+
+template<typename TF256, typename TF128, template<class ...> class TActivationFunction, typename T>
+static FORCE_INLINE void calculateBlock4(const NeuralNetLayer<T>& layer, ConstArrayView<float> inputs, ArrayView<float> outputs) {
+    float* outputVector = outputs.data();
+    const float* outputVectorEndAlignedToLastBlock8 = outputVector + layer.weights.rows.sizeAlignedToLastFullBlock;
+    float* outputVectorEnd = outputVector + layer.weights.rows.size;
+
+    const float* inputVector = inputs.data();
+    const float* inputVectorEndAlignedTo4 = inputVector + layer.weights.cols.sizeAlignedToLastFullBlock;
+    const float* inputVectorEndAlignedTo8 = inputVector + layer.weights.cols.sizeAlignedToSecondLastFullBlock;
+    const float* inputVectorEnd = inputVector + layer.weights.cols.size;
+
+    const T* weights = layer.weights.values.data();
+    const T* biases = layer.biases.data();
+    const float* activationParams = layer.activationFunctionParameters.data();
+
+    for (; outputVector < outputVectorEndAlignedToLastBlock8; outputVector += block8Height, biases += block8Height) {
+        alignas(TF256::alignment()) float outbuf[block8Height];
+        processBlocks8x4<TF256, TActivationFunction<TF256> >(inputVector,
+                                                             inputVectorEndAlignedTo4,
+                                                             inputVectorEnd,
+                                                             weights,
+                                                             biases,
+                                                             activationParams,
+                                                             outbuf);
+        std::memcpy(outputVector, outbuf, sizeof(outbuf));
+    }
+
+    for (; outputVector < outputVectorEnd; outputVector += block4Height, biases += block4Height) {
+        processBlocks4x8<TF128, TActivationFunction<TF128> >(inputVector,
+                                                             inputVectorEndAlignedTo8,
+                                                             inputVectorEnd,
+                                                             weights,
+                                                             biases,
+                                                             activationParams,
+                                                             outputVector);
+    }
+}
+
+}  // namespace rl4
+
+}  // namespace ml
+
+}  // namespace cpu
diff --git a/Engine/Plugins/Animation/RigLogic/Source/RigLogicLib/Private/riglogic/ml/cpu/Consts.h b/Engine/Plugins/Animation/RigLogic/Source/RigLogicLib/Private/riglogic/ml/cpu/Consts.h
new file mode 100644
--- /dev/null
+++ b/Engine/Plugins/Animation/RigLogic/Source/RigLogicLib/Private/riglogic/ml/cpu/Consts.h
@@ -0,0 +1,26 @@
+// Copyright Epic Games, Inc. All Rights Reserved.
+
+#pragma once
+
+#include <cstdint>
+
+namespace rl4 {
+
+namespace ml {
+
+namespace cpu {
+
+constexpr std::uint32_t block4Width = 8u;
+constexpr std::uint32_t block4Height = 4u;
+
+constexpr std::uint32_t block8Width = 4u;
+constexpr std::uint32_t block8Height = 8u;
+
+constexpr std::uint32_t block16Width = 4u;
+constexpr std::uint32_t block16Height = 16u;
+
+}  // namespace cpu
+
+}  // namespace ml
+
+}  // namespace rl4
diff --git a/Engine/Plugins/Animation/RigLogic/Source/RigLogicLib/Private/riglogic/ml/cpu/Evaluator.h b/Engine/Plugins/Animation/RigLogic/Source/RigLogicLib/Private/riglogic/ml/cpu/Evaluator.h
new file mode 100644
--- /dev/null
+++ b/Engine/Plugins/Animation/RigLogic/Source/RigLogicLib/Private/riglogic/ml/cpu/Evaluator.h
@@ -0,0 +1,97 @@
+// Copyright Epic Games, Inc. All Rights Reserved.
+
+#pragma once
+
+#include "riglogic/TypeDefs.h"
+#include "riglogic/controls/ControlsInputInstance.h"
+#include "riglogic/ml/LODSpec.h"
+#include "riglogic/ml/cpu/Inference.h"
+#include "riglogic/ml/cpu/NeuralNet.h"
+#include "riglogic/ml/cpu/OutputInstance.h"
+
+#include <cstddef>
+
+namespace rl4 {
+
+namespace ml {
+
+namespace cpu {
+
+template<typename T, typename TF256, typename TF128>
+class Evaluator : public MachineLearnedBehaviorEvaluator {
+    public:
+        using NeuralNetVectorType = Vector<NeuralNetInference<T, TF256, TF128> >;
+
+        struct Accessor;
+        friend Accessor;
+
+    public:
+        Evaluator(LODSpec&& lods_,
+                  NeuralNetVectorType&& neuralNets_,
+                  Vector<std::uint32_t>&& maxLayerOutputCounts_,
+                  OutputInstance::Factory instanceFactory_) :
+            lods{std::move(lods_)},
+            neuralNets{std::move(neuralNets_)},
+            maxLayerOutputCounts{std::move(maxLayerOutputCounts_)},
+            instanceFactory{instanceFactory_} {
+        }
+
+        MachineLearnedBehaviorOutputInstance::Pointer createInstance(MemoryResource* instanceMemRes) const override {
+            return instanceFactory(ConstArrayView<std::uint32_t>{maxLayerOutputCounts}, instanceMemRes);
+        }
+
+        void calculate(ControlsInputInstance* inputs, MachineLearnedBehaviorOutputInstance* intermediateOutputs,
+                       std::uint16_t lod) const override {
+            assert(lod < lods.netIndicesPerLOD.size());
+            const auto& netIndices = lods.netIndicesPerLOD[lod];
+            const auto inputBuffer = inputs->getInputBuffer();
+            const auto masks = intermediateOutputs->getMaskBuffer();
+            for (const auto neuralNetIndex : netIndices) {
+                assert(neuralNetIndex < masks.size());
+                assert(neuralNetIndex < neuralNets.size());
+                const float weight = masks[neuralNetIndex];
+                auto layerBuffer1 = static_cast<OutputInstance*>(intermediateOutputs)->getOutputBuffer1(neuralNetIndex);
+                auto layerBuffer2 = static_cast<OutputInstance*>(intermediateOutputs)->getOutputBuffer2(neuralNetIndex);
+                // inputBuffer is also outputBuffer
+                neuralNets[neuralNetIndex].calculate(inputBuffer, layerBuffer1, layerBuffer2, inputBuffer, weight);
+            }
+        }
+
+        void calculate(ControlsInputInstance* inputs,
+                       MachineLearnedBehaviorOutputInstance* intermediateOutputs,
+                       std::uint16_t lod,
+                       std::uint16_t neuralNetIndex) const override {
+            assert(lod < lods.netIndicesPerLOD.size());
+            static_cast<void>(lod);
+
+            const auto masks = intermediateOutputs->getMaskBuffer();
+            assert(neuralNetIndex < masks.size());
+            assert(neuralNetIndex < neuralNets.size());
+            const float weight = masks[neuralNetIndex];
+            auto inputBuffer = inputs->getInputBuffer();
+            auto layerBuffer1 = static_cast<OutputInstance*>(intermediateOutputs)->getOutputBuffer1(neuralNetIndex);
+            auto layerBuffer2 = static_cast<OutputInstance*>(intermediateOutputs)->getOutputBuffer2(neuralNetIndex);
+            // inputBuffer is also outputBuffer
+            neuralNets[neuralNetIndex].calculate(inputBuffer, layerBuffer1, layerBuffer2, inputBuffer, weight);
+        }
+
+        void load(terse::BinaryInputArchive<BoundedIOStream>& archive) override {
+            archive(lods, neuralNets, maxLayerOutputCounts);
+        }
+
+        void save(terse::BinaryOutputArchive<BoundedIOStream>& archive) override {
+            archive(lods, neuralNets, maxLayerOutputCounts);
+        }
+
+    private:
+        LODSpec lods;
+        NeuralNetVectorType neuralNets;
+        Vector<std::uint32_t> maxLayerOutputCounts;
+        OutputInstance::Factory instanceFactory;
+};
+
+}  // namespace cpu
+
+}  // namespace ml
+
+}  // namespace rl4
diff --git a/Engine/Plugins/Animation/RigLogic/Source/RigLogicLib/Private/riglogic/ml/cpu/Factory.h b/Engine/Plugins/Animation/RigLogic/Source/RigLogicLib/Private/riglogic/ml/cpu/Factory.h
new file mode 100644
--- /dev/null
+++ b/Engine/Plugins/Animation/RigLogic/Source/RigLogicLib/Private/riglogic/ml/cpu/Factory.h
@@ -0,0 +1,158 @@
+// Copyright Epic Games, Inc. All Rights Reserved.
+
+#pragma once
+
+#include "riglogic/TypeDefs.h"
+#include "riglogic/ml/LODSpec.h"
+#include "riglogic/ml/MachineLearnedBehaviorEvaluator.h"
+#include "riglogic/ml/cpu/Evaluator.h"
+#include "riglogic/ml/cpu/Inference.h"
+#include "riglogic/ml/cpu/NeuralNet.h"
+#include "riglogic/ml/cpu/OutputInstance.h"
+#include "riglogic/types/bpcm/Optimizer.h"
+#include "riglogic/utils/Extd.h"
+
+namespace rl4 {
+
+namespace ml {
+
+namespace cpu {
+
+static OutputInstance::Factory createInstanceFactory() {
+    return [](ConstArrayView<std::uint32_t> maxLayerOutputCounts, MemoryResource* memRes) {
+               return UniqueInstance<OutputInstance, MachineLearnedBehaviorOutputInstance>::with(memRes).create(
+                   maxLayerOutputCounts,
+                   memRes);
+    };
+}
+
+template<typename T, typename TF256, typename TF128>
+class Factory {
+    public:
+        static MachineLearnedBehaviorEvaluator::Pointer create(const dna::MachineLearnedBehaviorReader* reader,
+                                                               MemoryResource* memRes) {
+            Vector<NeuralNetInference<T, TF256, TF128> > neuralNets{memRes};
+            Vector<std::uint32_t> maxLayerOutputCountPerNet{memRes};
+            auto instanceFactory = createInstanceFactory();
+            auto factory = UniqueInstance<Evaluator<T, TF256, TF128>, MachineLearnedBehaviorEvaluator>::with(memRes);
+
+            if (reader == nullptr) {
+                return factory.create(LODSpec{memRes},
+                                      std::move(neuralNets),
+                                      std::move(maxLayerOutputCountPerNet),
+                                      instanceFactory);
+            }
+
+            auto lods = computeLODs(reader, memRes);
+            maxLayerOutputCountPerNet.resize(lods.netCount);
+            neuralNets.reserve(lods.netCount);
+
+            for (std::uint16_t neuralNetIdx = {}; neuralNetIdx < lods.netCount; ++neuralNetIdx) {
+                const auto layerCount = reader->getNeuralNetworkLayerCount(neuralNetIdx);
+                if (layerCount != 0u) {
+                    auto net = createNeuralNet(reader,
+                                               neuralNetIdx,
+                                               layerCount,
+                                               &maxLayerOutputCountPerNet[neuralNetIdx],
+                                               memRes);
+                    neuralNets.emplace_back(std::move(net), memRes);
+                }
+            }
+
+            return factory.create(std::move(lods),
+                                  std::move(neuralNets),
+                                  std::move(maxLayerOutputCountPerNet),
+                                  instanceFactory);
+        }
+
+    private:
+        static NeuralNet<T> createNeuralNet(const dna::MachineLearnedBehaviorReader* reader,
+                                            std::uint16_t neuralNetIdx,
+                                            std::uint16_t layerCount,
+                                            std::uint32_t* maxLayerOutputCount,
+                                            MemoryResource* memRes) {
+            const auto inputIndices = reader->getNeuralNetworkInputIndices(neuralNetIdx);
+            const auto outputIndices = reader->getNeuralNetworkOutputIndices(neuralNetIdx);
+            std::uint32_t inputCount = static_cast<std::uint32_t>(inputIndices.size());
+
+            NeuralNet<T> neuralNet{memRes};
+            * maxLayerOutputCount = inputCount;
+            neuralNet.layers.reserve(layerCount);
+            for (std::uint16_t layerIdx = {}; layerIdx < layerCount; ++layerIdx) {
+                const auto weights = reader->getNeuralNetworkLayerWeights(neuralNetIdx, layerIdx);
+                const auto biases = reader->getNeuralNetworkLayerBiases(neuralNetIdx, layerIdx);
+                const auto activationFunction = reader->getNeuralNetworkLayerActivationFunction(neuralNetIdx, layerIdx);
+                const auto activationFunctionParams = reader->getNeuralNetworkLayerActivationFunctionParameters(neuralNetIdx,
+                                                                                                                layerIdx);
+                const auto outputCount = static_cast<std::uint32_t>(biases.size());
+                const auto layer = createLayer(inputCount,
+                                               outputCount,
+                                               weights,
+                                               biases,
+                                               activationFunction,
+                                               activationFunctionParams,
+                                               memRes);
+                // Keep track of the layer with the largest number of outputs
+                // In the next layer, the current output count becomes the input count
+                * maxLayerOutputCount = std::max(*maxLayerOutputCount, layer.weights.padded.rows);
+                neuralNet.layers.push_back(std::move(layer));
+                inputCount = outputCount;
+            }
+
+            neuralNet.inputIndices.assign(inputIndices.begin(), inputIndices.end());
+            neuralNet.outputIndices.assign(outputIndices.begin(), outputIndices.end());
+            return neuralNet;
+        }
+
+        static NeuralNetLayer<T> createLayer(std::uint32_t inputCount,
+                                             std::uint32_t outputCount,
+                                             ConstArrayView<float> weights,
+                                             ConstArrayView<float> biases,
+                                             dna::ActivationFunction activationFunction,
+                                             ConstArrayView<float> activationFunctionParams,
+                                             MemoryResource* memRes) {
+            NeuralNetLayer<T> layer{memRes};
+
+            layer.weights.original = {outputCount, inputCount};
+            const std::uint32_t padding = extd::roundUp(layer.weights.original.rows, block4Height) - layer.weights.original.rows;
+            layer.weights.padded = {layer.weights.original.rows + padding, layer.weights.original.cols};
+            layer.weights.rows =
+                AlignedBlockView{layer.weights.original.rows, layer.weights.padded.rows, block8Height, block4Height};
+            layer.weights.cols = AlignedBlockView{
+                layer.weights.padded.cols,
+                layer.weights.padded.cols - (layer.weights.padded.cols % 4u),
+                layer.weights.padded.cols - (layer.weights.padded.cols % 8u)
+            };
+            layer.weights.values.resize(layer.weights.padded.size());
+            bpcm::Optimizer<TF128, block8Height, block4Height>::optimize(layer.weights.values.data(),
+                                                                         weights.data(),
+                                                                         layer.weights.original);
+
+            layer.biases.resize(layer.weights.padded.rows);
+            bpcm::Optimizer<TF128, block8Height, block4Height>::optimize(layer.biases.data(), biases.data(),
+                                                                         Extent{static_cast<std::uint32_t>(biases.size()), 1u});
+
+            layer.activationFunction = activationFunction;
+            layer.activationFunctionParameters.assign(activationFunctionParams.begin(), activationFunctionParams.end());
+            return layer;
+        }
+
+        static LODSpec computeLODs(const dna::MachineLearnedBehaviorReader* reader, MemoryResource* memRes) {
+            LODSpec lods{memRes};
+            const auto lodCount = reader->getLODCount();
+            lods.netIndicesPerLOD.resize(lodCount);
+            for (std::uint16_t lod = {}; lod < lodCount; ++lod) {
+                const auto netIndices = reader->getNeuralNetworkIndicesForLOD(lod);
+                lods.netIndicesPerLOD[lod].assign(netIndices.begin(), netIndices.end());
+            }
+            lods.netCount = reader->getNeuralNetworkCount();
+            return lods;
+        }
+
+};
+
+}  // namespace cpu
+
+}  // namespace ml
+
+}  // namespace rl4
diff --git a/Engine/Plugins/Animation/RigLogic/Source/RigLogicLib/Private/riglogic/ml/cpu/Inference.h b/Engine/Plugins/Animation/RigLogic/Source/RigLogicLib/Private/riglogic/ml/cpu/Inference.h
new file mode 100644
--- /dev/null
+++ b/Engine/Plugins/Animation/RigLogic/Source/RigLogicLib/Private/riglogic/ml/cpu/Inference.h
@@ -0,0 +1,110 @@
+// Copyright Epic Games, Inc. All Rights Reserved.
+
+#pragma once
+
+#include "riglogic/TypeDefs.h"
+#include "riglogic/ml/cpu/NeuralNet.h"
+#include "riglogic/ml/cpu/layers/LeakyReLULayerEvaluator.h"
+#include "riglogic/ml/cpu/layers/LinearLayerEvaluator.h"
+#include "riglogic/ml/cpu/layers/ReLULayerEvaluator.h"
+#include "riglogic/ml/cpu/layers/SigmoidLayerEvaluator.h"
+#include "riglogic/ml/cpu/layers/TanHLayerEvaluator.h"
+
+namespace rl4 {
+
+namespace ml {
+
+namespace cpu {
+
+template<typename T, typename TF256, typename TF128>
+struct LayerEvaluatorFactory {
+
+    static typename LayerEvaluator<T>::Pointer create(dna::ActivationFunction activationFunction, MemoryResource* memRes) {
+        switch (activationFunction) {
+            case dna::ActivationFunction::linear:
+                return UniqueInstance<LinearLayerEvaluator<T, TF256, TF128>, LayerEvaluator<T> >::with(memRes).create();
+            case dna::ActivationFunction::relu:
+                return UniqueInstance<ReLULayerEvaluator<T, TF256, TF128>, LayerEvaluator<T> >::with(memRes).create();
+            case dna::ActivationFunction::leakyrelu:
+                return UniqueInstance<LeakyReLULayerEvaluator<T, TF256, TF128>, LayerEvaluator<T> >::with(memRes).create();
+            case dna::ActivationFunction::tanh:
+                return UniqueInstance<TanHLayerEvaluator<T, TF256, TF128>, LayerEvaluator<T> >::with(memRes).create();
+            case dna::ActivationFunction::sigmoid:
+                return UniqueInstance<SigmoidLayerEvaluator<T, TF256, TF128>, LayerEvaluator<T> >::with(memRes).create();
+        }
+        return nullptr;
+    }
+
+};
+
+template<typename T, typename TF256, typename TF128>
+struct NeuralNetInference {
+    NeuralNet<T> neuralNet;
+    Vector<typename LayerEvaluator<T>::Pointer> layerEvaluators;
+
+    explicit NeuralNetInference(MemoryResource* memRes) : neuralNet{memRes}, layerEvaluators{memRes} {
+    }
+
+    NeuralNetInference(NeuralNet<T>&& neuralNet_, MemoryResource* memRes) : neuralNet{std::move(neuralNet_)},
+        layerEvaluators{memRes} {
+        createLayerEvaluators();
+    }
+
+    void calculate(ConstArrayView<float> inputBuffer, ArrayView<float> layerBuffer1, ArrayView<float> layerBuffer2,
+                   ArrayView<float> outputBuffer, float weight) const {
+
+        assert(neuralNet.inputIndices.size() <= layerBuffer1.size());
+        assert(neuralNet.outputIndices.size() <= layerBuffer1.size());
+        assert(neuralNet.inputIndices.size() <= layerBuffer2.size());
+        assert(neuralNet.outputIndices.size() <= layerBuffer2.size());
+        assert(layerEvaluators.size() == neuralNet.layers.size());
+
+        if (weight == 0.0f) {
+            for (std::size_t i = 0ul; i < neuralNet.outputIndices.size(); ++i) {
+                outputBuffer[neuralNet.outputIndices[i]] = 0.0f;
+            }
+            return;
+        }
+
+        for (std::size_t i = 0ul; i < neuralNet.inputIndices.size(); ++i) {
+            layerBuffer1[i] = inputBuffer[neuralNet.inputIndices[i]];
+        }
+
+        for (std::size_t layerIndex = 0u; layerIndex < neuralNet.layers.size(); ++layerIndex) {
+            layerEvaluators[layerIndex]->calculate(neuralNet.layers[layerIndex], layerBuffer1, layerBuffer2);
+            std::swap(layerBuffer1, layerBuffer2);
+        }
+
+        for (std::size_t i = 0ul; i < neuralNet.outputIndices.size(); ++i) {
+            outputBuffer[neuralNet.outputIndices[i]] = layerBuffer1[i] * weight;
+        }
+    }
+
+    template<class Archive>
+    void load(Archive& archive) {
+        archive(neuralNet);
+        createLayerEvaluators();
+    }
+
+    template<class Archive>
+    void save(Archive& archive) {
+        archive(neuralNet);
+    }
+
+    void createLayerEvaluators() {
+        layerEvaluators.resize(neuralNet.layers.size());
+        auto memRes = layerEvaluators.get_allocator().getMemoryResource();
+        for (std::size_t layerIndex = 0ul; layerIndex < neuralNet.layers.size(); ++layerIndex) {
+            layerEvaluators[layerIndex] = LayerEvaluatorFactory<T, TF256, TF128>::create(
+                neuralNet.layers[layerIndex].activationFunction,
+                memRes);
+        }
+    }
+
+};
+
+}  // namespace cpu
+
+}  // namespace ml
+
+}  // namespace rl4
diff --git a/Engine/Plugins/Animation/RigLogic/Source/RigLogicLib/Private/riglogic/ml/cpu/NeuralNet.h b/Engine/Plugins/Animation/RigLogic/Source/RigLogicLib/Private/riglogic/ml/cpu/NeuralNet.h
new file mode 100644
--- /dev/null
+++ b/Engine/Plugins/Animation/RigLogic/Source/RigLogicLib/Private/riglogic/ml/cpu/NeuralNet.h
@@ -0,0 +1,88 @@
+// Copyright Epic Games, Inc. All Rights Reserved.
+
+#pragma once
+
+#include "riglogic/TypeDefs.h"
+#include "riglogic/types/AlignedBlockView.h"
+#include "riglogic/types/Extent.h"
+
+#include <cstdint>
+
+namespace rl4 {
+
+namespace ml {
+
+namespace cpu {
+
+template<typename T>
+struct WeightMatrix {
+    Extent original;
+    Extent padded;
+    AlignedBlockView rows;
+    AlignedBlockView cols;
+    AlignedVector<T> values;
+
+    explicit WeightMatrix(MemoryResource* memRes) :
+        original{},
+        padded{},
+        rows{},
+        cols{},
+        values{memRes} {
+    }
+
+    template<class Archive>
+    void serialize(Archive& archive) {
+        archive(original, padded, rows, cols, values);
+    }
+
+};
+
+template<typename T>
+struct NeuralNetLayer {
+    WeightMatrix<T> weights;
+    AlignedVector<T> biases;
+    Vector<float> activationFunctionParameters;
+    dna::ActivationFunction activationFunction;
+
+    explicit NeuralNetLayer(MemoryResource* memRes) :
+        weights{memRes},
+        biases{memRes},
+        activationFunctionParameters{memRes},
+        activationFunction{} {
+    }
+
+    template<class Archive>
+    void serialize(Archive& archive) {
+        archive(weights, biases, activationFunctionParameters, activationFunction);
+    }
+
+};
+
+template<typename T>
+struct NeuralNet {
+    Vector<NeuralNetLayer<T> > layers;
+    Vector<std::uint16_t> inputIndices;
+    Vector<std::uint16_t> outputIndices;
+
+    explicit NeuralNet(MemoryResource* memRes) : layers{memRes}, inputIndices{memRes}, outputIndices{memRes} {
+    }
+
+    NeuralNet(Vector<NeuralNetLayer<T> >&& layers_, Vector<std::uint16_t>&& inputIndices_,
+              Vector<std::uint16_t>&& outputIndices_) :
+        layers{std::move(layers_)},
+        inputIndices{std::move(inputIndices_)},
+        outputIndices{std::move(outputIndices_)} {
+    }
+
+    template<class Archive>
+    void serialize(Archive& archive) {
+        archive(layers, inputIndices, outputIndices);
+    }
+
+};
+
+}  // namespace cpu
+
+}  // namespace ml
+
+}  // namespace rl4
diff --git a/Engine/Plugins/Animation/RigLogic/Source/RigLogicLib/Private/riglogic/ml/cpu/OutputInstance.cpp b/Engine/Plugins/Animation/RigLogic/Source/RigLogicLib/Private/riglogic/ml/cpu/OutputInstance.cpp
new file mode 100644
--- /dev/null
+++ b/Engine/Plugins/Animation/RigLogic/Source/RigLogicLib/Private/riglogic/ml/cpu/OutputInstance.cpp
@@ -0,0 +1,49 @@
+// Copyright Epic Games, Inc. All Rights Reserved.
+
+#include "riglogic/ml/cpu/OutputInstance.h"
+
+#include <cassert>
+#include <cstdint>
+
+namespace rl4 {
+
+namespace ml {
+
+namespace cpu {
+
+OutputInstance::OutputInstance(ConstArrayView<std::uint32_t> maxLayerOutputCounts, MemoryResource* memRes) :
+    outputBuffer1{memRes},
+    outputBuffer2{memRes},
+    maskBuffer{maxLayerOutputCounts.size(), 1.0f, memRes} {
+
+    outputBuffer1.resize(maxLayerOutputCounts.size());
+    outputBuffer2.resize(maxLayerOutputCounts.size());
+    for (std::size_t neuralNetIdx = 0ul; neuralNetIdx < maxLayerOutputCounts.size(); ++neuralNetIdx) {
+        outputBuffer1[neuralNetIdx].resize(maxLayerOutputCounts[neuralNetIdx]);
+        outputBuffer2[neuralNetIdx].resize(maxLayerOutputCounts[neuralNetIdx]);
+    }
+}
+
+ArrayView<float> OutputInstance::getOutputBuffer1(std::uint32_t neuralNetIndex) {
+    assert(neuralNetIndex < outputBuffer1.size());
+    return ArrayView<float>{outputBuffer1[neuralNetIndex]};
+}
+
+ArrayView<float> OutputInstance::getOutputBuffer2(std::uint32_t neuralNetIndex) {
+    assert(neuralNetIndex < outputBuffer2.size());
+    return ArrayView<float>{outputBuffer2[neuralNetIndex]};
+}
+
+ArrayView<float> OutputInstance::getMaskBuffer() {
+    return ArrayView<float>{maskBuffer};
+}
+
+ConstArrayView<float> OutputInstance::getMaskBuffer() const {
+    return ConstArrayView<float>{maskBuffer};
+}
+
+}  // namespace cpu
+
+}  // namespace ml
+
+}  // namespace rl4
diff --git a/Engine/Plugins/Animation/RigLogic/Source/RigLogicLib/Private/riglogic/ml/cpu/OutputInstance.h b/Engine/Plugins/Animation/RigLogic/Source/RigLogicLib/Private/riglogic/ml/cpu/OutputInstance.h
new file mode 100644
--- /dev/null
+++ b/Engine/Plugins/Animation/RigLogic/Source/RigLogicLib/Private/riglogic/ml/cpu/OutputInstance.h
@@ -0,0 +1,45 @@
+// Copyright Epic Games, Inc. All Rights Reserved.
+
+#pragma once
+
+#include "riglogic/TypeDefs.h"
+#include "riglogic/ml/MachineLearnedBehaviorOutputInstance.h"
+
+#ifdef _MSC_VER
+    #pragma warning(push)
+    #pragma warning(disable : 4365 4987)
+#endif
+#include <functional>
+#ifdef _MSC_VER
+    #pragma warning(pop)
+#endif
+
+namespace rl4 {
+
+namespace ml {
+
+namespace cpu {
+
+class OutputInstance : public MachineLearnedBehaviorOutputInstance {
+    public:
+        using Factory = std::function<Pointer(ConstArrayView<std::uint32_t>, MemoryResource*)>;
+
+    public:
+        OutputInstance(ConstArrayView<std::uint32_t> maxLayerOutputCounts, MemoryResource* memRes);
+        ArrayView<float> getOutputBuffer1(std::uint32_t neuralNetIndex);
+        ArrayView<float> getOutputBuffer2(std::uint32_t neuralNetIndex);
+        ArrayView<float> getMaskBuffer() override;
+        ConstArrayView<float> getMaskBuffer() const override;
+
+    private:
+        Vector<AlignedVector<float> > outputBuffer1;
+        Vector<AlignedVector<float> > outputBuffer2;
+        Vector<float> maskBuffer;
+
+};
+
+}  // namespace cpu
+
+}  // namespace ml
+
+}  // namespace rl4
diff --git a/Engine/Plugins/Animation/RigLogic/Source/RigLogicLib/Private/riglogic/ml/cpu/layers/LayerEvaluator.h b/Engine/Plugins/Animation/RigLogic/Source/RigLogicLib/Private/riglogic/ml/cpu/layers/LayerEvaluator.h
new file mode 100644
--- /dev/null
+++ b/Engine/Plugins/Animation/RigLogic/Source/RigLogicLib/Private/riglogic/ml/cpu/layers/LayerEvaluator.h
@@ -0,0 +1,29 @@
+// Copyright Epic Games, Inc. All Rights Reserved.
+
+#pragma once
+
+#include "riglogic/TypeDefs.h"
+#include "riglogic/ml/cpu/NeuralNet.h"
+
+namespace rl4 {
+
+namespace ml {
+
+namespace cpu {
+
+template<typename T>
+class LayerEvaluator {
+    public:
+        using Pointer = typename UniqueInstance<LayerEvaluator>::PointerType;
+
+    public:
+        virtual ~LayerEvaluator() = default;
+
+        virtual void calculate(const NeuralNetLayer<T>& layer, ConstArrayView<float> inputs, ArrayView<float> outputs) const = 0;
+};
+
+}  // namespace cpu
+
+}  // namespace ml
+
+}  // namespace rl4
diff --git a/Engine/Plugins/Animation/RigLogic/Source/RigLogicLib/Private/riglogic/ml/cpu/layers/LeakyReLULayerEvaluator.h b/Engine/Plugins/Animation/RigLogic/Source/RigLogicLib/Private/riglogic/ml/cpu/layers/LeakyReLULayerEvaluator.h
new file mode 100644
--- /dev/null
+++ b/Engine/Plugins/Animation/RigLogic/Source/RigLogicLib/Private/riglogic/ml/cpu/layers/LeakyReLULayerEvaluator.h
@@ -0,0 +1,44 @@
+// Copyright Epic Games, Inc. All Rights Reserved.
+
+#pragma once
+
+#include "riglogic/TypeDefs.h"
+#include "riglogic/ml/cpu/Block4.h"
+#include "riglogic/ml/cpu/Consts.h"
+#include "riglogic/ml/cpu/layers/LayerEvaluator.h"
+#include "riglogic/system/simd/SIMD.h"
+#include "riglogic/utils/Macros.h"
+
+namespace rl4 {
+
+namespace ml {
+
+namespace cpu {
+
+template<typename TFVec>
+struct LeakyReLUActivationFunction {
+
+    void operator()(TFVec& sum, const float* activationParams) {
+        const TFVec alpha{* activationParams};
+        const TFVec zero{};
+        const TFVec mask = (sum < zero);
+        const TFVec multiplied = (sum * alpha);
+        sum = trimd::andnot(mask, sum) | (multiplied & mask);
+    }
+
+};
+
+template<typename T, typename TF256, typename TF128>
+class LeakyReLULayerEvaluator : public LayerEvaluator<T> {
+    public:
+        void calculate(const NeuralNetLayer<T>& layer, ConstArrayView<float> inputs, ArrayView<float> outputs) const override {
+            calculateBlock4<TF256, TF128, LeakyReLUActivationFunction>(layer, inputs, outputs);
+        }
+
+};
+
+}  // namespace cpu
+
+}  // namespace ml
+
+}  // namespace rl4
diff --git a/Engine/Plugins/Animation/RigLogic/Source/RigLogicLib/Private/riglogic/ml/cpu/layers/LinearLayerEvaluator.h b/Engine/Plugins/Animation/RigLogic/Source/RigLogicLib/Private/riglogic/ml/cpu/layers/LinearLayerEvaluator.h
new file mode 100644
--- /dev/null
+++ b/Engine/Plugins/Animation/RigLogic/Source/RigLogicLib/Private/riglogic/ml/cpu/layers/LinearLayerEvaluator.h
@@ -0,0 +1,38 @@
+// Copyright Epic Games, Inc. All Rights Reserved.
+
+#pragma once
+
+#include "riglogic/TypeDefs.h"
+#include "riglogic/ml/cpu/Block4.h"
+#include "riglogic/ml/cpu/Consts.h"
+#include "riglogic/ml/cpu/layers/LayerEvaluator.h"
+#include "riglogic/utils/Macros.h"
+
+namespace rl4 {
+
+namespace ml {
+
+namespace cpu {
+
+template<typename TFVec>
+struct LinearActivationFunction {
+
+    void operator()(TFVec&  /*unused*/, const float*  /*unused*/) {
+    }
+
+};
+
+template<typename T, typename TF256, typename TF128>
+class LinearLayerEvaluator : public LayerEvaluator<T> {
+    public:
+        void calculate(const NeuralNetLayer<T>& layer, ConstArrayView<float> inputs, ArrayView<float> outputs) const override {
+            calculateBlock4<TF256, TF128, LinearActivationFunction>(layer, inputs, outputs);
+        }
+
+};
+
+}  // namespace cpu
+
+}  // namespace ml
+
+}  // namespace rl4
diff --git a/Engine/Plugins/Animation/RigLogic/Source/RigLogicLib/Private/riglogic/ml/cpu/layers/ReLULayerEvaluator.h b/Engine/Plugins/Animation/RigLogic/Source/RigLogicLib/Private/riglogic/ml/cpu/layers/ReLULayerEvaluator.h
new file mode 100644
--- /dev/null
+++ b/Engine/Plugins/Animation/RigLogic/Source/RigLogicLib/Private/riglogic/ml/cpu/layers/ReLULayerEvaluator.h
@@ -0,0 +1,42 @@
+// Copyright Epic Games, Inc. All Rights Reserved.
+
+#pragma once
+
+#include "riglogic/TypeDefs.h"
+#include "riglogic/ml/cpu/Block4.h"
+#include "riglogic/ml/cpu/Consts.h"
+#include "riglogic/ml/cpu/layers/LayerEvaluator.h"
+#include "riglogic/system/simd/SIMD.h"
+#include "riglogic/utils/Macros.h"
+
+namespace rl4 {
+
+namespace ml {
+
+namespace cpu {
+
+template<typename TFVec>
+struct ReLUActivationFunction {
+
+    void operator()(TFVec& sum, const float*  /*unused*/) {
+        const TFVec zero{};
+        const TFVec mask = (sum < zero);
+        sum = trimd::andnot(mask, sum);
+    }
+
+};
+
+template<typename T, typename TF256, typename TF128>
+class ReLULayerEvaluator : public LayerEvaluator<T> {
+    public:
+        void calculate(const NeuralNetLayer<T>& layer, ConstArrayView<float> inputs, ArrayView<float> outputs) const override {
+            calculateBlock4<TF256, TF128, ReLUActivationFunction>(layer, inputs, outputs);
+        }
+
+};
+
+}  // namespace cpu
+
+}  // namespace ml
+
+}  // namespace rl4
diff --git a/Engine/Plugins/Animation/RigLogic/Source/RigLogicLib/Private/riglogic/ml/cpu/layers/SigmoidLayerEvaluator.h b/Engine/Plugins/Animation/RigLogic/Source/RigLogicLib/Private/riglogic/ml/cpu/layers/SigmoidLayerEvaluator.h
new file mode 100644
--- /dev/null
+++ b/Engine/Plugins/Animation/RigLogic/Source/RigLogicLib/Private/riglogic/ml/cpu/layers/SigmoidLayerEvaluator.h
@@ -0,0 +1,79 @@
+// Copyright Epic Games, Inc. All Rights Reserved.
+
+#pragma once
+
+#include "riglogic/TypeDefs.h"
+#include "riglogic/ml/cpu/Block4.h"
+#include "riglogic/ml/cpu/Consts.h"
+#include "riglogic/ml/cpu/layers/LayerEvaluator.h"
+#include "riglogic/ml/cpu/layers/Utils.h"
+#include "riglogic/utils/Macros.h"
+
+#include <cmath>
+#include <type_traits>
+
+namespace rl4 {
+
+namespace ml {
+
+namespace cpu {
+
+template<typename TFVec, typename Enable = void>
+struct SigmoidActivationFunction;
+
+template<typename TF256>
+struct SigmoidActivationFunction<TF256, typename std::enable_if<HasSize<TF256, 8ul>::value>::type> {
+
+    void operator()(TF256& sum, const float*  /*unused*/) {
+        // FIXME: dangerously inefficient code
+        alignas(TF256::alignment()) float buf[TF256::size()];
+
+        sum.alignedStore(buf);
+
+        buf[0] = 1.0f / (1.0f + std::exp(-buf[0]));
+        buf[1] = 1.0f / (1.0f + std::exp(-buf[1]));
+        buf[2] = 1.0f / (1.0f + std::exp(-buf[2]));
+        buf[3] = 1.0f / (1.0f + std::exp(-buf[3]));
+        buf[4] = 1.0f / (1.0f + std::exp(-buf[4]));
+        buf[5] = 1.0f / (1.0f + std::exp(-buf[5]));
+        buf[6] = 1.0f / (1.0f + std::exp(-buf[6]));
+        buf[7] = 1.0f / (1.0f + std::exp(-buf[7]));
+
+        sum.alignedLoad(buf);
+    }
+
+};
+
+template<typename TF128>
+struct SigmoidActivationFunction<TF128, typename std::enable_if<HasSize<TF128, 4ul>::value>::type> {
+
+    void operator()(TF128& sum, const float*  /*unused*/) {
+        // FIXME: dangerously inefficient code
+        alignas(TF128::alignment()) float buf[TF128::size()];
+
+        sum.alignedStore(buf);
+
+        buf[0] = 1.0f / (1.0f + std::exp(-buf[0]));
+        buf[1] = 1.0f / (1.0f + std::exp(-buf[1]));
+        buf[2] = 1.0f / (1.0f + std::exp(-buf[2]));
+        buf[3] = 1.0f / (1.0f + std::exp(-buf[3]));
+
+        sum.alignedLoad(buf);
+    }
+
+};
+
+template<typename T, typename TF256, typename TF128>
+class SigmoidLayerEvaluator : public LayerEvaluator<T> {
+    public:
+        void calculate(const NeuralNetLayer<T>& layer, ConstArrayView<float> inputs, ArrayView<float> outputs) const override {
+            calculateBlock4<TF256, TF128, SigmoidActivationFunction>(layer, inputs, outputs);
+        }
+
+};
+
+}  // namespace cpu
+
+}  // namespace ml
+
+}  // namespace rl4
diff --git a/Engine/Plugins/Animation/RigLogic/Source/RigLogicLib/Private/riglogic/ml/cpu/layers/TanHLayerEvaluator.h b/Engine/Plugins/Animation/RigLogic/Source/RigLogicLib/Private/riglogic/ml/cpu/layers/TanHLayerEvaluator.h
new file mode 100644
--- /dev/null
+++ b/Engine/Plugins/Animation/RigLogic/Source/RigLogicLib/Private/riglogic/ml/cpu/layers/TanHLayerEvaluator.h
@@ -0,0 +1,80 @@
+// Copyright Epic Games, Inc. All Rights Reserved.
+
+#pragma once
+
+#include "riglogic/TypeDefs.h"
+#include "riglogic/ml/cpu/Block4.h"
+#include "riglogic/ml/cpu/Consts.h"
+#include "riglogic/ml/cpu/layers/LayerEvaluator.h"
+#include "riglogic/ml/cpu/layers/Utils.h"
+#include "riglogic/utils/Macros.h"
+
+#include <cmath>
+#include <type_traits>
+
+namespace rl4 {
+
+namespace ml {
+
+namespace cpu {
+
+template<typename TFVec, typename Enable = void>
+struct TanHActivationFunction;
+
+
+template<typename TF256>
+struct TanHActivationFunction<TF256, typename std::enable_if<HasSize<TF256, 8ul>::value>::type> {
+
+    void operator()(TF256& sum, const float*  /*unused*/) {
+        // FIXME: dangerously inefficient code
+        alignas(TF256::alignment()) float buf[TF256::size()];
+
+        sum.alignedStore(buf);
+
+        buf[0] = std::tanh(buf[0]);
+        buf[1] = std::tanh(buf[1]);
+        buf[2] = std::tanh(buf[2]);
+        buf[3] = std::tanh(buf[3]);
+        buf[4] = std::tanh(buf[4]);
+        buf[5] = std::tanh(buf[5]);
+        buf[6] = std::tanh(buf[6]);
+        buf[7] = std::tanh(buf[7]);
+
+        sum.alignedLoad(buf);
+    }
+
+};
+
+template<typename TF128>
+struct TanHActivationFunction<TF128, typename std::enable_if<HasSize<TF128, 4ul>::value>::type> {
+
+    void operator()(TF128& sum, const float*  /*unused*/) {
+        // FIXME: dangerously inefficient code
+        alignas(TF128::alignment()) float buf[TF128::size()];
+
+        sum.alignedStore(buf);
+
+        buf[0] = std::tanh(buf[0]);
+        buf[1] = std::tanh(buf[1]);
+        buf[2] = std::tanh(buf[2]);
+        buf[3] = std::tanh(buf[3]);
+
+        sum.alignedLoad(buf);
+    }
+
+};
+
+template<typename T, typename TF256, typename TF128>
+class TanHLayerEvaluator : public LayerEvaluator<T> {
+    public:
+        void calculate(const NeuralNetLayer<T>& layer, ConstArrayView<float> inputs, ArrayView<float> outputs) const override {
+            calculateBlock4<TF256, TF128, TanHActivationFunction>(layer, inputs, outputs);
+        }
+
+};
+
+}  // namespace cpu
+
+}  // namespace ml
+
+}  // namespace rl4
diff --git a/Engine/Plugins/Animation/RigLogic/Source/RigLogicLib/Private/riglogic/riglogic/RigInstanceImpl.cpp b/Engine/Plugins/Animation/RigLogic/Source/RigLogicLib/Private/riglogic/riglogic/RigInstanceImpl.cpp
--- a/Engine/Plugins/Animation/RigLogic/Source/RigLogicLib/Private/riglogic/riglogic/RigInstanceImpl.cpp
+++ b/Engine/Plugins/Animation/RigLogic/Source/RigLogicLib/Private/riglogic/riglogic/RigInstanceImpl.cpp
@@ -3,14 +3,15 @@
 #include "riglogic/riglogic/RigInstanceImpl.h"
 
 #include "riglogic/TypeDefs.h"
+#include "riglogic/animatedmaps/AnimatedMapsOutputInstance.h"
+#include "riglogic/blendshapes/BlendShapesOutputInstance.h"
+#include "riglogic/controls/ControlsInputInstance.h"
+#include "riglogic/joints/JointsOutputInstance.h"
 #include "riglogic/riglogic/RigLogicImpl.h"
 #include "riglogic/riglogic/RigMetrics.h"
 #include "riglogic/types/Aliases.h"
 #include "riglogic/utils/Extd.h"
 
-#include <pma/PolyAllocator.h>
-#include <pma/resources/AlignedMemoryResource.h>
-
 #ifdef _MSC_VER
     #pragma warning(push)
     #pragma warning(disable : 4365 4987)
@@ -18,6 +19,7 @@
 #include <algorithm>
 #include <cassert>
 #include <cstdint>
+#include <numeric>
 #ifdef _MSC_VER
     #pragma warning(pop)
 #endif
@@ -29,13 +31,17 @@ namespace {
 const float controlsMin = 0.0f;
 const float controlsMax = 1.0f;
 
+inline std::uint16_t getMaxLODLevel(std::uint16_t lodCount) {
+    return (lodCount > 0u ? static_cast<std::uint16_t>(lodCount - 1) : static_cast<std::uint16_t>(0));
+}
+
 }  // namespace
 
 RigInstance* RigInstance::create(RigLogic* rigLogic, MemoryResource* memRes) {
     // NOLINTNEXTLINE(cppcoreguidelines-pro-type-static-cast-downcast)
-    auto ptr = static_cast<RigLogicImpl*>(rigLogic);
+    auto pRigLogic = static_cast<RigLogicImpl*>(rigLogic);
     PolyAllocator<RigInstanceImpl> alloc{memRes};
-    return alloc.newObject(ptr->getRigMetrics(), memRes);
+    return alloc.newObject(pRigLogic->getRigMetrics(), pRigLogic, memRes);
 }
 
 void RigInstance::destroy(RigInstance* instance) {
@@ -47,46 +53,122 @@ void RigInstance::destroy(RigInstance* instance) {
 
 RigInstance::~RigInstance() = default;
 
-RigInstanceImpl::RigInstanceImpl(const RigMetrics& metrics, MemoryResource* memRes_) :
+RigInstanceImpl::RigInstanceImpl(const RigMetrics& metrics, RigLogicImpl* rigLogic, MemoryResource* memRes_) :
     memRes{memRes_},
-    lodMaxLevel{metrics.lodCount > 0u ? static_cast<std::uint16_t>(metrics.lodCount - 1) : static_cast<std::uint16_t>(0)},
+    lodMaxLevel{getMaxLODLevel(metrics.lodCount)},
     lodLevel{},
     guiControlCount{metrics.guiControlCount},
     rawControlCount{metrics.rawControlCount},
-    guiControlBuffer{guiControlCount, {}, memRes},
-    inputBuffer{static_cast<std::size_t>(rawControlCount + metrics.psdCount), {}, memRes},
-    jointOutputs{metrics.jointAttributeCount, {}, memRes},
-    blendShapeOutputs{metrics.blendShapeCount, {}, memRes},
-    animatedMapOutputs{metrics.animatedMapCount, {}, memRes} {
+    psdControlCount{metrics.psdControlCount},
+    mlControlCount{metrics.mlControlCount},
+    neuralNetworkCount{metrics.neuralNetworkCount},
+    controlsInstance{rigLogic->createControlsInstance(memRes)},
+    machineLearnedBehaviorInstance{rigLogic->createMachineLearnedBehaviorInstance(memRes)},
+    jointsInstance{rigLogic->createJointsInstance(memRes)},
+    blendShapesInstance{rigLogic->createBlendShapesInstance(memRes)},
+    animatedMapsInstance{rigLogic->createAnimatedMapsInstance(memRes)} {
 }
 
 std::uint16_t RigInstanceImpl::getGUIControlCount() const {
     return guiControlCount;
 }
 
+float RigInstanceImpl::getGUIControl(std::uint16_t index) const {
+    return getGUIControlValues()[index];
+}
+
 void RigInstanceImpl::setGUIControl(std::uint16_t index, float value) {
-    assert(index < guiControlBuffer.size());
+    auto guiControlBuffer = controlsInstance->getGUIControlBuffer();
     guiControlBuffer[index] = value;
 }
 
+ConstArrayView<float> RigInstanceImpl::getGUIControlValues() const {
+    return controlsInstance->getGUIControlBuffer();
+}
+
 void RigInstanceImpl::setGUIControlValues(const float* values) {
-    std::copy(values, values + guiControlCount, guiControlBuffer.begin());
+    auto guiControlBuffer = controlsInstance->getGUIControlBuffer();
+    #if (defined(_MSC_VER) && _MSC_VER >= 1900)
+        std::copy(values,
+                  values + guiControlCount,
+                  stdext::make_checked_array_iterator(guiControlBuffer.data(), guiControlBuffer.size()));
+    #else
+        std::copy(values, values + guiControlCount, guiControlBuffer.begin());
+    #endif
 }
 
 std::uint16_t RigInstanceImpl::getRawControlCount() const {
     return rawControlCount;
 }
 
+float RigInstanceImpl::getRawControl(std::uint16_t index) const {
+    return getRawControlValues()[index];
+}
+
 void RigInstanceImpl::setRawControl(std::uint16_t index, float value) {
-    assert(index < inputBuffer.size());
+    auto inputBuffer = controlsInstance->getInputBuffer();
     inputBuffer[index] = extd::clamp(value, controlsMin, controlsMax);
 }
 
+ConstArrayView<float> RigInstanceImpl::getRawControlValues() const {
+    return controlsInstance->getInputBuffer().subview(0ul, rawControlCount);
+}
+
 void RigInstanceImpl::setRawControlValues(const float* values) {
     auto clamper = [](float v) {
             return extd::clamp(v, controlsMin, controlsMax);
         };
-    std::transform(values, values + rawControlCount, inputBuffer.begin(), clamper);
+    auto inputBuffer = controlsInstance->getInputBuffer();
+
+    #if (defined(_MSC_VER) && _MSC_VER >= 1900)
+        std::transform(values,
+                       values + rawControlCount,
+                       stdext::make_checked_array_iterator(inputBuffer.data(), inputBuffer.size()),
+                       clamper);
+    #else
+        std::transform(values, values + rawControlCount, inputBuffer.begin(), clamper);
+    #endif
+}
+
+std::uint16_t RigInstanceImpl::getPSDControlCount() const {
+    return psdControlCount;
+}
+
+float RigInstanceImpl::getPSDControl(std::uint16_t index) const {
+    return getPSDControlValues()[index];
+}
+
+ConstArrayView<float> RigInstanceImpl::getPSDControlValues() const {
+    return controlsInstance->getInputBuffer().subview(rawControlCount, psdControlCount);
+}
+
+std::uint16_t RigInstanceImpl::getMLControlCount() const {
+    return mlControlCount;
+}
+
+float RigInstanceImpl::getMLControl(std::uint16_t index) const {
+    return getMLControlValues()[index];
+}
+
+ConstArrayView<float> RigInstanceImpl::getMLControlValues() const {
+    const auto mlControlsOffset = static_cast<std::size_t>(rawControlCount + psdControlCount);
+    return controlsInstance->getInputBuffer().subview(mlControlsOffset, mlControlCount);
+}
+
+std::uint16_t RigInstanceImpl::getNeuralNetworkCount() const {
+    return neuralNetworkCount;
+}
+
+float RigInstanceImpl::getNeuralNetworkMask(std::uint16_t neuralNetIndex) const {
+    assert(neuralNetIndex < neuralNetworkCount);
+    auto maskBuffer = machineLearnedBehaviorInstance->getMaskBuffer();
+    return maskBuffer[neuralNetIndex];
+}
+
+void RigInstanceImpl::setNeuralNetworkMask(std::uint16_t neuralNetIndex, float value) {
+    assert(neuralNetIndex < neuralNetworkCount);
+    auto maskBuffer = machineLearnedBehaviorInstance->getMaskBuffer();
+    maskBuffer[neuralNetIndex] = value;
 }
 
 std::uint16_t RigInstanceImpl::getLOD() const {
@@ -95,6 +177,9 @@ std::uint16_t RigInstanceImpl::getLOD() const {
 
 void RigInstanceImpl::setLOD(std::uint16_t level) {
     if (level != lodLevel) {
+        auto jointOutputs = jointsInstance->getOutputBuffer();
+        auto blendShapeOutputs = blendShapesInstance->getOutputBuffer();
+        auto animatedMapOutputs = animatedMapsInstance->getOutputBuffer();
         std::fill(jointOutputs.begin(), jointOutputs.end(), 0.0f);
         std::fill(blendShapeOutputs.begin(), blendShapeOutputs.end(), 0.0f);
         std::fill(animatedMapOutputs.begin(), animatedMapOutputs.end(), 0.0f);
@@ -102,48 +187,41 @@ void RigInstanceImpl::setLOD(std::uint16_t level) {
     lodLevel = extd::clamp(level, static_cast<std::uint16_t>(0), lodMaxLevel);
 }
 
-ArrayView<float> RigInstanceImpl::getGUIControlValues() {
-    return ArrayView<float>{guiControlBuffer};
-}
-
-ConstArrayView<float> RigInstanceImpl::getGUIControlValues() const {
-    return ConstArrayView<float>{guiControlBuffer};
-}
-
-ArrayView<float> RigInstanceImpl::getInputValues() {
-    return ArrayView<float>{inputBuffer};
+ConstArrayView<float> RigInstanceImpl::getRawJointOutputs() const {
+    return jointsInstance->getOutputBuffer();
 }
 
-ConstArrayView<float> RigInstanceImpl::getInputValues() const {
-    return ConstArrayView<float>{inputBuffer};
+TransformationArrayView RigInstanceImpl::getJointOutputs() const {
+    auto jointOutputs = jointsInstance->getOutputBuffer();
+    return TransformationArrayView{jointOutputs.data(), jointOutputs.size()};
 }
 
-ArrayView<float> RigInstanceImpl::getRawJointOutputs() {
-    return ArrayView<float>{jointOutputs};
+ConstArrayView<float> RigInstanceImpl::getBlendShapeOutputs() const {
+    return blendShapesInstance->getOutputBuffer();
 }
 
-ConstArrayView<float> RigInstanceImpl::getRawJointOutputs() const {
-    return ConstArrayView<float>{jointOutputs};
+ConstArrayView<float> RigInstanceImpl::getAnimatedMapOutputs() const {
+    return animatedMapsInstance->getOutputBuffer();
 }
 
-TransformationArrayView RigInstanceImpl::getJointOutputs() const {
-    return TransformationArrayView{jointOutputs.data(), jointOutputs.size()};
+ControlsInputInstance* RigInstanceImpl::getControlsInputInstance() {
+    return controlsInstance.get();
 }
 
-ArrayView<float> RigInstanceImpl::getBlendShapeOutputs() {
-    return ArrayView<float>{blendShapeOutputs};
+MachineLearnedBehaviorOutputInstance* RigInstanceImpl::getMachineLearnedBehaviorOutputInstance() {
+    return machineLearnedBehaviorInstance.get();
 }
 
-ConstArrayView<float> RigInstanceImpl::getBlendShapeOutputs() const {
-    return ConstArrayView<float>{blendShapeOutputs};
+JointsOutputInstance* RigInstanceImpl::getJointsOutputInstance() {
+    return jointsInstance.get();
 }
 
-ArrayView<float> RigInstanceImpl::getAnimatedMapOutputs() {
-    return ArrayView<float>{animatedMapOutputs};
+BlendShapesOutputInstance* RigInstanceImpl::getBlendShapesOutputInstance() {
+    return blendShapesInstance.get();
 }
 
-ConstArrayView<float> RigInstanceImpl::getAnimatedMapOutputs() const {
-    return ConstArrayView<float>{animatedMapOutputs};
+AnimatedMapsOutputInstance* RigInstanceImpl::getAnimatedMapOutputInstance() {
+    return animatedMapsInstance.get();
 }
 
 MemoryResource* RigInstanceImpl::getMemoryResource() {
diff --git a/Engine/Plugins/Animation/RigLogic/Source/RigLogicLib/Private/riglogic/riglogic/RigInstanceImpl.h b/Engine/Plugins/Animation/RigLogic/Source/RigLogicLib/Private/riglogic/riglogic/RigInstanceImpl.h
--- a/Engine/Plugins/Animation/RigLogic/Source/RigLogicLib/Private/riglogic/riglogic/RigInstanceImpl.h
+++ b/Engine/Plugins/Animation/RigLogic/Source/RigLogicLib/Private/riglogic/riglogic/RigInstanceImpl.h
@@ -3,6 +3,11 @@
 #pragma once
 
 #include "riglogic/TypeDefs.h"
+#include "riglogic/animatedmaps/AnimatedMapsOutputInstance.h"
+#include "riglogic/blendshapes/BlendShapesOutputInstance.h"
+#include "riglogic/controls/ControlsInputInstance.h"
+#include "riglogic/joints/JointsOutputInstance.h"
+#include "riglogic/ml/MachineLearnedBehaviorOutputInstance.h"
 #include "riglogic/riglogic/RigInstance.h"
 #include "riglogic/riglogic/RigMetrics.h"
 
@@ -10,33 +15,50 @@
 
 namespace rl4 {
 
+class RigLogicImpl;
+
 class RigInstanceImpl : public RigInstance {
     public:
-        RigInstanceImpl(const RigMetrics& metrics, MemoryResource* memRes_);
+        RigInstanceImpl(const RigMetrics& metrics, RigLogicImpl* rigLogic, MemoryResource* memRes_);
 
         std::uint16_t getGUIControlCount() const override;
+        float getGUIControl(std::uint16_t index) const override;
         void setGUIControl(std::uint16_t index, float value) override;
+        ConstArrayView<float> getGUIControlValues() const override;
         void setGUIControlValues(const float* values) override;
 
         std::uint16_t getRawControlCount() const override;
+        float getRawControl(std::uint16_t index) const override;
         void setRawControl(std::uint16_t index, float value) override;
+        ConstArrayView<float> getRawControlValues() const override;
         void setRawControlValues(const float* values) override;
 
+        std::uint16_t getPSDControlCount() const override;
+        float getPSDControl(std::uint16_t index) const override;
+        ConstArrayView<float> getPSDControlValues() const override;
+
+        std::uint16_t getMLControlCount() const override;
+        float getMLControl(std::uint16_t index) const override;
+        ConstArrayView<float> getMLControlValues() const override;
+
+        std::uint16_t getNeuralNetworkCount() const override;
+        float getNeuralNetworkMask(std::uint16_t neuralNetIndex) const override;
+        void setNeuralNetworkMask(std::uint16_t neuralNetIndex, float value) override;
+
         std::uint16_t getLOD() const override;
         void setLOD(std::uint16_t level) override;
 
-        ArrayView<float> getGUIControlValues();
-        ConstArrayView<float> getGUIControlValues() const;
-        ArrayView<float> getInputValues();
-        ConstArrayView<float> getInputValues() const;
-        ArrayView<float> getRawJointOutputs();
         ConstArrayView<float> getRawJointOutputs() const override;
         TransformationArrayView getJointOutputs() const override;
-        ArrayView<float> getBlendShapeOutputs();
         ConstArrayView<float> getBlendShapeOutputs() const override;
-        ArrayView<float> getAnimatedMapOutputs();
         ConstArrayView<float> getAnimatedMapOutputs() const override;
 
+        ControlsInputInstance* getControlsInputInstance();
+        MachineLearnedBehaviorOutputInstance* getMachineLearnedBehaviorOutputInstance();
+        JointsOutputInstance* getJointsOutputInstance();
+        BlendShapesOutputInstance* getBlendShapesOutputInstance();
+        AnimatedMapsOutputInstance* getAnimatedMapOutputInstance();
+
         MemoryResource* getMemoryResource();
 
     private:
@@ -46,12 +68,15 @@ class RigInstanceImpl : public RigInstance {
         std::uint16_t lodLevel;
         std::uint16_t guiControlCount;
         std::uint16_t rawControlCount;
+        std::uint16_t psdControlCount;
+        std::uint16_t mlControlCount;
+        std::uint16_t neuralNetworkCount;
 
-        Vector<float> guiControlBuffer;
-        Vector<float> inputBuffer;
-        AlignedVector<float> jointOutputs;
-        Vector<float> blendShapeOutputs;
-        Vector<float> animatedMapOutputs;
+        ControlsInputInstance::Pointer controlsInstance;
+        MachineLearnedBehaviorOutputInstance::Pointer machineLearnedBehaviorInstance;
+        JointsOutputInstance::Pointer jointsInstance;
+        BlendShapesOutputInstance::Pointer blendShapesInstance;
+        AnimatedMapsOutputInstance::Pointer animatedMapsInstance;
 
 };
 
