Update RigLogic to v9.1.1 and GeneSplicer to v9.4.0

diff --git a/Engine/Plugins/Animation/RigLogic/Source/RigLogicLib/Private/riglogic/riglogic/RigLogicImpl.cpp b/Engine/Plugins/Animation/RigLogic/Source/RigLogicLib/Private/riglogic/riglogic/RigLogicImpl.cpp
--- a/Engine/Plugins/Animation/RigLogic/Source/RigLogicLib/Private/riglogic/riglogic/RigLogicImpl.cpp
+++ b/Engine/Plugins/Animation/RigLogic/Source/RigLogicLib/Private/riglogic/riglogic/RigLogicImpl.cpp
@@ -2,25 +2,24 @@
 
 #include "riglogic/riglogic/RigLogicImpl.h"
 
+#include "riglogic/TypeDefs.h"
 #include "riglogic/animatedmaps/AnimatedMapsFactory.h"
 #include "riglogic/blendshapes/BlendShapesFactory.h"
 #include "riglogic/controls/ControlsFactory.h"
 #include "riglogic/joints/JointsFactory.h"
+#include "riglogic/ml/MachineLearnedBehaviorFactory.h"
+#include "riglogic/riglogic/ConfigurationSerializer.h"
 #include "riglogic/riglogic/RigInstanceImpl.h"
 #include "riglogic/riglogic/RigMetrics.h"
-
-#include <dna/layers/BehaviorReader.h>
-#include <pma/PolyAllocator.h>
-#include <pma/resources/AlignedMemoryResource.h>
-#include <terse/archives/binary/InputArchive.h>
-#include <terse/archives/binary/OutputArchive.h>
-#include <trio/Stream.h>
+#include "riglogic/utils/Extd.h"
 
 #ifdef _MSC_VER
     #pragma warning(push)
     #pragma warning(disable : 4365 4987)
 #endif
 #include <cstddef>
+#include <memory>
+#include <numeric>
 #include <utility>
 #ifdef _MSC_VER
     #pragma warning(pop)
@@ -33,34 +32,49 @@ static RigInstanceImpl* castInstance(RigInstance* instance) {
     return static_cast<RigInstanceImpl*>(instance);
 }
 
+static RigMetrics::Pointer computeRigMetrics(const dna::Reader* reader, MemoryResource* memRes) {
+    RigMetrics::Pointer metrics = UniqueInstance<RigMetrics>::with(memRes).create(memRes);
+    metrics->lodCount = reader->getLODCount();
+    metrics->guiControlCount = reader->getGUIControlCount();
+    metrics->rawControlCount = reader->getRawControlCount();
+    metrics->psdControlCount = reader->getPSDCount();
+    metrics->mlControlCount = reader->getMLControlCount();
+    metrics->jointAttributeCount = reader->getJointRowCount();
+    metrics->blendShapeCount = reader->getBlendShapeChannelCount();
+    metrics->animatedMapCount = reader->getAnimatedMapCount();
+    metrics->neuralNetworkCount = reader->getNeuralNetworkCount();
+    return metrics;
+}
+
 RigLogic::~RigLogic() = default;
 
-RigLogic* RigLogic::create(const dna::BehaviorReader* reader, CalculationType calculationType, MemoryResource* memRes) {
-    RigMetrics metrics{};
-    metrics.lodCount = reader->getLODCount();
-    metrics.guiControlCount = reader->getGUIControlCount();
-    metrics.rawControlCount = reader->getRawControlCount();
-    metrics.psdCount = reader->getPSDCount();
-    metrics.jointAttributeCount = reader->getJointRowCount();
-    metrics.blendShapeCount = reader->getBlendShapeChannelCount();
-    metrics.animatedMapCount = reader->getAnimatedMapCount();
-
-    Configuration config{calculationType};
-    auto controls = ControlsFactory::create(reader, memRes);
+RigLogic* RigLogic::create(const dna::Reader* reader, Configuration config, MemoryResource* memRes) {
+    auto metrics = computeRigMetrics(reader, memRes);
+
+    auto controls = ControlsFactory::create(config, reader, memRes);
+    auto machineLearnedBlendShapes = MachineLearnedBehaviorFactory::create(config, reader, memRes);
     auto joints = JointsFactory::create(config, reader, memRes);
-    auto blendShapes = BlendShapesFactory::create(reader, memRes);
-    auto animatedMaps = AnimatedMapsFactory::create(reader, memRes);
+    auto blendShapes = BlendShapesFactory::create(config, reader, memRes);
+    auto animatedMaps = AnimatedMapsFactory::create(config, reader, memRes);
 
     PolyAllocator<RigLogicImpl> alloc{memRes};
     return alloc.newObject(config,
-                           metrics,
+                           std::move(metrics),
                            std::move(controls),
+                           std::move(machineLearnedBlendShapes),
                            std::move(joints),
                            std::move(blendShapes),
                            std::move(animatedMaps),
                            memRes);
 }
 
+void RigLogic::destroy(RigLogic* instance) {
+    // NOLINTNEXTLINE(cppcoreguidelines-pro-type-static-cast-downcast)
+    auto ptr = static_cast<RigLogicImpl*>(instance);
+    PolyAllocator<RigLogicImpl> alloc{ptr->getMemoryResource()};
+    alloc.deleteObject(ptr);
+}
+
 RigLogic* RigLogic::restore(BoundedIOStream* source, MemoryResource* memRes) {
     PolyAllocator<RigLogicImpl> alloc{memRes};
 
@@ -69,50 +83,52 @@ RigLogic* RigLogic::restore(BoundedIOStream* source, MemoryResource* memRes) {
     Configuration config;
     archive >> config;
 
-    RigMetrics metrics;
-    archive >> metrics;
+    RigMetrics::Pointer metrics = UniqueInstance<RigMetrics>::with(memRes).create(memRes);
+    archive >> *metrics;
+
+    auto controls = ControlsFactory::create(config, *metrics, memRes);
+    auto machineLearnedBlendShapes = MachineLearnedBehaviorFactory::create(config, memRes);
+    auto joints = JointsFactory::create(config, *metrics, memRes);
+    auto blendShapes = BlendShapesFactory::create(config, *metrics, memRes);
+    auto animatedMaps = AnimatedMapsFactory::create(config, *metrics, memRes);
 
-    auto controls = ControlsFactory::create(memRes);
-    auto joints = JointsFactory::create(config, metrics, memRes);
-    auto blendShapes = BlendShapesFactory::create(memRes);
-    auto animatedMaps = AnimatedMapsFactory::create(memRes);
+    terse::VirtualSerializerProxy<AnimatedMaps> animatedMapsProxy{animatedMaps.get()};
+    terse::VirtualSerializerProxy<BlendShapes> blendShapesProxy{blendShapes.get()};
 
-    archive >> *controls >> *joints >> *blendShapes >> *animatedMaps;
+    archive >> *controls >> *machineLearnedBlendShapes >> *joints >> blendShapesProxy >> animatedMapsProxy;
     return alloc.newObject(config,
-                           metrics,
+                           std::move(metrics),
                            std::move(controls),
+                           std::move(machineLearnedBlendShapes),
                            std::move(joints),
                            std::move(blendShapes),
                            std::move(animatedMaps),
                            memRes);
 }
 
-void RigLogic::destroy(RigLogic* instance) {
-    // NOLINTNEXTLINE(cppcoreguidelines-pro-type-static-cast-downcast)
-    auto ptr = static_cast<RigLogicImpl*>(instance);
-    PolyAllocator<RigLogicImpl> alloc{ptr->getMemoryResource()};
-    alloc.deleteObject(ptr);
-}
-
 RigLogicImpl::RigLogicImpl(Configuration config_,
-                           RigMetrics metrics_,
-                           ControlsPtr controls_,
-                           JointsPtr joints_,
-                           BlendShapesPtr blendShapes_,
-                           AnimatedMapsPtr animatedMaps_,
+                           RigMetrics::Pointer metrics_,
+                           Controls::Pointer controls_,
+                           MachineLearnedBehavior::Pointer machineLearnedBehavior_,
+                           Joints::Pointer joints_,
+                           BlendShapes::Pointer blendShapes_,
+                           AnimatedMaps::Pointer animatedMaps_,
                            MemoryResource* memRes_) :
     memRes{memRes_},
+    config{config_},
+    metrics{std::move(metrics_)},
     controls{std::move(controls_)},
+    machineLearnedBehavior{std::move(machineLearnedBehavior_)},
     joints{std::move(joints_)},
     blendShapes{std::move(blendShapes_)},
-    animatedMaps{std::move(animatedMaps_)},
-    config{config_},
-    metrics{metrics_} {
+    animatedMaps{std::move(animatedMaps_)} {
 }
 
 void RigLogicImpl::dump(BoundedIOStream* destination) const {
     terse::BinaryOutputArchive<BoundedIOStream> archive{destination};
-    archive << config << metrics << *controls << *joints << *blendShapes << *animatedMaps;
+    terse::VirtualSerializerProxy<AnimatedMaps> animatedMapsProxy{animatedMaps.get()};
+    terse::VirtualSerializerProxy<BlendShapes> blendShapesProxy{blendShapes.get()};
+    archive << config << *metrics << *controls << *machineLearnedBehavior << *joints << blendShapesProxy << animatedMapsProxy;
 }
 
 const Configuration& RigLogicImpl::getConfiguration() const {
@@ -120,11 +136,11 @@ const Configuration& RigLogicImpl::getConfiguration() const {
 }
 
 const RigMetrics& RigLogicImpl::getRigMetrics() const {
-    return metrics;
+    return *metrics;
 }
 
 std::uint16_t RigLogicImpl::getLODCount() const {
-    return metrics.lodCount;
+    return metrics->lodCount;
 }
 
 ConstArrayView<float> RigLogicImpl::getRawNeutralJointValues() const {
@@ -135,46 +151,110 @@ TransformationArrayView RigLogicImpl::getNeutralJointValues() const {
     return joints->getNeutralValues();
 }
 
+ConstArrayView<std::uint16_t> RigLogicImpl::getJointVariableAttributeIndices(std::uint16_t lod) const {
+    return joints->getVariableAttributeIndices(lod);
+}
+
 std::uint16_t RigLogicImpl::getJointGroupCount() const {
     return joints->getJointGroupCount();
 }
 
-ConstArrayView<std::uint16_t> RigLogicImpl::getJointVariableAttributeIndices(std::uint16_t lod) const {
-    return joints->getVariableAttributeIndices(lod);
+std::uint16_t RigLogicImpl::getNeuralNetworkCount() const {
+    return metrics->neuralNetworkCount;
+}
+
+std::uint16_t RigLogicImpl::getMeshCount() const {
+    return machineLearnedBehavior->getMeshCount();
+}
+
+std::uint16_t RigLogicImpl::getMeshRegionCount(std::uint16_t meshIndex) const {
+    return machineLearnedBehavior->getMeshRegionCount(meshIndex);
+}
+
+ConstArrayView<std::uint16_t> RigLogicImpl::getNeuralNetworkIndices(std::uint16_t meshIndex, std::uint16_t regionIndex) const {
+    return machineLearnedBehavior->getNeuralNetworkIndices(meshIndex, regionIndex);
+}
+
+ControlsInputInstance::Pointer RigLogicImpl::createControlsInstance(MemoryResource* instanceMemRes) const {
+    return controls->createInstance(instanceMemRes);
+}
+
+MachineLearnedBehaviorOutputInstance::Pointer RigLogicImpl::createMachineLearnedBehaviorInstance(MemoryResource* instanceMemRes)
+const {
+    return machineLearnedBehavior->createInstance(instanceMemRes);
+}
+
+JointsOutputInstance::Pointer RigLogicImpl::createJointsInstance(MemoryResource* instanceMemRes) const {
+    return joints->createInstance(instanceMemRes);
+}
+
+BlendShapesOutputInstance::Pointer RigLogicImpl::createBlendShapesInstance(MemoryResource* instanceMemRes) const {
+    return blendShapes->createInstance(instanceMemRes);
+}
+
+AnimatedMapsOutputInstance::Pointer RigLogicImpl::createAnimatedMapsInstance(MemoryResource* instanceMemRes) const {
+    return animatedMaps->createInstance(instanceMemRes);
 }
 
 void RigLogicImpl::mapGUIToRawControls(RigInstance* instance) const {
     auto pRigInstance = castInstance(instance);
-    controls->mapGUIToRaw(pRigInstance->getGUIControlValues(), pRigInstance->getInputValues());
+    controls->mapGUIToRaw(pRigInstance->getControlsInputInstance());
+}
+
+void RigLogicImpl::mapRawToGUIControls(RigInstance* instance) const {
+    auto pRigInstance = castInstance(instance);
+    controls->mapRawToGUI(pRigInstance->getControlsInputInstance());
 }
 
 void RigLogicImpl::calculateControls(RigInstance* instance) const {
     auto pRigInstance = castInstance(instance);
-    controls->calculate(pRigInstance->getInputValues());
+    controls->calculate(pRigInstance->getControlsInputInstance());
+}
+
+void RigLogicImpl::calculateMachineLearnedBehaviorControls(RigInstance* instance) const {
+    auto pRigInstance = castInstance(instance);
+    machineLearnedBehavior->calculate(pRigInstance->getControlsInputInstance(),
+                                      pRigInstance->getMachineLearnedBehaviorOutputInstance(),
+                                      pRigInstance->getLOD());
+}
+
+void RigLogicImpl::calculateMachineLearnedBehaviorControls(RigInstance* instance, std::uint16_t neuralNetIndex) const {
+    auto pRigInstance = castInstance(instance);
+    machineLearnedBehavior->calculate(pRigInstance->getControlsInputInstance(),
+                                      pRigInstance->getMachineLearnedBehaviorOutputInstance(),
+                                      pRigInstance->getLOD(),
+                                      neuralNetIndex);
 }
 
 void RigLogicImpl::calculateJoints(RigInstance* instance) const {
     auto pRigInstance = castInstance(instance);
-    joints->calculate(pRigInstance->getInputValues(), pRigInstance->getRawJointOutputs(), pRigInstance->getLOD());
+    joints->calculate(pRigInstance->getControlsInputInstance(), pRigInstance->getJointsOutputInstance(), pRigInstance->getLOD());
 }
 
 void RigLogicImpl::calculateJoints(RigInstance* instance, std::uint16_t jointGroupIndex) const {
     auto pRigInstance = castInstance(instance);
-    joints->calculate(pRigInstance->getInputValues(), pRigInstance->getRawJointOutputs(), pRigInstance->getLOD(),
+    joints->calculate(pRigInstance->getControlsInputInstance(),
+                      pRigInstance->getJointsOutputInstance(),
+                      pRigInstance->getLOD(),
                       jointGroupIndex);
 }
 
 void RigLogicImpl::calculateBlendShapes(RigInstance* instance) const {
     auto pRigInstance = castInstance(instance);
-    blendShapes->calculate(pRigInstance->getInputValues(), pRigInstance->getBlendShapeOutputs(), pRigInstance->getLOD());
+    blendShapes->calculate(pRigInstance->getControlsInputInstance(),
+                           pRigInstance->getBlendShapesOutputInstance(),
+                           pRigInstance->getLOD());
 }
 
 void RigLogicImpl::calculateAnimatedMaps(RigInstance* instance) const {
     auto pRigInstance = castInstance(instance);
-    animatedMaps->calculate(pRigInstance->getInputValues(), pRigInstance->getAnimatedMapOutputs(), pRigInstance->getLOD());
+    animatedMaps->calculate(pRigInstance->getControlsInputInstance(),
+                            pRigInstance->getAnimatedMapOutputInstance(),
+                            pRigInstance->getLOD());
 }
 
 void RigLogicImpl::calculate(RigInstance* instance) const {
+    calculateMachineLearnedBehaviorControls(instance);
     calculateControls(instance);
     calculateJoints(instance);
     calculateBlendShapes(instance);
diff --git a/Engine/Plugins/Animation/RigLogic/Source/RigLogicLib/Private/riglogic/riglogic/RigLogicImpl.h b/Engine/Plugins/Animation/RigLogic/Source/RigLogicLib/Private/riglogic/riglogic/RigLogicImpl.h
--- a/Engine/Plugins/Animation/RigLogic/Source/RigLogicLib/Private/riglogic/riglogic/RigLogicImpl.h
+++ b/Engine/Plugins/Animation/RigLogic/Source/RigLogicLib/Private/riglogic/riglogic/RigLogicImpl.h
@@ -2,40 +2,32 @@
 
 #pragma once
 
+#include "riglogic/TypeDefs.h"
 #include "riglogic/animatedmaps/AnimatedMaps.h"
+#include "riglogic/animatedmaps/AnimatedMapsOutputInstance.h"
 #include "riglogic/blendshapes/BlendShapes.h"
+#include "riglogic/blendshapes/BlendShapesOutputInstance.h"
 #include "riglogic/controls/Controls.h"
+#include "riglogic/controls/ControlsInputInstance.h"
 #include "riglogic/joints/Joints.h"
+#include "riglogic/joints/JointsOutputInstance.h"
+#include "riglogic/ml/MachineLearnedBehavior.h"
+#include "riglogic/ml/MachineLearnedBehaviorOutputInstance.h"
 #include "riglogic/riglogic/Configuration.h"
 #include "riglogic/riglogic/RigLogic.h"
 #include "riglogic/riglogic/RigMetrics.h"
 
-#ifdef _MSC_VER
-    #pragma warning(push)
-    #pragma warning(disable : 4365 4987)
-#endif
-#include <functional>
-#include <memory>
-#ifdef _MSC_VER
-    #pragma warning(pop)
-#endif
-
 namespace rl4 {
 
 class RigLogicImpl : public RigLogic {
-    private:
-        using ControlsPtr = std::unique_ptr<Controls, std::function<void (Controls*)> >;
-        using JointsPtr = std::unique_ptr<Joints, std::function<void (Joints*)> >;
-        using BlendShapesPtr = std::unique_ptr<BlendShapes, std::function<void (BlendShapes*)> >;
-        using AnimatedMapsPtr = std::unique_ptr<AnimatedMaps, std::function<void (AnimatedMaps*)> >;
-
     public:
         RigLogicImpl(Configuration config_,
-                     RigMetrics metrics_,
-                     ControlsPtr controls_,
-                     JointsPtr joints_,
-                     BlendShapesPtr blendShapes_,
-                     AnimatedMapsPtr animatedMaps_,
+                     RigMetrics::Pointer metrics_,
+                     Controls::Pointer controls_,
+                     MachineLearnedBehavior::Pointer machineLearnedBehavior_,
+                     Joints::Pointer joints_,
+                     BlendShapes::Pointer blendShapes_,
+                     AnimatedMaps::Pointer animatedMaps_,
                      MemoryResource* memRes_);
 
         void dump(BoundedIOStream* destination) const override;
@@ -44,10 +36,24 @@ class RigLogicImpl : public RigLogic {
         std::uint16_t getLODCount() const override;
         ConstArrayView<float> getRawNeutralJointValues() const override;
         TransformationArrayView getNeutralJointValues() const override;
-        std::uint16_t getJointGroupCount() const override;
         ConstArrayView<std::uint16_t> getJointVariableAttributeIndices(std::uint16_t lod) const override;
+        std::uint16_t getJointGroupCount() const override;
+        std::uint16_t getNeuralNetworkCount() const override;
+        std::uint16_t getMeshCount() const override;
+        std::uint16_t getMeshRegionCount(std::uint16_t meshIndex) const override;
+        ConstArrayView<std::uint16_t> getNeuralNetworkIndices(std::uint16_t meshIndex, std::uint16_t regionIndex) const override;
+
+        ControlsInputInstance::Pointer createControlsInstance(MemoryResource* instanceMemRes) const;
+        MachineLearnedBehaviorOutputInstance::Pointer createMachineLearnedBehaviorInstance(MemoryResource* instanceMemRes) const;
+        JointsOutputInstance::Pointer createJointsInstance(MemoryResource* instanceMemRes) const;
+        BlendShapesOutputInstance::Pointer createBlendShapesInstance(MemoryResource* instanceMemRes) const;
+        AnimatedMapsOutputInstance::Pointer createAnimatedMapsInstance(MemoryResource* instanceMemRes) const;
+
         void mapGUIToRawControls(RigInstance* instance) const override;
+        void mapRawToGUIControls(RigInstance* instance) const override;
         void calculateControls(RigInstance* instance) const override;
+        void calculateMachineLearnedBehaviorControls(RigInstance* instance) const override;
+        void calculateMachineLearnedBehaviorControls(RigInstance* instance, std::uint16_t neuralNetIndex) const override;
         void calculateJoints(RigInstance* instance) const override;
         void calculateJoints(RigInstance* instance, std::uint16_t jointGroupIndex) const override;
         void calculateBlendShapes(RigInstance* instance) const override;
@@ -58,12 +64,13 @@ class RigLogicImpl : public RigLogic {
 
     private:
         MemoryResource* memRes;
-        ControlsPtr controls;
-        JointsPtr joints;
-        BlendShapesPtr blendShapes;
-        AnimatedMapsPtr animatedMaps;
         Configuration config;
-        RigMetrics metrics;
+        RigMetrics::Pointer metrics;
+        Controls::Pointer controls;
+        MachineLearnedBehavior::Pointer machineLearnedBehavior;
+        Joints::Pointer joints;
+        BlendShapes::Pointer blendShapes;
+        AnimatedMaps::Pointer animatedMaps;
 
 };
 
diff --git a/Engine/Plugins/Animation/RigLogic/Source/RigLogicLib/Private/riglogic/riglogic/RigMetrics.h b/Engine/Plugins/Animation/RigLogic/Source/RigLogicLib/Private/riglogic/riglogic/RigMetrics.h
--- a/Engine/Plugins/Animation/RigLogic/Source/RigLogicLib/Private/riglogic/riglogic/RigMetrics.h
+++ b/Engine/Plugins/Animation/RigLogic/Source/RigLogicLib/Private/riglogic/riglogic/RigMetrics.h
@@ -2,28 +2,48 @@
 
 #pragma once
 
+#include "riglogic/TypeDefs.h"
+
 #include <cstdint>
 
 namespace rl4 {
 
 struct RigMetrics {
+    using Pointer = UniqueInstance<RigMetrics>::PointerType;
+
     std::uint16_t lodCount;
     std::uint16_t guiControlCount;
     std::uint16_t rawControlCount;
-    std::uint16_t psdCount;
+    std::uint16_t psdControlCount;
+    std::uint16_t mlControlCount;
     std::uint16_t jointAttributeCount;
     std::uint16_t blendShapeCount;
     std::uint16_t animatedMapCount;
+    std::uint16_t neuralNetworkCount;
+
+    explicit RigMetrics(MemoryResource*  /*unused*/) :
+        lodCount{},
+        guiControlCount{},
+        rawControlCount{},
+        psdControlCount{},
+        mlControlCount{},
+        jointAttributeCount{},
+        blendShapeCount{},
+        animatedMapCount{},
+        neuralNetworkCount{} {
+    }
 
     template<class Archive>
     void serialize(Archive& archive) {
         archive(lodCount,
                 guiControlCount,
                 rawControlCount,
-                psdCount,
+                psdControlCount,
+                mlControlCount,
                 jointAttributeCount,
                 blendShapeCount,
-                animatedMapCount);
+                animatedMapCount,
+                neuralNetworkCount);
     }
 
 };
diff --git a/Engine/Plugins/Animation/RigLogic/Source/RigLogicLib/Private/riglogic/system/simd/Detect.h b/Engine/Plugins/Animation/RigLogic/Source/RigLogicLib/Private/riglogic/system/simd/Detect.h
new file mode 100644
--- /dev/null
+++ b/Engine/Plugins/Animation/RigLogic/Source/RigLogicLib/Private/riglogic/system/simd/Detect.h
@@ -0,0 +1,10 @@
+// Copyright Epic Games, Inc. All Rights Reserved.
+
+#pragma once
+
+#if defined(RL_AUTODETECT_SSE) && !defined(RL_BUILD_WITH_SSE)
+    // Assume any x86 architecture we try to build on will support SSE
+    #if defined(i386) || defined(__i386) || defined(__i386__) || defined(__IA32__) || defined(_M_IX86) || defined(__X86__) || defined(_X86_) || defined(__THW_INTEL__) || defined(__I86__) || defined(__amd64__) || defined(__amd64) || defined(__x86_64__) || defined(__x86_64) || defined(_M_X64) || defined(_M_AMD64)
+        #define RL_BUILD_WITH_SSE 1
+    #endif
+#endif  // RL_AUTODETECT_SSE
diff --git a/Engine/Plugins/Animation/RigLogic/Source/RigLogicLib/Private/riglogic/system/simd/SIMD.h b/Engine/Plugins/Animation/RigLogic/Source/RigLogicLib/Private/riglogic/system/simd/SIMD.h
--- a/Engine/Plugins/Animation/RigLogic/Source/RigLogicLib/Private/riglogic/system/simd/SIMD.h
+++ b/Engine/Plugins/Animation/RigLogic/Source/RigLogicLib/Private/riglogic/system/simd/SIMD.h
@@ -2,6 +2,8 @@
 
 #pragma once
 
+#include "riglogic/system/simd/Detect.h"
+
 #if defined(RL_USE_HALF_FLOATS) && !defined(TRIMD_ENABLE_F16C)
     #define TRIMD_ENABLE_F16C
 #endif  // RL_USE_HALF_FLOATS
diff --git a/Engine/Plugins/Animation/RigLogic/Source/RigLogicLib/Private/riglogic/types/AlignedBlockView.h b/Engine/Plugins/Animation/RigLogic/Source/RigLogicLib/Private/riglogic/types/AlignedBlockView.h
new file mode 100644
--- /dev/null
+++ b/Engine/Plugins/Animation/RigLogic/Source/RigLogicLib/Private/riglogic/types/AlignedBlockView.h
@@ -0,0 +1,66 @@
+// Copyright Epic Games, Inc. All Rights Reserved.
+
+#pragma once
+
+#include "riglogic/utils/Extd.h"
+
+#include <cstdint>
+
+namespace rl4 {
+
+struct AlignedBlockView {
+    // Unaligned size
+    std::uint32_t size;
+    // If the last full block needs to be masked-off because it contains padding,
+    // this view will include that last full block as well
+    std::uint32_t sizeAlignedToLastFullBlock;
+    // If the last full block needs to be masked-off because it contains padding,
+    // that full block won't be included in this boundary as it needs special care
+    std::uint32_t sizeAlignedToSecondLastFullBlock;
+
+    AlignedBlockView() = default;
+    AlignedBlockView(std::uint32_t viewRowCount, std::uint32_t paddedRowCount, std::uint32_t blockHeight,
+                     std::uint32_t padTo) : size{viewRowCount} {
+        // Check whether the block was padded to blockHeight size or padTo size
+        const bool endsWithPadToBlock = (paddedRowCount % blockHeight) == padTo;
+        // Check if view boundary is within the region of the last padTo rows
+        const bool isViewBoundaryInLastPadToRows = (paddedRowCount - viewRowCount) < padTo;
+        // Check if view boundary is within the region of the last padTo-sized block
+        const bool isViewBoundaryInLastPadToBlock = (endsWithPadToBlock && isViewBoundaryInLastPadToRows);
+        // Round up the view boundary row to the nearest index that is modulo padTo or blockHeight
+        // If boundary is in the last padTo-sized block, it rounds to padTo, otherwise to blockHeight
+        const std::uint32_t alignTo = (isViewBoundaryInLastPadToBlock ? padTo : blockHeight);
+        const std::uint32_t upAlignedViewRowCount = extd::roundUp(viewRowCount, alignTo);
+        // Check if the unaligned boundary coincides with a block boundary
+        const bool isViewBoundaryOnBlockBoundary = (upAlignedViewRowCount == viewRowCount);
+        // Should the last blockHeight-sized block be masked-off or not
+        // If view boundary coincides with a block boundary, no mask-off is needed at all
+        // If view boundary is within the last padTo-sized block, no mask-off is needed for the last blockHeight-sized block
+        const bool maskOffLastBlock = (!isViewBoundaryOnBlockBoundary && !isViewBoundaryInLastPadToBlock);
+        // Find the last blockHeight-sized block boundary before the aligned view boundary
+        // Sometimes this equals the aligned view boundary, but not if the matrix is padded to padTo block-size
+        // and the view boundary is within that last padTo-sized block region
+        sizeAlignedToLastFullBlock = upAlignedViewRowCount - (upAlignedViewRowCount % blockHeight);
+        // This is one blockHeight before `sizeAlignedToLastFullBlock` because the last blockHeight-sized block
+        // might need special care to mask off results that came from rows after the view boundary
+        // If no mask-off is needed, it equals to `sizeAlignedToLastFullBlock`
+        sizeAlignedToSecondLastFullBlock = (maskOffLastBlock ? sizeAlignedToLastFullBlock - blockHeight
+                                        : sizeAlignedToLastFullBlock);
+    }
+
+    AlignedBlockView(std::uint32_t viewRowCount,
+                     std::uint32_t boundaryRowCountAlignedToLastFullBlock,
+                     std::uint32_t boundaryRowCountAlignedToSecondLastFullBlock) :
+        size{viewRowCount},
+        sizeAlignedToLastFullBlock{boundaryRowCountAlignedToLastFullBlock},
+        sizeAlignedToSecondLastFullBlock{boundaryRowCountAlignedToSecondLastFullBlock} {
+    }
+
+    template<class Archive>
+    void serialize(Archive& archive) {
+        archive(size, sizeAlignedToLastFullBlock, sizeAlignedToSecondLastFullBlock);
+    }
+
+};
+
+}  // namespace rl4
diff --git a/Engine/Plugins/Animation/RigLogic/Source/RigLogicLib/Private/riglogic/utils/Extd.h b/Engine/Plugins/Animation/RigLogic/Source/RigLogicLib/Private/riglogic/utils/Extd.h
--- a/Engine/Plugins/Animation/RigLogic/Source/RigLogicLib/Private/riglogic/utils/Extd.h
+++ b/Engine/Plugins/Animation/RigLogic/Source/RigLogicLib/Private/riglogic/utils/Extd.h
@@ -108,15 +108,20 @@ inline impl::LUTFilter<TLookUpTable, impl::LUTStrategy::ByPosition> byPosition(c
     return impl::LUTFilter<TLookUpTable, impl::LUTStrategy::ByPosition>{lookUpTable};
 }
 
+template<class TContainer, class TComparator>
+inline typename TContainer::value_type maxOf(const TContainer& container, TComparator&& comparator) {
+    using ValueType = typename TContainer::value_type;
+    const auto it = std::max_element(container.begin(), container.end(), comparator);
+    return (it == container.end() ? ValueType{} : *it);
+}
+
 template<class TContainer>
 inline typename TContainer::value_type maxOf(const TContainer& container) {
-    assert(!container.empty());
     using ValueType = typename TContainer::value_type;
     const auto compare = [](const ValueType& lhs, const ValueType& rhs) {
-            return lhs.second < rhs.second;
+            return lhs < rhs;
         };
-    const auto it = std::max_element(container.begin(), container.end(), compare);
-    return (it == container.end() ? ValueType{} : *it);
+    return maxOf(container, compare);
 }
 
 template<class TSource, class TDestination>
diff --git a/Engine/Plugins/Animation/RigLogic/Source/RigLogicLib/Private/riglogic/utils/Macros.h b/Engine/Plugins/Animation/RigLogic/Source/RigLogicLib/Private/riglogic/utils/Macros.h
--- a/Engine/Plugins/Animation/RigLogic/Source/RigLogicLib/Private/riglogic/utils/Macros.h
+++ b/Engine/Plugins/Animation/RigLogic/Source/RigLogicLib/Private/riglogic/utils/Macros.h
@@ -2,8 +2,10 @@
 
 #pragma once
 
-#if defined(_MSC_VER)
-    #define FORCE_INLINE __forceinline
-#else
-    #define FORCE_INLINE inline __attribute__((always_inline))
+#if !defined(FORCE_INLINE)
+    #if defined(_MSC_VER)
+        #define FORCE_INLINE __forceinline
+    #else
+        #define FORCE_INLINE inline __attribute__((always_inline))
+    #endif
 #endif
diff --git a/Engine/Plugins/Animation/RigLogic/Source/RigLogicLib/Private/status/Provider.cpp b/Engine/Plugins/Animation/RigLogic/Source/RigLogicLib/Private/status/Provider.cpp
--- a/Engine/Plugins/Animation/RigLogic/Source/RigLogicLib/Private/status/Provider.cpp
+++ b/Engine/Plugins/Animation/RigLogic/Source/RigLogicLib/Private/status/Provider.cpp
@@ -31,7 +31,17 @@ bool StatusProvider::isOk() {
 }
 
 void StatusProvider::set(StatusCode status) {
+    status.message = execHook(status, 0ul, status.message);
+    execSet(status);
+}
+
+void StatusProvider::execSet(StatusCode status) {
     StatusStorage::set(status);
 }
 
+const char* StatusProvider::execHook(StatusCode status, std::size_t index, const char* data) {
+    auto hook = StatusStorage::getHook();
+    return (hook == nullptr ? data : hook(status, index, data));
+}
+
 }  // namespace sc
diff --git a/Engine/Plugins/Animation/RigLogic/Source/RigLogicLib/Private/status/Status.cpp b/Engine/Plugins/Animation/RigLogic/Source/RigLogicLib/Private/status/Status.cpp
--- a/Engine/Plugins/Animation/RigLogic/Source/RigLogicLib/Private/status/Status.cpp
+++ b/Engine/Plugins/Animation/RigLogic/Source/RigLogicLib/Private/status/Status.cpp
@@ -14,4 +14,12 @@ bool Status::isOk() {
     return StatusStorage::isOk();
 }
 
+HookFunction Status::getHook() {
+    return StatusStorage::getHook();
+}
+
+void Status::setHook(HookFunction hook) {
+    StatusStorage::setHook(hook);
+}
+
 }  // namespace sc
diff --git a/Engine/Plugins/Animation/RigLogic/Source/RigLogicLib/Private/status/Storage.cpp b/Engine/Plugins/Animation/RigLogic/Source/RigLogicLib/Private/status/Storage.cpp
--- a/Engine/Plugins/Animation/RigLogic/Source/RigLogicLib/Private/status/Storage.cpp
+++ b/Engine/Plugins/Animation/RigLogic/Source/RigLogicLib/Private/status/Storage.cpp
@@ -17,6 +17,7 @@ namespace {
 struct StatusCodeStorage {
     int code;
     char message[512];
+    HookFunction hook;
 };
 
 void strcopy(char* destination, const char* source, std::size_t bufferSize) {
@@ -29,7 +30,7 @@ void strcopy(char* destination, const char* source, std::size_t bufferSize) {
 
 }  // namespace
 
-thread_local static StatusCodeStorage currentStatus{0, "Ok"};
+thread_local static StatusCodeStorage currentStatus{0, "Ok", nullptr};
 
 void StatusStorage::set(StatusCode status) {
     // The Release build will eliminate this call, as it's really just a sanity check
@@ -56,4 +57,12 @@ constexpr std::size_t StatusStorage::bufferSize() {
     return sizeof(currentStatus.message);
 }
 
+HookFunction StatusStorage::getHook() {
+    return currentStatus.hook;
+}
+
+void StatusStorage::setHook(HookFunction hook) {
+    currentStatus.hook = hook;
+}
+
 }  // namespace sc
diff --git a/Engine/Plugins/Animation/RigLogic/Source/RigLogicLib/Private/status/Storage.h b/Engine/Plugins/Animation/RigLogic/Source/RigLogicLib/Private/status/Storage.h
--- a/Engine/Plugins/Animation/RigLogic/Source/RigLogicLib/Private/status/Storage.h
+++ b/Engine/Plugins/Animation/RigLogic/Source/RigLogicLib/Private/status/Storage.h
@@ -2,6 +2,7 @@
 
 #pragma once
 
+#include "status/Status.h"
 #include "status/StatusCode.h"
 
 #include <cstddef>
@@ -15,6 +16,8 @@ class StatusStorage {
         static StatusCode get();
         static bool isOk();
         static constexpr std::size_t bufferSize();
+        static HookFunction getHook();
+        static void setHook(HookFunction hook);
 
 };
 
diff --git a/Engine/Plugins/Animation/RigLogic/Source/RigLogicLib/Private/trio/streams/MemoryStreamImpl.cpp b/Engine/Plugins/Animation/RigLogic/Source/RigLogicLib/Private/trio/streams/MemoryStreamImpl.cpp
--- a/Engine/Plugins/Animation/RigLogic/Source/RigLogicLib/Private/trio/streams/MemoryStreamImpl.cpp
+++ b/Engine/Plugins/Animation/RigLogic/Source/RigLogicLib/Private/trio/streams/MemoryStreamImpl.cpp
@@ -133,7 +133,7 @@ std::size_t MemoryStreamImpl::read(Writable* destination, std::size_t size) {
 
     const std::size_t available = data.size() - position;
     const std::size_t bytesToRead = std::min(size, available);
-    const std::size_t bytesCopied = destination->write(&data[position], bytesToRead);
+    const std::size_t bytesCopied = (bytesToRead > 0ul ? destination->write(&data[position], bytesToRead) : 0ul);
     position += bytesCopied;
     return bytesCopied;
 }
diff --git a/Engine/Plugins/Animation/RigLogic/Source/RigLogicLib/Public/arrayview/version/Version.h b/Engine/Plugins/Animation/RigLogic/Source/RigLogicLib/Public/arrayview/version/Version.h
new file mode 100644
--- /dev/null
+++ b/Engine/Plugins/Animation/RigLogic/Source/RigLogicLib/Public/arrayview/version/Version.h
@@ -0,0 +1,8 @@
+// Copyright Epic Games, Inc. All Rights Reserved.
+
+#pragma once
+
+#define AV_MAJOR_VERSION 1
+#define AV_MINOR_VERSION 1
+#define AV_PATCH_VERSION 3
+#define AV_VERSION_STRING "1.1.3"
diff --git a/Engine/Plugins/Animation/RigLogic/Source/RigLogicLib/Public/dna/BinaryStreamReader.h b/Engine/Plugins/Animation/RigLogic/Source/RigLogicLib/Public/dna/BinaryStreamReader.h
--- a/Engine/Plugins/Animation/RigLogic/Source/RigLogicLib/Public/dna/BinaryStreamReader.h
+++ b/Engine/Plugins/Animation/RigLogic/Source/RigLogicLib/Public/dna/BinaryStreamReader.h
@@ -21,6 +21,9 @@ class DNAAPI BinaryStreamReader : public StreamReader {
                 The Definition data layer depends on and thus implicitly loads the Descriptor layer.
                 The Behavior data layer depends on and thus implicitly loads the Definition layer.
                 The Geometry data layer depends on and thus also implicitly loads the Definition layer.
+                The MachineLearnedBehavior data layer depends on and thus also implicitly loads the Definition layer.
+            @param policy
+                Specify whether unknown layers are to be preserved or just ignored.
             @param maxLOD
                 The maximum level of details to be loaded.
             @note
@@ -38,6 +41,7 @@ class DNAAPI BinaryStreamReader : public StreamReader {
         */
         static BinaryStreamReader* create(BoundedIOStream* stream,
                                           DataLayer layer = DataLayer::All,
+                                          UnknownLayerPolicy policy = UnknownLayerPolicy::Preserve,
                                           std::uint16_t maxLOD = 0u,
                                           MemoryResource* memRes = nullptr);
         /**
@@ -50,6 +54,9 @@ class DNAAPI BinaryStreamReader : public StreamReader {
                 The Definition data layer depends on and thus implicitly loads the Descriptor layer.
                 The Behavior data layer depends on and thus implicitly loads the Definition layer.
                 The Geometry data layer depends on and thus also implicitly loads the Definition layer.
+                The MachineLearnedBehavior data layer depends on and thus also implicitly loads the Definition layer.
+            @param policy
+                Specify whether unknown layers are to be preserved or just ignored.
             @param maxLOD
                 The maximum level of details to be loaded.
             @param minLOD
@@ -69,6 +76,7 @@ class DNAAPI BinaryStreamReader : public StreamReader {
         */
         static BinaryStreamReader* create(BoundedIOStream* stream,
                                           DataLayer layer,
+                                          UnknownLayerPolicy policy,
                                           std::uint16_t maxLOD,
                                           std::uint16_t minLOD,
                                           MemoryResource* memRes = nullptr);
@@ -82,6 +90,9 @@ class DNAAPI BinaryStreamReader : public StreamReader {
                 The Definition data layer depends on and thus implicitly loads the Descriptor layer.
                 The Behavior data layer depends on and thus implicitly loads the Definition layer.
                 The Geometry data layer depends on and thus also implicitly loads the Definition layer.
+                The MachineLearnedBehavior data layer depends on and thus also implicitly loads the Definition layer.
+            @param policy
+                Specify whether unknown layers are to be preserved or just ignored.
             @param lods
                 An array specifying which exact lods to load.
             @warning
@@ -102,6 +113,7 @@ class DNAAPI BinaryStreamReader : public StreamReader {
         */
         static BinaryStreamReader* create(BoundedIOStream* stream,
                                           DataLayer layer,
+                                          UnknownLayerPolicy policy,
                                           std::uint16_t* lods,
                                           std::uint16_t lodCount,
                                           MemoryResource* memRes = nullptr);
diff --git a/Engine/Plugins/Animation/RigLogic/Source/RigLogicLib/Public/dna/BinaryStreamWriter.h b/Engine/Plugins/Animation/RigLogic/Source/RigLogicLib/Public/dna/BinaryStreamWriter.h
--- a/Engine/Plugins/Animation/RigLogic/Source/RigLogicLib/Public/dna/BinaryStreamWriter.h
+++ b/Engine/Plugins/Animation/RigLogic/Source/RigLogicLib/Public/dna/BinaryStreamWriter.h
@@ -8,6 +8,9 @@
 
 namespace dna {
 
+class BinaryStreamReader;
+class JSONStreamReader;
+
 class DNAAPI BinaryStreamWriter : public StreamWriter {
     public:
         /**
@@ -32,6 +35,16 @@ class DNAAPI BinaryStreamWriter : public StreamWriter {
         static void destroy(BinaryStreamWriter* instance);
 
         ~BinaryStreamWriter() override;
+
+        using StreamWriter::setFrom;
+        virtual void setFrom(const BinaryStreamReader* source,
+                             DataLayer layer = DataLayer::All,
+                             UnknownLayerPolicy policy = UnknownLayerPolicy::Preserve,
+                             MemoryResource* memRes = nullptr) = 0;
+        virtual void setFrom(const JSONStreamReader* source,
+                             DataLayer layer = DataLayer::All,
+                             UnknownLayerPolicy policy = UnknownLayerPolicy::Preserve,
+                             MemoryResource* memRes = nullptr) = 0;
 };
 
 }  // namespace dna
diff --git a/Engine/Plugins/Animation/RigLogic/Source/RigLogicLib/Public/dna/DataLayer.h b/Engine/Plugins/Animation/RigLogic/Source/RigLogicLib/Public/dna/DataLayer.h
--- a/Engine/Plugins/Animation/RigLogic/Source/RigLogicLib/Public/dna/DataLayer.h
+++ b/Engine/Plugins/Animation/RigLogic/Source/RigLogicLib/Public/dna/DataLayer.h
@@ -2,16 +2,27 @@
 
 #pragma once
 
+#include <cstdint>
+
 namespace dna {
 
-enum class DataLayer {
-    Descriptor,
-    Definition,  // Includes Descriptor
-    Behavior,  // Includes Descriptor and Definition
-    Geometry,  // Includes Descriptor and Definition
-    GeometryWithoutBlendShapes,  // Includes Descriptor and Definition
-    AllWithoutBlendShapes,  // Includes everything except blend shapes from Geometry
-    All
+enum class DataLayer : std::uint32_t {
+    Descriptor = 1,
+    Definition = 2 | Descriptor,  // Implicitly loads Descriptor
+    Behavior = 4 | Definition,  // Implicitly loads Descriptor and Definition
+    Geometry = 8 | Definition,  // Implicitly loads Descriptor and Definition
+    GeometryWithoutBlendShapes = 16 | Definition,  // Implicitly loads Descriptor and Definition
+    MachineLearnedBehavior = 32 | Definition,  // Implicitly loads Definition
+    All = Behavior | Geometry | MachineLearnedBehavior
+};
+
+inline DataLayer operator|(DataLayer lhs, DataLayer rhs) {
+    return static_cast<DataLayer>(static_cast<std::uint32_t>(lhs) | static_cast<std::uint32_t>(rhs));
+}
+
+enum class UnknownLayerPolicy {
+    Preserve,
+    Ignore
 };
 
 }  // namespace dna
diff --git a/Engine/Plugins/Animation/RigLogic/Source/RigLogicLib/Public/dna/JSONStreamWriter.h b/Engine/Plugins/Animation/RigLogic/Source/RigLogicLib/Public/dna/JSONStreamWriter.h
--- a/Engine/Plugins/Animation/RigLogic/Source/RigLogicLib/Public/dna/JSONStreamWriter.h
+++ b/Engine/Plugins/Animation/RigLogic/Source/RigLogicLib/Public/dna/JSONStreamWriter.h
@@ -8,6 +8,9 @@
 
 namespace dna {
 
+class BinaryStreamReader;
+class JSONStreamReader;
+
 class DNAAPI JSONStreamWriter : public StreamWriter {
     public:
         /**
@@ -35,6 +38,16 @@ class DNAAPI JSONStreamWriter : public StreamWriter {
         static void destroy(JSONStreamWriter* instance);
 
         ~JSONStreamWriter() override;
+
+        using StreamWriter::setFrom;
+        virtual void setFrom(const BinaryStreamReader* source,
+                             DataLayer layer = DataLayer::All,
+                             UnknownLayerPolicy policy = UnknownLayerPolicy::Preserve,
+                             MemoryResource* memRes = nullptr) = 0;
+        virtual void setFrom(const JSONStreamReader* source,
+                             DataLayer layer = DataLayer::All,
+                             UnknownLayerPolicy policy = UnknownLayerPolicy::Preserve,
+                             MemoryResource* memRes = nullptr) = 0;
 };
 
 }  // namespace dna
diff --git a/Engine/Plugins/Animation/RigLogic/Source/RigLogicLib/Public/dna/Reader.h b/Engine/Plugins/Animation/RigLogic/Source/RigLogicLib/Public/dna/Reader.h
--- a/Engine/Plugins/Animation/RigLogic/Source/RigLogicLib/Public/dna/Reader.h
+++ b/Engine/Plugins/Animation/RigLogic/Source/RigLogicLib/Public/dna/Reader.h
@@ -6,6 +6,7 @@
 #include "dna/DataLayer.h"
 #include "dna/layers/BehaviorReader.h"
 #include "dna/layers/GeometryReader.h"
+#include "dna/layers/MachineLearnedBehaviorReader.h"
 
 namespace dna {
 
@@ -18,13 +19,13 @@ namespace dna {
         selectively loaded, it might be convenient to slice-off interfaces which layers were
         not loaded.
 */
-class DNAAPI Reader : public BehaviorReader, public GeometryReader {
+class DNAAPI Reader : public BehaviorReader, public GeometryReader, public MachineLearnedBehaviorReader {
     public:
         ~Reader() override;
         /**
             @brief Unload all data of the specified layer and all layers dependent on it.
             @param layer
-                Layer which data should be unloaded.
+                Layer whose data should be unloaded.
         */
         virtual void unload(DataLayer layer) = 0;
 };
diff --git a/Engine/Plugins/Animation/RigLogic/Source/RigLogicLib/Public/dna/Writer.h b/Engine/Plugins/Animation/RigLogic/Source/RigLogicLib/Public/dna/Writer.h
--- a/Engine/Plugins/Animation/RigLogic/Source/RigLogicLib/Public/dna/Writer.h
+++ b/Engine/Plugins/Animation/RigLogic/Source/RigLogicLib/Public/dna/Writer.h
@@ -6,6 +6,7 @@
 #include "dna/Defs.h"
 #include "dna/layers/BehaviorWriter.h"
 #include "dna/layers/GeometryWriter.h"
+#include "dna/layers/MachineLearnedBehaviorWriter.h"
 #include "dna/types/Aliases.h"
 
 namespace dna {
@@ -20,7 +21,7 @@ class Reader;
         structure of the Reader hierarchy, as it's not possible to selectively write only
         specific layers.
 */
-class DNAAPI Writer : public BehaviorWriter, public GeometryWriter {
+class DNAAPI Writer : public BehaviorWriter, public GeometryWriter, public MachineLearnedBehaviorWriter {
     public:
         ~Writer() override;
         /**
@@ -35,10 +36,15 @@ class DNAAPI Writer : public BehaviorWriter, public GeometryWriter {
                 The source DNA Reader from which the data needs to be copied.
             @param layer
                 Limit which layers should be taken over from the given source reader.
+            @param policy
+                Specify whether unknown layers are to be preserved or just ignored.
             @param memRes
                 Optional memory resource to use for temporary allocations during copying.
         */
-        void setFrom(const Reader* source, DataLayer layer = DataLayer::All, MemoryResource* memRes = nullptr);
+        virtual void setFrom(const Reader* source,
+                             DataLayer layer = DataLayer::All,
+                             UnknownLayerPolicy policy = UnknownLayerPolicy::Preserve,
+                             MemoryResource* memRes = nullptr);
 };
 
 }  // namespace dna
diff --git a/Engine/Plugins/Animation/RigLogic/Source/RigLogicLib/Public/dna/layers/BehaviorReader.h b/Engine/Plugins/Animation/RigLogic/Source/RigLogicLib/Public/dna/layers/BehaviorReader.h
--- a/Engine/Plugins/Animation/RigLogic/Source/RigLogicLib/Public/dna/layers/BehaviorReader.h
+++ b/Engine/Plugins/Animation/RigLogic/Source/RigLogicLib/Public/dna/layers/BehaviorReader.h
@@ -24,41 +24,35 @@ class DNAAPI BehaviorReader : public virtual DefinitionReader {
         /**
             @brief Input indices used for mapping gui to raw controls.
             @return View over the array of input indices.
-            @see Controls
         */
         virtual ConstArrayView<std::uint16_t> getGUIToRawInputIndices() const = 0;
         /**
             @brief Output indices used for mapping gui to raw controls.
             @return View over the array of output indices.
-            @see Controls
         */
         virtual ConstArrayView<std::uint16_t> getGUIToRawOutputIndices() const = 0;
         /**
             @brief Filter values(lower-bounds) used to decide whether a particular
                 entry should be evaluated or not during gui to raw control mapping.
             @return View over the array of filter values.
-            @see Controls
         */
         virtual ConstArrayView<float> getGUIToRawFromValues() const = 0;
         /**
             @brief Filter values(upper-bounds) used to decide whether a particular
                 entry should be evaluated or not during gui to raw control mapping.
             @return View over the array of filter values.
-            @see Controls
         */
         virtual ConstArrayView<float> getGUIToRawToValues() const = 0;
         /**
             @brief Computational values(slope/gradient) used for calculating the
                 output value during gui to raw control mapping.
             @return View over the array of computational values.
-            @see Controls
         */
         virtual ConstArrayView<float> getGUIToRawSlopeValues() const = 0;
         /**
             @brief Computational values(vertical intercept) used for calculating the
                 output value during gui to raw control mapping.
             @return View over the array of computational values.
-            @see Controls
         */
         virtual ConstArrayView<float> getGUIToRawCutValues() const = 0;
         /**
@@ -68,19 +62,16 @@ class DNAAPI BehaviorReader : public virtual DefinitionReader {
         /**
             @brief PSD(input) indices.
             @return View over the array of PSD indices.
-            @see Controls
         */
         virtual ConstArrayView<std::uint16_t> getPSDRowIndices() const = 0;
         /**
             @brief Control(input) indices.
             @return View over the array of control indices.
-            @see Controls
         */
         virtual ConstArrayView<std::uint16_t> getPSDColumnIndices() const = 0;
         /**
             @brief Weights associated with each PSD row and column pair.
             @return View over the array of weights.
-            @see Controls
         */
         virtual ConstArrayView<float> getPSDValues() const = 0;
         /**
diff --git a/Engine/Plugins/Animation/RigLogic/Source/RigLogicLib/Public/dna/layers/BehaviorWriter.h b/Engine/Plugins/Animation/RigLogic/Source/RigLogicLib/Public/dna/layers/BehaviorWriter.h
--- a/Engine/Plugins/Animation/RigLogic/Source/RigLogicLib/Public/dna/layers/BehaviorWriter.h
+++ b/Engine/Plugins/Animation/RigLogic/Source/RigLogicLib/Public/dna/layers/BehaviorWriter.h
@@ -27,7 +27,6 @@ class DNAAPI BehaviorWriter : public virtual DefinitionWriter {
                 The source address from which the input indices are to be copied.
             @param count
                 The number of input indices to copy.
-            @see Controls
         */
         virtual void setGUIToRawInputIndices(const std::uint16_t* inputIndices, std::uint16_t count) = 0;
         /**
@@ -36,7 +35,6 @@ class DNAAPI BehaviorWriter : public virtual DefinitionWriter {
                 The source address from which the output indices are to be copied.
             @param count
                 The number of output indices to copy.
-            @see Controls
         */
         virtual void setGUIToRawOutputIndices(const std::uint16_t* outputIndices, std::uint16_t count) = 0;
         /**
@@ -46,7 +44,6 @@ class DNAAPI BehaviorWriter : public virtual DefinitionWriter {
                 The source address from which the filter values are to be copied.
             @param count
                 The number of filter values to copy.
-            @see Controls
         */
         virtual void setGUIToRawFromValues(const float* fromValues, std::uint16_t count) = 0;
         /**
@@ -56,7 +53,6 @@ class DNAAPI BehaviorWriter : public virtual DefinitionWriter {
                 The source address from which the filter values are to be copied.
             @param count
                 The number of filter values to copy.
-            @see Controls
         */
         virtual void setGUIToRawToValues(const float* toValues, std::uint16_t count) = 0;
         /**
@@ -66,7 +62,6 @@ class DNAAPI BehaviorWriter : public virtual DefinitionWriter {
                 The source address from which the computational values are to be copied.
             @param count
                 The number of computational values to copy.
-            @see Controls
         */
         virtual void setGUIToRawSlopeValues(const float* slopeValues, std::uint16_t count) = 0;
         /**
@@ -76,7 +71,6 @@ class DNAAPI BehaviorWriter : public virtual DefinitionWriter {
                 The source address from which the computational values are to be copied.
             @param count
                 The number of computational values to copy.
-            @see Controls
         */
         virtual void setGUIToRawCutValues(const float* cutValues, std::uint16_t count) = 0;
         /**
@@ -89,7 +83,6 @@ class DNAAPI BehaviorWriter : public virtual DefinitionWriter {
                 The source address from which the PSD indices are to be copied.
             @param count
                 The number of PSD indices to copy.
-            @see Controls
         */
         virtual void setPSDRowIndices(const std::uint16_t* rowIndices, std::uint16_t count) = 0;
         /**
@@ -98,7 +91,6 @@ class DNAAPI BehaviorWriter : public virtual DefinitionWriter {
                 The source address from which the control indices are to be copied.
             @param count
                 The number of control indices to copy.
-            @see Controls
         */
         virtual void setPSDColumnIndices(const std::uint16_t* columnIndices, std::uint16_t count) = 0;
         /**
@@ -107,7 +99,6 @@ class DNAAPI BehaviorWriter : public virtual DefinitionWriter {
                 The source address from which the weight values are to be copied.
             @param count
                 The number of weight values to copy.
-            @see Controls
         */
         virtual void setPSDValues(const float* weights, std::uint16_t count) = 0;
         /**
diff --git a/Engine/Plugins/Animation/RigLogic/Source/RigLogicLib/Public/dna/layers/DefinitionReader.h b/Engine/Plugins/Animation/RigLogic/Source/RigLogicLib/Public/dna/layers/DefinitionReader.h
--- a/Engine/Plugins/Animation/RigLogic/Source/RigLogicLib/Public/dna/layers/DefinitionReader.h
+++ b/Engine/Plugins/Animation/RigLogic/Source/RigLogicLib/Public/dna/layers/DefinitionReader.h
@@ -30,6 +30,9 @@ class DNAAPI DefinitionReader : public DescriptorReader {
         virtual ~DefinitionReader();
 
     public:
+        /**
+            @brief Number of GUI controls.
+        */
         virtual std::uint16_t getGUIControlCount() const = 0;
         /**
             @brief Name of the requested GUI control.
@@ -38,10 +41,11 @@ class DNAAPI DefinitionReader : public DescriptorReader {
             @warning
                 The index must be less than the value returned by getGUIControlCount.
             @return View over the GUI control name string.
-            @see Controls
         */
         virtual StringView getGUIControlName(std::uint16_t index) const = 0;
-
+        /**
+            @brief Number of raw controls.
+        */
         virtual std::uint16_t getRawControlCount() const = 0;
         /**
             @brief Name of the requested raw control.
@@ -50,10 +54,11 @@ class DNAAPI DefinitionReader : public DescriptorReader {
             @warning
                 The index must be less than the value returned by getRawControlCount.
             @return View over the control name string.
-            @see Controls
         */
         virtual StringView getRawControlName(std::uint16_t index) const = 0;
-
+        /**
+            @brief Number of raw joints.
+        */
         virtual std::uint16_t getJointCount() const = 0;
         /**
             @brief Name of the requested joint.
@@ -106,14 +111,16 @@ class DNAAPI DefinitionReader : public DescriptorReader {
                 The joint index which parent is being requested.
         */
         virtual std::uint16_t getJointParentIndex(std::uint16_t index) const = 0;
-
+        /**
+            @brief Number of blend shape channels.
+        */
         virtual std::uint16_t getBlendShapeChannelCount() const = 0;
         /**
             @brief Name of the requested blend shape channel.
             @param index
                 A name's position in the zero-indexed array of blend shape channel names.
             @warning
-                The index must be less than the value returned by BlendShapeChannelExtentReader::getBlendShapeChannelCount.
+                The index must be less than the value returned by getBlendShapeChannelCount.
             @return View over the blend shape channel name string.
             @see BlendShapes
         */
@@ -129,15 +136,17 @@ class DNAAPI DefinitionReader : public DescriptorReader {
             @param lod
                 The level of detail which blend shape channels are being requested.
             @warning
-                The lod index must be less than the value returned by LODExtentReader::getLODCount.
+                The lod index must be less than the value returned by DescriptorReader::getLODCount.
             @return View over the blend shape channel indices.
             @warning
                 These LOD indices are not interchangeable with the LOD values from BehaviorReader::getBlendShapeChannelLODs.
-            @see LODExtentReader::getLODCount
+            @see DescriptorReader::getLODCount
             @see getBlendShapeChannelName
         */
         virtual ConstArrayView<std::uint16_t> getBlendShapeChannelIndicesForLOD(std::uint16_t lod) const = 0;
-
+        /**
+            @brief Number of animated maps.
+        */
         virtual std::uint16_t getAnimatedMapCount() const = 0;
         /**
             @brief Name of the requested animated map.
@@ -166,7 +175,9 @@ class DNAAPI DefinitionReader : public DescriptorReader {
             @see getAnimatedMapName
         */
         virtual ConstArrayView<std::uint16_t> getAnimatedMapIndicesForLOD(std::uint16_t lod) const = 0;
-
+        /**
+            @brief Number of meshes.
+        */
         virtual std::uint16_t getMeshCount() const = 0;
         /**
             @brief Name of the requested mesh.
diff --git a/Engine/Plugins/Animation/RigLogic/Source/RigLogicLib/Public/dna/layers/DefinitionWriter.h b/Engine/Plugins/Animation/RigLogic/Source/RigLogicLib/Public/dna/layers/DefinitionWriter.h
--- a/Engine/Plugins/Animation/RigLogic/Source/RigLogicLib/Public/dna/layers/DefinitionWriter.h
+++ b/Engine/Plugins/Animation/RigLogic/Source/RigLogicLib/Public/dna/layers/DefinitionWriter.h
@@ -37,7 +37,6 @@ class DNAAPI DefinitionWriter : public DescriptorWriter {
                 A null-terminated string.
             @note
                 The passed in name is copied, which will involve an additional allocation.
-            @see Controls
         */
         virtual void setGUIControlName(std::uint16_t index, const char* name) = 0;
         /**
@@ -55,7 +54,6 @@ class DNAAPI DefinitionWriter : public DescriptorWriter {
                 A null-terminated string.
             @note
                 The passed in name is copied, which will involve an additional allocation.
-            @see Controls
         */
         virtual void setRawControlName(std::uint16_t index, const char* name) = 0;
         /**
diff --git a/Engine/Plugins/Animation/RigLogic/Source/RigLogicLib/Public/dna/layers/DescriptorReader.h b/Engine/Plugins/Animation/RigLogic/Source/RigLogicLib/Public/dna/layers/DescriptorReader.h
--- a/Engine/Plugins/Animation/RigLogic/Source/RigLogicLib/Public/dna/layers/DescriptorReader.h
+++ b/Engine/Plugins/Animation/RigLogic/Source/RigLogicLib/Public/dna/layers/DescriptorReader.h
@@ -4,6 +4,7 @@
 
 #include "dna/Defs.h"
 #include "dna/layers/Descriptor.h"
+#include "dna/layers/HeaderReader.h"
 #include "dna/types/Aliases.h"
 
 #include <cstdint>
@@ -14,11 +15,9 @@ namespace dna {
     @brief Read-only accessors for various metadata about the character and the rig.
     @warning
         Implementors should inherit from Reader itself and not this class.
-    @note
-        The inherited LODExtentReader provides the interface to access the LOD count.
     @see Reader
 */
-class DNAAPI DescriptorReader {
+class DNAAPI DescriptorReader : public HeaderReader {
     protected:
         virtual ~DescriptorReader();
 
diff --git a/Engine/Plugins/Animation/RigLogic/Source/RigLogicLib/Public/dna/layers/DescriptorWriter.h b/Engine/Plugins/Animation/RigLogic/Source/RigLogicLib/Public/dna/layers/DescriptorWriter.h
--- a/Engine/Plugins/Animation/RigLogic/Source/RigLogicLib/Public/dna/layers/DescriptorWriter.h
+++ b/Engine/Plugins/Animation/RigLogic/Source/RigLogicLib/Public/dna/layers/DescriptorWriter.h
@@ -4,6 +4,7 @@
 
 #include "dna/Defs.h"
 #include "dna/layers/Descriptor.h"
+#include "dna/layers/HeaderWriter.h"
 #include "dna/types/Aliases.h"
 
 #include <cstdint>
@@ -16,7 +17,7 @@ namespace dna {
         Implementors should inherit from Writer itself and not this class.
     @see Writer
 */
-class DNAAPI DescriptorWriter {
+class DNAAPI DescriptorWriter : public HeaderWriter {
     protected:
         virtual ~DescriptorWriter();
 
diff --git a/Engine/Plugins/Animation/RigLogic/Source/RigLogicLib/Public/dna/types/Aliases.h b/Engine/Plugins/Animation/RigLogic/Source/RigLogicLib/Public/dna/types/Aliases.h
--- a/Engine/Plugins/Animation/RigLogic/Source/RigLogicLib/Public/dna/types/Aliases.h
+++ b/Engine/Plugins/Animation/RigLogic/Source/RigLogicLib/Public/dna/types/Aliases.h
@@ -2,9 +2,8 @@
 
 #pragma once
 
-#include "dna/types/ArrayView.h"
-#include "dna/types/StringView.h"
-
+#include <arrayview/ArrayView.h>
+#include <arrayview/StringView.h>
 #include <pma/MemoryResource.h>
 #include <pma/ScopedPtr.h>
 #include <status/Status.h>
@@ -16,18 +15,13 @@
 
 namespace dna {
 
-template<typename T>
-using ArrayView = trust::ArrayView<T>;
-
-template<typename T>
-using ConstArrayView = trust::ConstArrayView<T>;
-
+using sc::Status;
 using trio::BoundedIOStream;
 using trio::FileStream;
 using trio::MemoryMappedFileStream;
 using trio::MemoryStream;
-using sc::Status;
 
+using namespace av;
 using namespace pma;
 
 }  // namespace dna
diff --git a/Engine/Plugins/Animation/RigLogic/Source/RigLogicLib/Public/dna/version/Version.h b/Engine/Plugins/Animation/RigLogic/Source/RigLogicLib/Public/dna/version/Version.h
--- a/Engine/Plugins/Animation/RigLogic/Source/RigLogicLib/Public/dna/version/Version.h
+++ b/Engine/Plugins/Animation/RigLogic/Source/RigLogicLib/Public/dna/version/Version.h
@@ -2,7 +2,7 @@
 
 #pragma once
 
-#define DNA_MAJOR_VERSION 7
+#define DNA_MAJOR_VERSION 9
 #define DNA_MINOR_VERSION 1
-#define DNA_PATCH_VERSION 0
-#define DNA_VERSION_STRING "7.1.0"
+#define DNA_PATCH_VERSION 1
+#define DNA_VERSION_STRING "9.1.1"
diff --git a/Engine/Plugins/Animation/RigLogic/Source/RigLogicLib/Public/pma/version/Version.h b/Engine/Plugins/Animation/RigLogic/Source/RigLogicLib/Public/pma/version/Version.h
--- a/Engine/Plugins/Animation/RigLogic/Source/RigLogicLib/Public/pma/version/Version.h
+++ b/Engine/Plugins/Animation/RigLogic/Source/RigLogicLib/Public/pma/version/Version.h
@@ -4,5 +4,5 @@
 
 #define PMA_MAJOR_VERSION 1
 #define PMA_MINOR_VERSION 3
-#define PMA_PATCH_VERSION 3
-#define PMA_VERSION_STRING "1.3.3"
+#define PMA_PATCH_VERSION 6
+#define PMA_VERSION_STRING "1.3.6"
diff --git a/Engine/Plugins/Animation/RigLogic/Source/RigLogicLib/Public/riglogic/riglogic/Configuration.h b/Engine/Plugins/Animation/RigLogic/Source/RigLogicLib/Public/riglogic/riglogic/Configuration.h
new file mode 100644
--- /dev/null
+++ b/Engine/Plugins/Animation/RigLogic/Source/RigLogicLib/Public/riglogic/riglogic/Configuration.h
@@ -0,0 +1,25 @@
+// Copyright Epic Games, Inc. All Rights Reserved.
+
+#pragma once
+
+namespace rl4 {
+
+/**
+    @brief Implementation type for RigLogic calculations.
+*/
+enum class CalculationType {
+    Scalar,  ///< scalar CPU algorithm
+    SSE,  ///< vectorized (SSE) CPU algorithm
+    AVX  ///< vectorized (AVX) CPU algorithm (RigLogic must be built with AVX support,
+         ///< otherwise it falls back to using the Scalar version)
+};
+
+struct Configuration {
+    CalculationType calculationType = CalculationType::SSE;
+    bool loadJoints = true;
+    bool loadBlendShapes = true;
+    bool loadAnimatedMaps = true;
+    bool loadMachineLearnedBehavior = true;
+};
+
+}  // namespace rl4
diff --git a/Engine/Plugins/Animation/RigLogic/Source/RigLogicLib/Public/riglogic/riglogic/RigInstance.h b/Engine/Plugins/Animation/RigLogic/Source/RigLogicLib/Public/riglogic/riglogic/RigInstance.h
--- a/Engine/Plugins/Animation/RigLogic/Source/RigLogicLib/Public/riglogic/riglogic/RigInstance.h
+++ b/Engine/Plugins/Animation/RigLogic/Source/RigLogicLib/Public/riglogic/riglogic/RigInstance.h
@@ -51,12 +51,29 @@ class RLAPI RigInstance {
         static void destroy(RigInstance* instance);
 
         virtual std::uint16_t getGUIControlCount() const = 0;
+        virtual float getGUIControl(std::uint16_t index) const = 0;
         virtual void setGUIControl(std::uint16_t index, float value) = 0;
+        virtual ConstArrayView<float> getGUIControlValues() const = 0;
         virtual void setGUIControlValues(const float* values) = 0;
 
         virtual std::uint16_t getRawControlCount() const = 0;
+        virtual float getRawControl(std::uint16_t index) const = 0;
         virtual void setRawControl(std::uint16_t index, float value) = 0;
+        virtual ConstArrayView<float> getRawControlValues() const = 0;
         virtual void setRawControlValues(const float* values) = 0;
+
+        virtual std::uint16_t getPSDControlCount() const = 0;
+        virtual float getPSDControl(std::uint16_t index) const = 0;
+        virtual ConstArrayView<float> getPSDControlValues() const = 0;
+
+        virtual std::uint16_t getMLControlCount() const = 0;
+        virtual float getMLControl(std::uint16_t index) const = 0;
+        virtual ConstArrayView<float> getMLControlValues() const = 0;
+
+        virtual std::uint16_t getNeuralNetworkCount() const = 0;
+        virtual float getNeuralNetworkMask(std::uint16_t neuralNetIndex) const = 0;
+        virtual void setNeuralNetworkMask(std::uint16_t neuralNetIndex, float value) = 0;
+
         /**
             @brief Calculated values for joint transformations.
             @note
diff --git a/Engine/Plugins/Animation/RigLogic/Source/RigLogicLib/Public/riglogic/riglogic/RigLogic.h b/Engine/Plugins/Animation/RigLogic/Source/RigLogicLib/Public/riglogic/riglogic/RigLogic.h
--- a/Engine/Plugins/Animation/RigLogic/Source/RigLogicLib/Public/riglogic/riglogic/RigLogic.h
+++ b/Engine/Plugins/Animation/RigLogic/Source/RigLogicLib/Public/riglogic/riglogic/RigLogic.h
@@ -3,7 +3,7 @@
 #pragma once
 
 #include "riglogic/Defs.h"
-#include "riglogic/riglogic/CalculationType.h"
+#include "riglogic/riglogic/Configuration.h"
 #include "riglogic/transformation/Transformation.h"
 #include "riglogic/types/Aliases.h"
 
@@ -12,7 +12,6 @@
 namespace rl4 {
 
 class RigInstance;
-class Stream;
 
 /**
     @brief RigLogic calculates rig output values based on input control values.
@@ -24,7 +23,7 @@ class Stream;
 */
 class RLAPI RigLogic {
     public:
-        using CalculationType = rl4::CalculationType;
+        using Configuration = rl4::Configuration;
 
     protected:
         virtual ~RigLogic();
@@ -34,8 +33,8 @@ class RLAPI RigLogic {
             @brief Factory method for the creation of RigLogic.
             @param reader
                 Source from which to copy and optimize DNA data, which is used for rig evaluation
-            @param calculationType
-                Determines which algorithm implementation is used for rig evaluation
+            @param config
+                Determines which algorithm implementation is used for rig evaluation and which submodules to load (affects memory allocations)
             @param memRes
                 A custom memory resource to be used for allocations.
             @note
@@ -44,9 +43,14 @@ class RLAPI RigLogic {
                 User is responsible for releasing the returned pointer by calling destroy.
             @see destroy
         */
-        static RigLogic* create(const dna::BehaviorReader* reader,
-                                CalculationType calculationType = CalculationType::SSE,
-                                MemoryResource* memRes = nullptr);
+        static RigLogic* create(const dna::Reader* reader, Configuration config = {}, MemoryResource* memRes = nullptr);
+        /**
+            @brief Method for freeing RigLogic.
+            @param instance
+                Instance of RigLogic to be freed.
+            @see create
+        */
+        static void destroy(RigLogic* instance);
         /**
             @brief Factory method for restoring an instance of RigLogic from a memory dump.
             @note
@@ -65,13 +69,6 @@ class RLAPI RigLogic {
             @see destroy
         */
         static RigLogic* restore(BoundedIOStream* source, MemoryResource* memRes = nullptr);
-        /**
-            @brief Method for freeing RigLogic.
-            @param instance
-                Instance of RigLogic to be freed.
-            @see create
-        */
-        static void destroy(RigLogic* instance);
         /**
             @brief Create a snapshot of an initialized RigLogic instance.
             @param destination
@@ -101,22 +98,47 @@ class RLAPI RigLogic {
             @see Transformation
         */
         virtual TransformationArrayView getNeutralJointValues() const = 0;
+        /**
+            @brief All joint output indices concatenated into a single chunk per each LOD.
+        */
+        virtual ConstArrayView<std::uint16_t> getJointVariableAttributeIndices(std::uint16_t lod) const = 0;
         /**
             @brief Number of joint groups present in the entire joint matrix.
             @see calculateJoints
         */
         virtual std::uint16_t getJointGroupCount() const = 0;
         /**
-            @brief All joint output indices concatenated into a single chunk per each LOD.
+            @brief Number of neural networks for driving machine learned behavior.
+            @see calculateMachineLearnedBehavior
         */
-        virtual ConstArrayView<std::uint16_t> getJointVariableAttributeIndices(std::uint16_t lod) const = 0;
+        virtual std::uint16_t getNeuralNetworkCount() const = 0;
+        /**
+            @brief Number of meshes.
+        */
+        virtual std::uint16_t getMeshCount() const = 0;
+        /**
+            @brief Number of regions within the specified mesh.
+        */
+        virtual std::uint16_t getMeshRegionCount(std::uint16_t meshIndex) const = 0;
+        /**
+            @brief Map a mesh and region index pair to the neural network indices that drive it.
+            @see RigInstance::getNeuralNetworkMask
+            @see RigInstance::setNeuralNetworkMask
+        */
+        virtual ConstArrayView<std::uint16_t> getNeuralNetworkIndices(std::uint16_t meshIndex,
+                                                                      std::uint16_t regionIndex) const = 0;
         /**
             @brief Maps GUI controls to raw controls.
             @note
-                This method must be called after either RigInstance::setGUIControlValues call
-                or RigInstance::::setGUIControl calls were issued.
+                This method may be called after either RigInstance::setGUIControlValues or RigInstance::setGUIControl was used.
         */
         virtual void mapGUIToRawControls(RigInstance* instance) const = 0;
+        /**
+            @brief Maps raw controls to GUI controls.
+            @note
+                This method may be called after either RigInstance::setRawControlValues or RigInstance::setRawControl was used.
+        */
+        virtual void mapRawToGUIControls(RigInstance* instance) const = 0;
         /**
             @brief Calculate only the input control values of the rig.
             @note
@@ -125,16 +147,38 @@ class RLAPI RigLogic {
             @note
                 This is considered as an advanced usage use case.
             @param instance
-                The rig instance which outputs are to be calculated.
+                The rig instance whose outputs are to be calculated.
             @see calculate
         */
         virtual void calculateControls(RigInstance* instance) const = 0;
+        /**
+            @brief Calculate controls driving the machine learned behavior of the rig.
+            @note
+                This is considered as an advanced usage use case.
+            @param instance
+                The rig instance whose controls for driving machine learned behavior are to be calculated.
+            @see calculate
+        */
+        virtual void calculateMachineLearnedBehaviorControls(RigInstance* instance) const = 0;
+        /**
+            @brief Calculate controls driving the machine learned behavior of the rig.
+            @note
+                This is considered as an advanced usage use case.
+            @param instance
+                The rig instance whose controls for driving machine learned behavior are to be calculated.
+            @param neuralNetIndex
+                The neural network whose outputs need to be calculated.
+            @warning
+                The index must be less than the value returned by getNeuralNetworkCount.
+            @see calculate
+        */
+        virtual void calculateMachineLearnedBehaviorControls(RigInstance* instance, std::uint16_t neuralNetIndex) const = 0;
         /**
             @brief Calculate only the joint outputs of the rig.
             @note
                 This is considered as an advanced usage use case.
             @param instance
-                The rig instance which outputs are to be calculated.
+                The rig instance whose outputs are to be calculated.
             @see calculate
         */
         virtual void calculateJoints(RigInstance* instance) const = 0;
@@ -146,7 +190,7 @@ class RLAPI RigLogic {
             @note
                 This is considered as an expert usage use case.
             @param instance
-                The rig instance which outputs are to be calculated.
+                The rig instance whose outputs are to be calculated.
             @param jointGroupIndex
                 A joint group's position in the zero-indexed array of joint groups.
             @warning
@@ -155,11 +199,11 @@ class RLAPI RigLogic {
         */
         virtual void calculateJoints(RigInstance* instance, std::uint16_t jointGroupIndex) const = 0;
         /**
-            @brief Calculate only the blend shape outputs of the rig.
+            @brief Calculate only the blend shape channel weights of the rig.
             @note
                 This is considered as an advanced usage use case.
             @param instance
-                The rig instance which outputs are to be calculated.
+                The rig instance whose blend shape channel weights are to be calculated.
             @see calculate
         */
         virtual void calculateBlendShapes(RigInstance* instance) const = 0;
@@ -168,7 +212,7 @@ class RLAPI RigLogic {
             @note
                 This is considered as an advanced usage use case.
             @param instance
-                The rig instance which outputs are to be calculated.
+                The rig instance whose outputs are to be calculated.
             @see calculate
         */
         virtual void calculateAnimatedMaps(RigInstance* instance) const = 0;
@@ -181,7 +225,7 @@ class RLAPI RigLogic {
                   - Calculate blend shape output values
                   - Calculate animated map output values
             @param instance
-                The rig instance which outputs are to be calculated.
+                The rig instance whose outputs are to be calculated.
         */
         virtual void calculate(RigInstance* instance) const = 0;
 
diff --git a/Engine/Plugins/Animation/RigLogic/Source/RigLogicLib/Public/riglogic/types/Aliases.h b/Engine/Plugins/Animation/RigLogic/Source/RigLogicLib/Public/riglogic/types/Aliases.h
--- a/Engine/Plugins/Animation/RigLogic/Source/RigLogicLib/Public/riglogic/types/Aliases.h
+++ b/Engine/Plugins/Animation/RigLogic/Source/RigLogicLib/Public/riglogic/types/Aliases.h
@@ -33,6 +33,7 @@ using dna::BinaryStreamWriter;
 using dna::JSONStreamReader;
 using dna::JSONStreamWriter;
 using dna::StringView;
+using dna::UnknownLayerPolicy;
 using dna::Vector3;
 
 template<typename T>
diff --git a/Engine/Plugins/Animation/RigLogic/Source/RigLogicLib/Public/riglogic/version/Version.h b/Engine/Plugins/Animation/RigLogic/Source/RigLogicLib/Public/riglogic/version/Version.h
--- a/Engine/Plugins/Animation/RigLogic/Source/RigLogicLib/Public/riglogic/version/Version.h
+++ b/Engine/Plugins/Animation/RigLogic/Source/RigLogicLib/Public/riglogic/version/Version.h
@@ -2,7 +2,7 @@
 
 #pragma once
 
-#define RL_MAJOR_VERSION 8
+#define RL_MAJOR_VERSION 9
 #define RL_MINOR_VERSION 1
-#define RL_PATCH_VERSION 0
-#define RL_VERSION_STRING "8.1.0"
+#define RL_PATCH_VERSION 1
+#define RL_VERSION_STRING "9.1.1"
diff --git a/Engine/Plugins/Animation/RigLogic/Source/RigLogicLib/Public/status/Provider.h b/Engine/Plugins/Animation/RigLogic/Source/RigLogicLib/Public/status/Provider.h
--- a/Engine/Plugins/Animation/RigLogic/Source/RigLogicLib/Public/status/Provider.h
+++ b/Engine/Plugins/Animation/RigLogic/Source/RigLogicLib/Public/status/Provider.h
@@ -12,13 +12,37 @@
 #include <stdio.h>
 
 #include <algorithm>
-#include <array>
 #ifdef _MSC_VER
     #pragma warning(pop)
 #endif
 
 namespace sc {
 
+namespace impl {
+
+template<std::size_t... Is>
+struct ISeq {
+    static constexpr std::size_t size() {
+        return sizeof...(Is);
+    }
+
+};
+
+template<std::size_t Size, std::size_t Offset, std::size_t... Is>
+struct Make
+    : Make<(Size - 1), Offset, (Offset + Size) - 1, Is...> {
+};
+
+template<std::size_t Offset, std::size_t... Is>
+struct Make<0, Offset, Is ...> {
+    using Type = ISeq<Is...>;
+};
+
+template<std::size_t Size, std::size_t Offset = std::size_t{}>
+using MakeISeq = typename Make<Size, Offset>::Type;
+
+}  // namespace impl
+
 class SCAPI StatusProvider {
     public:
         explicit StatusProvider(std::initializer_list<StatusCode> statuses);
@@ -28,24 +52,46 @@ class SCAPI StatusProvider {
         static bool isOk();
         static void set(StatusCode status);
 
-        template<typename ... Args>
-        static void set(StatusCode status, Args&& ... args) {
-            std::array<char, 512> buffer{};
+        template<std::size_t... Is, typename ... Args>
+        static void set(StatusCode status, impl::ISeq<Is...>  /*unused*/, Args&& ... args) {
+            constexpr std::size_t bufferSize = 512ul;
+            char buffer[bufferSize] = {};
             #if !defined(__clang__) && defined(__GNUC__)
                 #pragma GCC diagnostic push
                 #pragma GCC diagnostic ignored "-Wformat-security"
             #endif
+            // Invoke the hook with index = 0, denoting that this is the status message itself being hooked
+            // The return value from the hook will override the original message
+            status.message = execHook(status, 0ul, status.message);
             // The returned number of bytes to be written does not include the null terminator
-            const auto neededSize = snprintf(nullptr, 0ul, status.message, args ...) + 1;
-            const auto size = std::min(buffer.size(), static_cast<std::size_t>(neededSize));
-            snprintf(buffer.data(), size, status.message, args ...);
+            const auto neededSize = snprintf(nullptr, 0ul, status.message, execHook(status, Is, args) ...) + 1;
+            const auto size = std::min(bufferSize, static_cast<std::size_t>(neededSize));
+            // Invoke the hook with the index denoting the argument position [1..nargs] (for each const char* argument passed to `set`)
+            // The return value from the hook will override the original argument
+            // Arguments of type other than const char* remain untouched and just pass through
+            snprintf(buffer, size, status.message, execHook(status, Is, args) ...);
             #if !defined(__clang__) && defined(__GNUC__)
                 #pragma GCC diagnostic pop
             #endif
-            status.message = buffer.data();
-            set(status);
+            status.message = buffer;
+            execSet(status);
         }
 
+        template<typename ... Args>
+        static void set(StatusCode status, Args&& ... args) {
+            // Generate a compile-time integer sequence for each argument [1..nargs]
+            using ArgIndices = impl::MakeISeq<sizeof...(Args), static_cast<std::size_t>(1)>;
+            set(status, ArgIndices{}, args ...);
+        }
+
+    private:
+        static void execSet(StatusCode status);
+        static const char* execHook(StatusCode status, std::size_t index, const char* data);
+
+        template<typename T>
+        static T && execHook(StatusCode  /*unused*/, std::size_t  /*unused*/, T && data) {
+            return data;
+        }
 };
 
 }  // namespace sc
diff --git a/Engine/Plugins/Animation/RigLogic/Source/RigLogicLib/Public/status/Status.h b/Engine/Plugins/Animation/RigLogic/Source/RigLogicLib/Public/status/Status.h
--- a/Engine/Plugins/Animation/RigLogic/Source/RigLogicLib/Public/status/Status.h
+++ b/Engine/Plugins/Animation/RigLogic/Source/RigLogicLib/Public/status/Status.h
@@ -7,10 +7,29 @@
 
 namespace sc {
 
+/**
+ * @brief Hook function that allows overriding the status message itself or any of the string arguments used for interpolation of the final status message.
+ *
+ * @param StatusCode
+ *     The StatusCode itself whose message or argument is being hooked at the time of call.
+ * @param std::size_t
+ *     The index denoting which argument is being processed by the hook.
+ * @note
+ *     The value will be 0 when the hook is invoked for the status message itself, and from 1..N for any of the arguments used for interpolating the status message.
+ *     Although the hook is invoked only for const char* arguments, the indices will still denote their actual position in the argument list.
+ * @param const char*
+ *     The original status message or interpolation argument.
+ * @return
+ *     Another message that will be used as status message or interpolation argument instead of the passed in const char*.
+ */
+using HookFunction = const char* (*)(StatusCode, std::size_t, const char*);
+
 class SCAPI Status {
     public:
         static bool isOk();
         static StatusCode get();
+        static HookFunction getHook();
+        static void setHook(HookFunction hook);
 };
 
 }  // namespace sc
diff --git a/Engine/Plugins/Animation/RigLogic/Source/RigLogicLib/Public/status/version/Version.h b/Engine/Plugins/Animation/RigLogic/Source/RigLogicLib/Public/status/version/Version.h
--- a/Engine/Plugins/Animation/RigLogic/Source/RigLogicLib/Public/status/version/Version.h
+++ b/Engine/Plugins/Animation/RigLogic/Source/RigLogicLib/Public/status/version/Version.h
@@ -3,6 +3,6 @@
 #pragma once
 
 #define SC_MAJOR_VERSION 1
-#define SC_MINOR_VERSION 1
-#define SC_PATCH_VERSION 4
-#define SC_VERSION_STRING "1.1.4"
+#define SC_MINOR_VERSION 2
+#define SC_PATCH_VERSION 2
+#define SC_VERSION_STRING "1.2.2"
diff --git a/Engine/Plugins/Animation/RigLogic/Source/RigLogicLib/Public/terse/archives/binary/InputArchive.h b/Engine/Plugins/Animation/RigLogic/Source/RigLogicLib/Public/terse/archives/binary/InputArchive.h
--- a/Engine/Plugins/Animation/RigLogic/Source/RigLogicLib/Public/terse/archives/binary/InputArchive.h
+++ b/Engine/Plugins/Animation/RigLogic/Source/RigLogicLib/Public/terse/archives/binary/InputArchive.h
@@ -66,6 +66,11 @@ class ExtendableBinaryInputArchive : public Archive<TExtender> {
     protected:
         template<typename T>
         void process(Transparent<T>&& dest) {
+            process(dest);
+        }
+
+        template<typename T>
+        void process(Transparent<T>& dest) {
             process(dest.data);
         }
 
@@ -106,24 +111,48 @@ class ExtendableBinaryInputArchive : public Archive<TExtender> {
         template<typename T, typename V>
         typename std::enable_if<traits::has_versioned_load_member<T, V>::value,
                                 void>::type process(Versioned<T, V>&& dest) {
+            process(dest);
+        }
+
+        template<typename T, typename V>
+        typename std::enable_if<traits::has_versioned_load_member<T, V>::value,
+                                void>::type process(Versioned<T, V>& dest) {
             dest.data.load(*static_cast<TExtender*>(this), V{});
         }
 
         template<typename T, typename V>
         typename std::enable_if<traits::has_versioned_serialize_member<T, V>::value,
                                 void>::type process(Versioned<T, V>&& dest) {
+            process(dest);
+        }
+
+        template<typename T, typename V>
+        typename std::enable_if<traits::has_versioned_serialize_member<T, V>::value,
+                                void>::type process(Versioned<T, V>& dest) {
             dest.data.serialize(*static_cast<TExtender*>(this), V{});
         }
 
         template<typename T, typename V>
         typename std::enable_if<traits::has_versioned_load_function<T, V>::value,
                                 void>::type process(Versioned<T, V>&& dest) {
+            process(dest);
+        }
+
+        template<typename T, typename V>
+        typename std::enable_if<traits::has_versioned_load_function<T, V>::value,
+                                void>::type process(Versioned<T, V>& dest) {
             load(*static_cast<TExtender*>(this), V{}, dest.data);
         }
 
         template<typename T, typename V>
         typename std::enable_if<traits::has_versioned_serialize_function<T, V>::value,
                                 void>::type process(Versioned<T, V>&& dest) {
+            process(dest);
+        }
+
+        template<typename T, typename V>
+        typename std::enable_if<traits::has_versioned_serialize_function<T, V>::value,
+                                void>::type process(Versioned<T, V>& dest) {
             serialize(*static_cast<TExtender*>(this), V{}, dest.data);
         }
 
@@ -133,6 +162,15 @@ class ExtendableBinaryInputArchive : public Archive<TExtender> {
                                 !traits::has_versioned_load_function<T, V>::value &&
                                 !traits::has_versioned_serialize_function<T, V>::value,
                                 void>::type process(Versioned<T, V>&& dest) {
+            process(dest);
+        }
+
+        template<typename T, typename V>
+        typename std::enable_if<!traits::has_versioned_load_member<T, V>::value &&
+                                !traits::has_versioned_serialize_member<T, V>::value &&
+                                !traits::has_versioned_load_function<T, V>::value &&
+                                !traits::has_versioned_serialize_function<T, V>::value,
+                                void>::type process(Versioned<T, V>& dest) {
             // If no versioned serializer is found, but version information was passed along, fall back to unversioned
             // functions if available.
             BaseArchive::dispatch(dest.data);
diff --git a/Engine/Plugins/Animation/RigLogic/Source/RigLogicLib/Public/terse/archives/binary/OutputArchive.h b/Engine/Plugins/Animation/RigLogic/Source/RigLogicLib/Public/terse/archives/binary/OutputArchive.h
--- a/Engine/Plugins/Animation/RigLogic/Source/RigLogicLib/Public/terse/archives/binary/OutputArchive.h
+++ b/Engine/Plugins/Animation/RigLogic/Source/RigLogicLib/Public/terse/archives/binary/OutputArchive.h
@@ -66,6 +66,11 @@ class ExtendableBinaryOutputArchive : public Archive<TExtender> {
     protected:
         template<typename T>
         void process(Transparent<T>&& source) {
+            process(source);
+        }
+
+        template<typename T>
+        void process(Transparent<T>& source) {
             process(source.data);
         }
 
@@ -146,26 +151,59 @@ class ExtendableBinaryOutputArchive : public Archive<TExtender> {
 
         template<typename T, typename V>
         typename std::enable_if<traits::has_versioned_save_member<T, V>::value,
-                                void>::type process(const Versioned<T, V>& source) {
-            const_cast<T&>(source.data).save(*static_cast<TExtender*>(this), V{});
+                                void>::type process(Versioned<T, V>&& source) {
+            process(source);
+        }
+
+        template<typename T, typename V>
+        typename std::enable_if<traits::has_versioned_save_member<T, V>::value,
+                                void>::type process(Versioned<T, V>& source) {
+            source.data.save(*static_cast<TExtender*>(this), V{});
+        }
+
+        template<typename T, typename V>
+        typename std::enable_if<traits::has_versioned_serialize_member<T, V>::value,
+                                void>::type process(Versioned<T, V>&& source) {
+            process(source);
         }
 
         template<typename T, typename V>
         typename std::enable_if<traits::has_versioned_serialize_member<T, V>::value,
-                                void>::type process(const Versioned<T, V>& source) {
-            const_cast<T&>(source.data).serialize(*static_cast<TExtender*>(this), V{});
+                                void>::type process(Versioned<T, V>& source) {
+            source.data.serialize(*static_cast<TExtender*>(this), V{});
+        }
+
+        template<typename T, typename V>
+        typename std::enable_if<traits::has_versioned_save_function<T, V>::value,
+                                void>::type process(Versioned<T, V>&& source) {
+            process(source);
         }
 
         template<typename T, typename V>
         typename std::enable_if<traits::has_versioned_save_function<T, V>::value,
-                                void>::type process(const Versioned<T, V>& source) {
-            save(*static_cast<TExtender*>(this), V{}, const_cast<T&>(source.data));
+                                void>::type process(Versioned<T, V>& source) {
+            save(*static_cast<TExtender*>(this), V{}, source.data);
+        }
+
+        template<typename T, typename V>
+        typename std::enable_if<traits::has_versioned_serialize_function<T, V>::value,
+                                void>::type process(Versioned<T, V>&& source) {
+            process(source);
         }
 
         template<typename T, typename V>
         typename std::enable_if<traits::has_versioned_serialize_function<T, V>::value,
-                                void>::type process(const Versioned<T, V>& source) {
-            serialize(*static_cast<TExtender*>(this), V{}, const_cast<T&>(source.data));
+                                void>::type process(Versioned<T, V>& source) {
+            serialize(*static_cast<TExtender*>(this), V{}, source.data);
+        }
+
+        template<typename T, typename V>
+        typename std::enable_if<!traits::has_versioned_save_member<T, V>::value &&
+                                !traits::has_versioned_serialize_member<T, V>::value &&
+                                !traits::has_versioned_save_function<T, V>::value &&
+                                !traits::has_versioned_serialize_function<T, V>::value,
+                                void>::type process(Versioned<T, V>&& source) {
+            process(source);
         }
 
         template<typename T, typename V>
@@ -173,7 +211,7 @@ class ExtendableBinaryOutputArchive : public Archive<TExtender> {
                                 !traits::has_versioned_serialize_member<T, V>::value &&
                                 !traits::has_versioned_save_function<T, V>::value &&
                                 !traits::has_versioned_serialize_function<T, V>::value,
-                                void>::type process(Versioned<T, V>&& dest) {
+                                void>::type process(Versioned<T, V>& dest) {
             // If no versioned serializer is found, but version information was passed along, fall back to unversioned
             // functions if available.
             BaseArchive::dispatch(dest.data);
diff --git a/Engine/Plugins/Animation/RigLogic/Source/RigLogicLib/Public/terse/archives/json/InputArchive.h b/Engine/Plugins/Animation/RigLogic/Source/RigLogicLib/Public/terse/archives/json/InputArchive.h
--- a/Engine/Plugins/Animation/RigLogic/Source/RigLogicLib/Public/terse/archives/json/InputArchive.h
+++ b/Engine/Plugins/Animation/RigLogic/Source/RigLogicLib/Public/terse/archives/json/InputArchive.h
@@ -101,6 +101,11 @@ class ExtendableJSONInputArchive : public Archive<TExtender> {
     protected:
         template<typename T>
         void process(Transparent<T>&& dest) {
+            process(dest);
+        }
+
+        template<typename T>
+        void process(Transparent<T>& dest) {
             pushTransparency();
             process(dest.data);
         }
@@ -136,6 +141,12 @@ class ExtendableJSONInputArchive : public Archive<TExtender> {
         template<typename T, typename V>
         typename std::enable_if<traits::has_versioned_load_member<T, V>::value,
                                 void>::type process(Versioned<T, V>&& dest) {
+            process(dest);
+        }
+
+        template<typename T, typename V>
+        typename std::enable_if<traits::has_versioned_load_member<T, V>::value,
+                                void>::type process(Versioned<T, V>& dest) {
             if (state.malformed) {
                 return;
             }
@@ -155,6 +166,12 @@ class ExtendableJSONInputArchive : public Archive<TExtender> {
         template<typename T, typename V>
         typename std::enable_if<traits::has_versioned_serialize_member<T, V>::value,
                                 void>::type process(Versioned<T, V>&& dest) {
+            process(dest);
+        }
+
+        template<typename T, typename V>
+        typename std::enable_if<traits::has_versioned_serialize_member<T, V>::value,
+                                void>::type process(Versioned<T, V>& dest) {
             if (state.malformed) {
                 return;
             }
@@ -174,6 +191,12 @@ class ExtendableJSONInputArchive : public Archive<TExtender> {
         template<typename T, typename V>
         typename std::enable_if<traits::has_versioned_load_function<T, V>::value,
                                 void>::type process(Versioned<T, V>&& dest) {
+            process(dest);
+        }
+
+        template<typename T, typename V>
+        typename std::enable_if<traits::has_versioned_load_function<T, V>::value,
+                                void>::type process(Versioned<T, V>& dest) {
             if (state.malformed) {
                 return;
             }
@@ -193,6 +216,12 @@ class ExtendableJSONInputArchive : public Archive<TExtender> {
         template<typename T, typename V>
         typename std::enable_if<traits::has_versioned_serialize_function<T, V>::value,
                                 void>::type process(Versioned<T, V>&& dest) {
+            process(dest);
+        }
+
+        template<typename T, typename V>
+        typename std::enable_if<traits::has_versioned_serialize_function<T, V>::value,
+                                void>::type process(Versioned<T, V>& dest) {
             if (state.malformed) {
                 return;
             }
@@ -215,6 +244,15 @@ class ExtendableJSONInputArchive : public Archive<TExtender> {
                                 !traits::has_versioned_load_function<T, V>::value &&
                                 !traits::has_versioned_serialize_function<T, V>::value,
                                 void>::type process(Versioned<T, V>&& dest) {
+            process(dest);
+        }
+
+        template<typename T, typename V>
+        typename std::enable_if<!traits::has_versioned_load_member<T, V>::value &&
+                                !traits::has_versioned_serialize_member<T, V>::value &&
+                                !traits::has_versioned_load_function<T, V>::value &&
+                                !traits::has_versioned_serialize_function<T, V>::value,
+                                void>::type process(Versioned<T, V>& dest) {
             // If no versioned serializer is found, but version information was passed along, fall back to unversioned
             // functions if available.
             BaseArchive::dispatch(dest.data);
