Update RigLogic to v9.1.1 and GeneSplicer to v9.4.0

diff --git a/Engine/Plugins/Animation/RigLogic/Source/RigLogicLibTest/Private/dnatests/Fixturesv23.h b/Engine/Plugins/Animation/RigLogic/Source/RigLogicLibTest/Private/dnatests/Fixturesv23.h
new file mode 100644
--- /dev/null
+++ b/Engine/Plugins/Animation/RigLogic/Source/RigLogicLibTest/Private/dnatests/Fixturesv23.h
@@ -0,0 +1,131 @@
+// Copyright Epic Games, Inc. All Rights Reserved.
+
+#pragma once
+
+#include "dnatests/Defs.h"
+
+#include "dna/DNA.h"
+#include "dna/Reader.h"
+
+#include <pma/TypeDefs.h>
+
+namespace dna {
+
+struct RawV23 {
+    static const unsigned char header[91];
+    static const unsigned char descriptor[87];
+    static const unsigned char definition[794];
+    static const unsigned char controls[2];
+    static const unsigned char conditionals[324];
+    static const unsigned char psds[204];
+    static const unsigned char joints[444];
+    static const unsigned char blendshapes[44];
+    static const unsigned char animatedmaps[8];
+    static const unsigned char geometry[1080];
+    static const unsigned char machineLearnedBehavior[762];
+
+    static std::vector<char> getBytes();
+};
+
+struct RawV22DowngradedFromV23 {
+    static std::vector<char> getBytes();
+};
+
+struct DecodedV23 {
+    using VectorOfCharStringMatrix = pma::Vector<pma::Matrix<pma::String<char> > >;
+    using StringPair = std::pair<pma::String<char>, pma::String<char> >;
+
+    // Descriptor
+    static const pma::String<char> name;
+    static const Archetype archetype;
+    static const Gender gender;
+    static const std::uint16_t age;
+
+    static const pma::Vector<StringPair> metadata;
+    static const TranslationUnit translationUnit;
+    static const RotationUnit rotationUnit;
+    static const CoordinateSystem coordinateSystem;
+    static const std::uint16_t lodCount[3ul];
+    static const std::uint16_t maxLODs[3ul];
+    static const pma::String<char> complexity;
+    static const pma::String<char> dbName;
+
+    // Definition
+    static const pma::Vector<pma::String<char> > guiControlNames;
+    static const pma::Vector<pma::String<char> > rawControlNames;
+    static const VectorOfCharStringMatrix jointNames;
+    static const VectorOfCharStringMatrix blendShapeNames;
+    static const VectorOfCharStringMatrix animatedMapNames;
+    static const VectorOfCharStringMatrix meshNames;
+    static const pma::Vector<pma::Matrix<std::uint16_t> > meshBlendShapeIndices;
+    static const pma::Matrix<std::uint16_t> jointHierarchy;
+    static const pma::Vector<pma::Matrix<Vector3> > neutralJointTranslations;
+    static const pma::Vector<pma::Matrix<Vector3> > neutralJointRotations;
+
+    // Behavior
+    static const std::uint16_t guiControlCount;
+    static const std::uint16_t rawControlCount;
+    static const std::uint16_t psdCount;
+    static const pma::Vector<pma::Matrix<std::uint16_t> > conditionalInputIndices;
+    static const pma::Vector<pma::Matrix<std::uint16_t> > conditionalOutputIndices;
+    static const pma::Vector<pma::Matrix<float> > conditionalFromValues;
+    static const pma::Vector<pma::Matrix<float> > conditionalToValues;
+    static const pma::Vector<pma::Matrix<float> > conditionalSlopeValues;
+    static const pma::Vector<pma::Matrix<float> > conditionalCutValues;
+    static const pma::Vector<std::uint16_t> psdRowIndices;
+    static const pma::Vector<std::uint16_t> psdColumnIndices;
+    static const pma::Vector<float> psdValues;
+    static const pma::Vector<std::uint16_t> jointRowCount;
+    static const std::uint16_t jointColumnCount;
+    static const pma::Vector<pma::Matrix<std::uint16_t> > jointVariableIndices;
+    static const pma::Vector<pma::Matrix<std::uint16_t> > jointGroupLODs;
+    static const pma::Matrix<pma::Matrix<std::uint16_t> > jointGroupInputIndices;
+    static const pma::Matrix<pma::Matrix<std::uint16_t> > jointGroupOutputIndices;
+    static const pma::Matrix<pma::Matrix<float> > jointGroupValues;
+    static const pma::Matrix<pma::Matrix<std::uint16_t> > jointGroupJointIndices;
+    static const pma::Matrix<std::uint16_t> blendShapeLODs;
+    static const pma::Vector<pma::Matrix<std::uint16_t> > blendShapeInputIndices;
+    static const pma::Vector<pma::Matrix<std::uint16_t> > blendShapeOutputIndices;
+    static const pma::Vector<std::uint16_t> animatedMapCount;
+    static const pma::Matrix<std::uint16_t> animatedMapLODs;
+
+    // Geometry
+    static const pma::Vector<std::uint32_t> meshCount;
+    static const pma::Vector<pma::Matrix<Vector3> > vertexPositions;
+    static const pma::Vector<pma::Matrix<TextureCoordinate> > vertexTextureCoordinates;
+    static const pma::Vector<pma::Matrix<Vector3> > vertexNormals;
+    static const pma::Vector<pma::Matrix<VertexLayout> > vertexLayouts;
+    static const pma::Matrix<pma::Matrix<std::uint32_t> > faces;
+    static const pma::Matrix<std::uint16_t> maxInfluencePerVertex;
+    static const pma::Matrix<pma::Matrix<float> > skinWeightsValues;
+    static const pma::Matrix<pma::Matrix<std::uint16_t> > skinWeightsJointIndices;
+    static const pma::Vector<pma::Matrix<std::uint16_t> > correctiveBlendShapeIndices;
+    static const pma::Matrix<pma::Matrix<Vector3> > correctiveBlendShapeDeltas;
+    static const pma::Matrix<pma::Matrix<std::uint32_t> > correctiveBlendShapeVertexIndices;
+
+    // Machine learned behavior
+    static const pma::Vector<pma::String<char> > mlControlNames;
+    static const VectorOfCharStringMatrix regionNames;
+    static const pma::Matrix<std::uint16_t> neuralNetworkIndicesPerLOD;
+    static const pma::Matrix<pma::Matrix<std::uint16_t> > neuralNetworkIndicesPerMeshRegion;
+    static const pma::Vector<pma::Matrix<std::uint16_t> > neuralNetworkInputIndices;
+    static const pma::Vector<pma::Matrix<std::uint16_t> > neuralNetworkOutputIndices;
+    static const pma::Matrix<std::uint16_t> neuralNetworkLayerCount;
+    static const pma::Vector<pma::Matrix<std::uint16_t> > neuralNetworkActivationFunction;
+    static const pma::Matrix<pma::Matrix<float> > neuralNetworkActivationFunctionParameters;
+    static const pma::Matrix<pma::Matrix<float> > neuralNetworkBiases;
+    static const pma::Matrix<pma::Matrix<float> > neuralNetworkWeights;
+
+    static std::size_t lodConstraintToIndex(std::uint16_t maxLOD, std::uint16_t minLOD);
+    static RawJoints getJoints(std::uint16_t currentMaxLOD, std::uint16_t currentMinLOD, pma::MemoryResource* memRes);
+    static RawBlendShapeChannels getBlendShapes(std::uint16_t currentMaxLOD,
+                                                std::uint16_t currentMinLOD,
+                                                pma::MemoryResource* memRes);
+    static RawConditionalTable getConditionals(std::uint16_t currentMaxLOD,
+                                               std::uint16_t currentMinLOD,
+                                               pma::MemoryResource* memRes);
+    static RawAnimatedMaps getAnimatedMaps(std::uint16_t currentMaxLOD, std::uint16_t currentMinLOD, pma::MemoryResource* memRes);
+
+};
+
+}  // namespace dna
diff --git a/Engine/Plugins/Animation/RigLogic/Source/RigLogicLibTest/Private/dnatests/TestBinaryStreamWriter.h b/Engine/Plugins/Animation/RigLogic/Source/RigLogicLibTest/Private/dnatests/TestBinaryStreamWriter.h
--- a/Engine/Plugins/Animation/RigLogic/Source/RigLogicLibTest/Private/dnatests/TestBinaryStreamWriter.h
+++ b/Engine/Plugins/Animation/RigLogic/Source/RigLogicLibTest/Private/dnatests/TestBinaryStreamWriter.h
@@ -17,7 +17,11 @@ class BinaryStreamWriterTest : public ::testing::Test {
         void SetUp() override {
             stream = pma::makeScoped<trio::MemoryStream>();
             writer = dna::BinaryStreamWriter::create(stream.get(), &memRes);
-            reader = dna::BinaryStreamReader::create(stream.get(), dna::DataLayer::All, 0u, &memRes);
+            reader = dna::BinaryStreamReader::create(stream.get(),
+                                                     dna::DataLayer::All,
+                                                     dna::UnknownLayerPolicy::Preserve,
+                                                     0u,
+                                                     &memRes);
         }
 
         void TearDown() override {
diff --git a/Engine/Plugins/Animation/RigLogic/Source/RigLogicLibTest/Private/dnatests/TestLODMapping.cpp b/Engine/Plugins/Animation/RigLogic/Source/RigLogicLibTest/Private/dnatests/TestLODMapping.cpp
--- a/Engine/Plugins/Animation/RigLogic/Source/RigLogicLibTest/Private/dnatests/TestLODMapping.cpp
+++ b/Engine/Plugins/Animation/RigLogic/Source/RigLogicLibTest/Private/dnatests/TestLODMapping.cpp
@@ -1,7 +1,6 @@
 // Copyright Epic Games, Inc. All Rights Reserved.
 
 #include "dnatests/Defs.h"
-#include "dnatests/Fixtures.h"
 
 #include "dna/LODConstraint.h"
 #include "dna/LODMapping.h"
diff --git a/Engine/Plugins/Animation/RigLogic/Source/RigLogicLibTest/Private/dnatests/TestStreamReadWriteIntegration.cpp b/Engine/Plugins/Animation/RigLogic/Source/RigLogicLibTest/Private/dnatests/TestStreamReadWriteIntegration.cpp
--- a/Engine/Plugins/Animation/RigLogic/Source/RigLogicLibTest/Private/dnatests/TestStreamReadWriteIntegration.cpp
+++ b/Engine/Plugins/Animation/RigLogic/Source/RigLogicLibTest/Private/dnatests/TestStreamReadWriteIntegration.cpp
@@ -3,7 +3,10 @@
 #include "dnatests/TestStreamReadWriteIntegration.h"
 
 #include "dnatests/Defs.h"
-#include "dnatests/Fixtures.h"
+#include "dnatests/FixturesJSON.h"
+#include "dnatests/Fixturesv21.h"
+#include "dnatests/Fixturesv22.h"
+#include "dnatests/Fixturesv23.h"
 
 #include "dna/DataLayer.h"
 #include "dna/DNA.h"
@@ -12,385 +15,597 @@
 #include "dna/JSONStreamReader.h"
 #include "dna/JSONStreamWriter.h"
 
-#include <pma/resources/AlignedMemoryResource.h>
-
-StreamReadWriteIntegrationTest::~StreamReadWriteIntegrationTest() = default;
-
 namespace dna {
 
-static void verifyDescriptor(DescriptorReader* reader, const LODParameters& params) {
-    ASSERT_EQ(reader->getName(), StringView{decoded::name});
-    ASSERT_EQ(reader->getArchetype(), decoded::archetype);
-    ASSERT_EQ(reader->getGender(), decoded::gender);
-    ASSERT_EQ(reader->getAge(), decoded::age);
+template<class TAPICopyParameters>
+static void verifyDescriptor(DescriptorReader* reader) {
+    using DecodedDNA = typename TAPICopyParameters::DecodedData;
+    const auto index = DecodedDNA::lodConstraintToIndex(TAPICopyParameters::maxLOD(), TAPICopyParameters::minLOD());
+
+    ASSERT_EQ(reader->getName(), StringView{DecodedDNA::name});
+    ASSERT_EQ(reader->getArchetype(), DecodedDNA::archetype);
+    ASSERT_EQ(reader->getGender(), DecodedDNA::gender);
+    ASSERT_EQ(reader->getAge(), DecodedDNA::age);
 
     const auto metaDataCount = reader->getMetaDataCount();
     ASSERT_EQ(metaDataCount, 2u);
     for (std::uint32_t i = 0u; i < metaDataCount; ++i) {
         const auto key = reader->getMetaDataKey(i);
         const auto value = reader->getMetaDataValue(key);
-        ASSERT_EQ(key, StringView{decoded::metadata[i].first});
-        ASSERT_EQ(value, StringView{decoded::metadata[i].second});
+        ASSERT_EQ(key, StringView{DecodedDNA::metadata[i].first});
+        ASSERT_EQ(value, StringView{DecodedDNA::metadata[i].second});
     }
 
-    ASSERT_EQ(reader->getTranslationUnit(), decoded::translationUnit);
-    ASSERT_EQ(reader->getRotationUnit(), decoded::rotationUnit);
+    ASSERT_EQ(reader->getTranslationUnit(), DecodedDNA::translationUnit);
+    ASSERT_EQ(reader->getRotationUnit(), DecodedDNA::rotationUnit);
 
     const auto coordinateSystem = reader->getCoordinateSystem();
-    ASSERT_EQ(coordinateSystem.xAxis, decoded::coordinateSystem.xAxis);
-    ASSERT_EQ(coordinateSystem.yAxis, decoded::coordinateSystem.yAxis);
-    ASSERT_EQ(coordinateSystem.zAxis, decoded::coordinateSystem.zAxis);
-
-    ASSERT_EQ(reader->getLODCount(), decoded::lodCount[params.maxLOD]);
-    ASSERT_EQ(reader->getDBMaxLOD(), decoded::maxLODs[params.maxLOD]);
-    ASSERT_EQ(reader->getDBComplexity(), StringView{decoded::complexity});
-    ASSERT_EQ(reader->getDBName(), StringView{decoded::dbName});
+    ASSERT_EQ(coordinateSystem.xAxis, DecodedDNA::coordinateSystem.xAxis);
+    ASSERT_EQ(coordinateSystem.yAxis, DecodedDNA::coordinateSystem.yAxis);
+    ASSERT_EQ(coordinateSystem.zAxis, DecodedDNA::coordinateSystem.zAxis);
+
+    ASSERT_EQ(reader->getLODCount(), DecodedDNA::lodCount[index]);
+    ASSERT_EQ(reader->getDBMaxLOD(), DecodedDNA::maxLODs[index]);
+    ASSERT_EQ(reader->getDBComplexity(), StringView{DecodedDNA::complexity});
+    ASSERT_EQ(reader->getDBName(), StringView{DecodedDNA::dbName});
 }
 
-static void verifyDefinition(DefinitionReader* reader, const LODParameters& params) {
+template<class TAPICopyParameters>
+static void verifyDefinition(DefinitionReader* reader) {
+    using DecodedDNA = typename TAPICopyParameters::DecodedData;
+    const auto index = DecodedDNA::lodConstraintToIndex(TAPICopyParameters::maxLOD(), TAPICopyParameters::minLOD());
+
     const auto guiControlCount = reader->getGUIControlCount();
-    ASSERT_EQ(guiControlCount, decoded::guiControlNames.size());
+    ASSERT_EQ(guiControlCount, DecodedDNA::guiControlNames.size());
     for (std::uint16_t i = 0u; i < guiControlCount; ++i) {
-        ASSERT_EQ(reader->getGUIControlName(i), StringView{decoded::guiControlNames[i]});
+        ASSERT_EQ(reader->getGUIControlName(i), StringView{DecodedDNA::guiControlNames[i]});
     }
 
     const auto rawControlCount = reader->getRawControlCount();
-    ASSERT_EQ(rawControlCount, decoded::rawControlNames.size());
+    ASSERT_EQ(rawControlCount, DecodedDNA::rawControlNames.size());
     for (std::uint16_t i = 0u; i < rawControlCount; ++i) {
-        ASSERT_EQ(reader->getRawControlName(i), StringView{decoded::rawControlNames[i]});
+        ASSERT_EQ(reader->getRawControlName(i), StringView{DecodedDNA::rawControlNames[i]});
     }
 
-    ASSERT_EQ(reader->getJointCount(), decoded::jointNames[params.maxLOD][0ul].size());
-    const auto& expectedJointNames = decoded::jointNames[params.maxLOD][params.currentLOD];
-    const auto jointIndices = reader->getJointIndicesForLOD(params.currentLOD);
+    ASSERT_EQ(reader->getJointCount(), DecodedDNA::jointNames[index][0ul].size());
+    const auto& expectedJointNames = DecodedDNA::jointNames[index][TAPICopyParameters::currentLOD()];
+    const auto jointIndices = reader->getJointIndicesForLOD(TAPICopyParameters::currentLOD());
     ASSERT_EQ(jointIndices.size(), expectedJointNames.size());
     for (std::size_t i = 0ul; i < jointIndices.size(); ++i) {
         ASSERT_EQ(reader->getJointName(jointIndices[i]), StringView{expectedJointNames[i]});
     }
 
     for (std::uint16_t i = 0u; i < reader->getJointCount(); ++i) {
-        ASSERT_EQ(reader->getJointParentIndex(i), decoded::jointHierarchy[params.maxLOD][i]);
+        ASSERT_EQ(reader->getJointParentIndex(i), DecodedDNA::jointHierarchy[index][i]);
     }
 
-    ASSERT_EQ(reader->getBlendShapeChannelCount(), decoded::blendShapeNames[params.maxLOD][0ul].size());
-    const auto& expectedBlendShapeNames = decoded::blendShapeNames[params.maxLOD][params.currentLOD];
-    const auto blendShapeIndices = reader->getBlendShapeChannelIndicesForLOD(params.currentLOD);
+    ASSERT_EQ(reader->getBlendShapeChannelCount(), DecodedDNA::blendShapeNames[index][0ul].size());
+    const auto& expectedBlendShapeNames = DecodedDNA::blendShapeNames[index][TAPICopyParameters::currentLOD()];
+    const auto blendShapeIndices = reader->getBlendShapeChannelIndicesForLOD(TAPICopyParameters::currentLOD());
     ASSERT_EQ(blendShapeIndices.size(), expectedBlendShapeNames.size());
     for (std::size_t i = 0ul; i < blendShapeIndices.size(); ++i) {
         ASSERT_EQ(reader->getBlendShapeChannelName(blendShapeIndices[i]), StringView{expectedBlendShapeNames[i]});
     }
 
-    ASSERT_EQ(reader->getAnimatedMapCount(), decoded::animatedMapNames[params.maxLOD][0ul].size());
-    const auto& expectedAnimatedMapNames = decoded::animatedMapNames[params.maxLOD][params.currentLOD];
-    const auto animatedMapIndices = reader->getAnimatedMapIndicesForLOD(params.currentLOD);
+    ASSERT_EQ(reader->getAnimatedMapCount(), DecodedDNA::animatedMapNames[index][0ul].size());
+    const auto& expectedAnimatedMapNames = DecodedDNA::animatedMapNames[index][TAPICopyParameters::currentLOD()];
+    const auto animatedMapIndices = reader->getAnimatedMapIndicesForLOD(TAPICopyParameters::currentLOD());
     ASSERT_EQ(animatedMapIndices.size(), expectedAnimatedMapNames.size());
     for (std::size_t i = 0ul; i < animatedMapIndices.size(); ++i) {
         ASSERT_EQ(reader->getAnimatedMapName(animatedMapIndices[i]), StringView{expectedAnimatedMapNames[i]});
     }
 
-    ASSERT_EQ(reader->getMeshCount(), decoded::meshNames[params.maxLOD][0ul].size());
-    const auto& expectedMeshNames = decoded::meshNames[params.maxLOD][params.currentLOD];
-    const auto meshIndices = reader->getMeshIndicesForLOD(params.currentLOD);
+    std::uint16_t expectedMeshCount = {};
+    for (std::uint16_t i = 0ul; i < DecodedDNA::meshNames[index].size(); ++i) {
+        expectedMeshCount = static_cast<std::uint16_t>(expectedMeshCount + DecodedDNA::meshNames[index][i].size());
+    }
+    ASSERT_EQ(reader->getMeshCount(), expectedMeshCount);
+    const auto& expectedMeshNames = DecodedDNA::meshNames[index][TAPICopyParameters::currentLOD()];
+    const auto meshIndices = reader->getMeshIndicesForLOD(TAPICopyParameters::currentLOD());
     ASSERT_EQ(meshIndices.size(), expectedMeshNames.size());
     for (std::size_t i = 0ul; i < meshIndices.size(); ++i) {
         ASSERT_EQ(reader->getMeshName(meshIndices[i]), StringView{expectedMeshNames[i]});
     }
 
-    ASSERT_EQ(reader->getMeshBlendShapeChannelMappingCount(), decoded::meshBlendShapeIndices[params.maxLOD][0ul].size());
-    const auto meshBlendShapeIndices = reader->getMeshBlendShapeChannelMappingIndicesForLOD(params.currentLOD);
-    const auto& expectedMeshBlendShapeIndices = decoded::meshBlendShapeIndices[params.maxLOD][params.currentLOD];
+    std::uint16_t expectedMeshBlendShapeMappingCount = {};
+    for (std::uint16_t i = 0ul; i < DecodedDNA::meshBlendShapeIndices[index].size(); ++i) {
+        expectedMeshBlendShapeMappingCount =
+            static_cast<std::uint16_t>(expectedMeshBlendShapeMappingCount + DecodedDNA::meshBlendShapeIndices[index][i].size());
+    }
+    ASSERT_EQ(reader->getMeshBlendShapeChannelMappingCount(), expectedMeshBlendShapeMappingCount);
+    const auto meshBlendShapeIndices = reader->getMeshBlendShapeChannelMappingIndicesForLOD(TAPICopyParameters::currentLOD());
+    const auto& expectedMeshBlendShapeIndices = DecodedDNA::meshBlendShapeIndices[index][TAPICopyParameters::currentLOD()];
     ASSERT_EQ(meshBlendShapeIndices, ConstArrayView<std::uint16_t>{expectedMeshBlendShapeIndices});
 
-    const auto& expectedNeutralJointTranslations = decoded::neutralJointTranslations[params.maxLOD][params.currentLOD];
+    const auto& expectedNeutralJointTranslations = DecodedDNA::neutralJointTranslations[index][TAPICopyParameters::currentLOD()];
     ASSERT_EQ(jointIndices.size(), expectedNeutralJointTranslations.size());
     for (std::size_t i = 0ul; i < jointIndices.size(); ++i) {
         ASSERT_EQ(reader->getNeutralJointTranslation(jointIndices[i]), expectedNeutralJointTranslations[i]);
     }
 
-    const auto& expectedNeutralJointRotations = decoded::neutralJointRotations[params.maxLOD][params.currentLOD];
+    const auto& expectedNeutralJointRotations = DecodedDNA::neutralJointRotations[index][TAPICopyParameters::currentLOD()];
     ASSERT_EQ(jointIndices.size(), expectedNeutralJointRotations.size());
     for (std::size_t i = 0ul; i < jointIndices.size(); ++i) {
         ASSERT_EQ(reader->getNeutralJointRotation(jointIndices[i]), expectedNeutralJointRotations[i]);
     }
 }
 
-static void verifyBehavior(BehaviorReader* reader, const LODParameters& params) {
+template<class TAPICopyParameters>
+static void verifyBehavior(BehaviorReader* reader) {
+    using DecodedDNA = typename TAPICopyParameters::DecodedData;
+    const auto index = DecodedDNA::lodConstraintToIndex(TAPICopyParameters::maxLOD(), TAPICopyParameters::minLOD());
+
     const auto guiToRawInputIndices = reader->getGUIToRawInputIndices();
-    const auto& expectedG2RInputIndices = decoded::conditionalInputIndices[0ul][0ul];
+    const auto& expectedG2RInputIndices = DecodedDNA::conditionalInputIndices[0ul][0ul];
     ASSERT_EQ(guiToRawInputIndices, ConstArrayView<std::uint16_t>{expectedG2RInputIndices});
 
     const auto guiToRawOutputIndices = reader->getGUIToRawOutputIndices();
-    const auto& expectedG2ROutputIndices = decoded::conditionalOutputIndices[0ul][0ul];
+    const auto& expectedG2ROutputIndices = DecodedDNA::conditionalOutputIndices[0ul][0ul];
     ASSERT_EQ(guiToRawOutputIndices, ConstArrayView<std::uint16_t>{expectedG2ROutputIndices});
 
     const auto guiToRawFromValues = reader->getGUIToRawFromValues();
-    const auto& expectedG2RFromValues = decoded::conditionalFromValues[0ul][0ul];
+    const auto& expectedG2RFromValues = DecodedDNA::conditionalFromValues[0ul][0ul];
     ASSERT_EQ(guiToRawFromValues, ConstArrayView<float>{expectedG2RFromValues});
 
     const auto guiToRawToValues = reader->getGUIToRawToValues();
-    const auto& expectedG2RToValues = decoded::conditionalToValues[0ul][0ul];
+    const auto& expectedG2RToValues = DecodedDNA::conditionalToValues[0ul][0ul];
     ASSERT_EQ(guiToRawToValues, ConstArrayView<float>{expectedG2RToValues});
 
     const auto guiToRawSlopeValues = reader->getGUIToRawSlopeValues();
-    const auto& expectedG2RSlopeValues = decoded::conditionalSlopeValues[0ul][0ul];
+    const auto& expectedG2RSlopeValues = DecodedDNA::conditionalSlopeValues[0ul][0ul];
     ASSERT_EQ(guiToRawSlopeValues, ConstArrayView<float>{expectedG2RSlopeValues});
 
     const auto guiToRawCutValues = reader->getGUIToRawCutValues();
-    const auto& expectedG2RCutValues = decoded::conditionalCutValues[0ul][0ul];
+    const auto& expectedG2RCutValues = DecodedDNA::conditionalCutValues[0ul][0ul];
     ASSERT_EQ(guiToRawCutValues, ConstArrayView<float>{expectedG2RCutValues});
 
     const auto psdRowIndices = reader->getPSDRowIndices();
-    ASSERT_EQ(psdRowIndices, ConstArrayView<std::uint16_t>{decoded::psdRowIndices});
+    ASSERT_EQ(psdRowIndices, ConstArrayView<std::uint16_t>{DecodedDNA::psdRowIndices});
 
     const auto psdColumnIndices = reader->getPSDColumnIndices();
-    ASSERT_EQ(psdColumnIndices, ConstArrayView<std::uint16_t>{decoded::psdColumnIndices});
+    ASSERT_EQ(psdColumnIndices, ConstArrayView<std::uint16_t>{DecodedDNA::psdColumnIndices});
 
     const auto psdValues = reader->getPSDValues();
-    ASSERT_EQ(psdValues, ConstArrayView<float>{decoded::psdValues});
+    ASSERT_EQ(psdValues, ConstArrayView<float>{DecodedDNA::psdValues});
 
-    ASSERT_EQ(reader->getPSDCount(), decoded::psdCount);
-    ASSERT_EQ(reader->getJointRowCount(), decoded::jointRowCount[params.maxLOD]);
-    ASSERT_EQ(reader->getJointColumnCount(), decoded::jointColumnCount);
+    ASSERT_EQ(reader->getPSDCount(), DecodedDNA::psdCount);
+    ASSERT_EQ(reader->getJointRowCount(), DecodedDNA::jointRowCount[index]);
+    ASSERT_EQ(reader->getJointColumnCount(), DecodedDNA::jointColumnCount);
 
-    const auto jointVariableAttrIndices = reader->getJointVariableAttributeIndices(params.currentLOD);
-    const auto& expectedJointVariableAttrIndices = decoded::jointVariableIndices[params.maxLOD][params.currentLOD];
+    const auto jointVariableAttrIndices = reader->getJointVariableAttributeIndices(TAPICopyParameters::currentLOD());
+    const auto& expectedJointVariableAttrIndices = DecodedDNA::jointVariableIndices[index][TAPICopyParameters::currentLOD()];
     ASSERT_EQ(jointVariableAttrIndices, ConstArrayView<std::uint16_t>{expectedJointVariableAttrIndices});
 
     const auto jointGroupCount = reader->getJointGroupCount();
-    ASSERT_EQ(jointGroupCount, decoded::jointGroupLODs.size());
+    ASSERT_EQ(jointGroupCount, DecodedDNA::jointGroupLODs.size());
 
     for (std::uint16_t i = 0u; i < jointGroupCount; ++i) {
-        const auto& expectedLODs = decoded::jointGroupLODs[i][params.maxLOD];
+        const auto& expectedLODs = DecodedDNA::jointGroupLODs[i][index];
         ASSERT_EQ(reader->getJointGroupLODs(i), ConstArrayView<std::uint16_t>{expectedLODs});
 
-        const auto& expectedInputIndices = decoded::jointGroupInputIndices[i][params.maxLOD][0ul];
+        const auto& expectedInputIndices = DecodedDNA::jointGroupInputIndices[i][index][0ul];
         ASSERT_EQ(reader->getJointGroupInputIndices(i), ConstArrayView<std::uint16_t>{expectedInputIndices});
 
         const auto outputIndices = reader->getJointGroupOutputIndices(i);
         ASSERT_EQ(outputIndices.size(), expectedLODs[0ul]);
 
-        ConstArrayView<std::uint16_t> outputIndicesForLOD{outputIndices.data(), expectedLODs[params.currentLOD]};
-        const auto& expectedOutputIndices = decoded::jointGroupOutputIndices[i][params.maxLOD][params.currentLOD];
+        ConstArrayView<std::uint16_t> outputIndicesForLOD{outputIndices.data(), expectedLODs[TAPICopyParameters::currentLOD()]};
+        const auto& expectedOutputIndices = DecodedDNA::jointGroupOutputIndices[i][index][TAPICopyParameters::currentLOD()];
         ASSERT_EQ(outputIndicesForLOD, ConstArrayView<std::uint16_t>{expectedOutputIndices});
 
         const auto values = reader->getJointGroupValues(i);
         ASSERT_EQ(values.size(), expectedLODs[0ul] * expectedInputIndices.size());
 
-        ConstArrayView<float> valuesForLOD{values.data(), expectedLODs[params.currentLOD] * expectedInputIndices.size()};
-        const auto& expectedValues = decoded::jointGroupValues[i][params.maxLOD][params.currentLOD];
+        ConstArrayView<float> valuesForLOD{values.data(),
+                              expectedLODs[TAPICopyParameters::currentLOD()] * expectedInputIndices.size()};
+        const auto& expectedValues = DecodedDNA::jointGroupValues[i][index][TAPICopyParameters::currentLOD()];
         ASSERT_EQ(valuesForLOD, ConstArrayView<float>{expectedValues});
 
-        const auto& expectedJointIndices = decoded::jointGroupJointIndices[i][params.maxLOD][0ul];
+        const auto& expectedJointIndices = DecodedDNA::jointGroupJointIndices[i][index][0ul];
         ASSERT_EQ(reader->getJointGroupJointIndices(i), ConstArrayView<std::uint16_t>{expectedJointIndices});
     }
 
     ASSERT_EQ(reader->getBlendShapeChannelLODs(),
-              ConstArrayView<std::uint16_t>{decoded::blendShapeLODs[params.maxLOD]});
+              ConstArrayView<std::uint16_t>{DecodedDNA::blendShapeLODs[index]});
 
     const auto blendShapeChannelInputIndices = reader->getBlendShapeChannelInputIndices();
-    ASSERT_EQ(blendShapeChannelInputIndices.size(), decoded::blendShapeLODs[params.maxLOD][0ul]);
+    ASSERT_EQ(blendShapeChannelInputIndices.size(), DecodedDNA::blendShapeLODs[index][0ul]);
     ConstArrayView<std::uint16_t> blendShapeInputIndicesForLOD{
         blendShapeChannelInputIndices.data(),
-        decoded::blendShapeLODs[params.maxLOD][params.currentLOD]
+        DecodedDNA::blendShapeLODs[index][TAPICopyParameters::currentLOD()]
     };
     ASSERT_EQ(blendShapeInputIndicesForLOD,
-              ConstArrayView<std::uint16_t>{decoded::blendShapeInputIndices[params.maxLOD][params.currentLOD]});
+              ConstArrayView<std::uint16_t>{DecodedDNA::blendShapeInputIndices[index][TAPICopyParameters::currentLOD()]});
 
     const auto blendShapeChannelOutputIndices = reader->getBlendShapeChannelOutputIndices();
-    ASSERT_EQ(blendShapeChannelOutputIndices.size(), decoded::blendShapeLODs[params.maxLOD][0ul]);
+    ASSERT_EQ(blendShapeChannelOutputIndices.size(), DecodedDNA::blendShapeLODs[index][0ul]);
     ConstArrayView<std::uint16_t> blendShapeOutputIndicesForLOD{
         blendShapeChannelOutputIndices.data(),
-        decoded::blendShapeLODs[params.maxLOD][params.currentLOD]
+        DecodedDNA::blendShapeLODs[index][TAPICopyParameters::currentLOD()]
     };
     ASSERT_EQ(blendShapeOutputIndicesForLOD,
-              ConstArrayView<std::uint16_t>{decoded::blendShapeOutputIndices[params.maxLOD][params.currentLOD]});
+              ConstArrayView<std::uint16_t>{DecodedDNA::blendShapeOutputIndices[index][TAPICopyParameters::currentLOD()]});
 
     ASSERT_EQ(reader->getAnimatedMapLODs(),
-              ConstArrayView<std::uint16_t>{decoded::animatedMapLODs[params.maxLOD]});
+              ConstArrayView<std::uint16_t>{DecodedDNA::animatedMapLODs[index]});
 
-    ASSERT_EQ(reader->getAnimatedMapCount(), decoded::animatedMapCount[params.maxLOD]);
+    ASSERT_EQ(reader->getAnimatedMapCount(), DecodedDNA::animatedMapCount[index]);
 
-    const auto animatdMapLOD = decoded::animatedMapLODs[params.maxLOD][params.currentLOD];
+    const auto animatedMapLOD = DecodedDNA::animatedMapLODs[index][TAPICopyParameters::currentLOD()];
 
     const auto animatedMapInputIndices = reader->getAnimatedMapInputIndices();
-    ASSERT_EQ(animatedMapInputIndices.size(), decoded::animatedMapLODs[params.maxLOD][0ul]);
-    ConstArrayView<std::uint16_t> animatedMapInputIndicesForLOD{animatedMapInputIndices.data(), animatdMapLOD};
+    ASSERT_EQ(animatedMapInputIndices.size(), DecodedDNA::animatedMapLODs[index][0ul]);
+    ConstArrayView<std::uint16_t> animatedMapInputIndicesForLOD{animatedMapInputIndices.data(), animatedMapLOD};
     ASSERT_EQ(animatedMapInputIndicesForLOD,
-              ConstArrayView<std::uint16_t>{decoded::conditionalInputIndices[params.maxLOD][params.currentLOD]});
+              ConstArrayView<std::uint16_t>{DecodedDNA::conditionalInputIndices[index][TAPICopyParameters::currentLOD()]});
 
     const auto animatedMapOutputIndices = reader->getAnimatedMapOutputIndices();
-    ASSERT_EQ(animatedMapOutputIndices.size(), decoded::animatedMapLODs[params.maxLOD][0ul]);
-    ConstArrayView<std::uint16_t> animatedMapOutputIndicesForLOD{animatedMapOutputIndices.data(), animatdMapLOD};
+    ASSERT_EQ(animatedMapOutputIndices.size(), DecodedDNA::animatedMapLODs[index][0ul]);
+    ConstArrayView<std::uint16_t> animatedMapOutputIndicesForLOD{animatedMapOutputIndices.data(), animatedMapLOD};
     ASSERT_EQ(animatedMapOutputIndicesForLOD,
-              ConstArrayView<std::uint16_t>{decoded::conditionalOutputIndices[params.maxLOD][params.currentLOD]});
+              ConstArrayView<std::uint16_t>{DecodedDNA::conditionalOutputIndices[index][TAPICopyParameters::currentLOD()]});
 
     const auto animatedMapFromValues = reader->getAnimatedMapFromValues();
-    ASSERT_EQ(animatedMapFromValues.size(), decoded::animatedMapLODs[params.maxLOD][0ul]);
-    ConstArrayView<float> animatedMapFromValuesForLOD{animatedMapFromValues.data(), animatdMapLOD};
+    ASSERT_EQ(animatedMapFromValues.size(), DecodedDNA::animatedMapLODs[index][0ul]);
+    ConstArrayView<float> animatedMapFromValuesForLOD{animatedMapFromValues.data(), animatedMapLOD};
     ASSERT_EQ(animatedMapFromValuesForLOD,
-              ConstArrayView<float>{decoded::conditionalFromValues[params.maxLOD][params.currentLOD]});
+              ConstArrayView<float>{DecodedDNA::conditionalFromValues[index][TAPICopyParameters::currentLOD()]});
 
     const auto animatedMapToValues = reader->getAnimatedMapToValues();
-    ASSERT_EQ(animatedMapToValues.size(), decoded::animatedMapLODs[params.maxLOD][0ul]);
-    ConstArrayView<float> animatedMapToValuesForLOD{animatedMapToValues.data(), animatdMapLOD};
+    ASSERT_EQ(animatedMapToValues.size(), DecodedDNA::animatedMapLODs[index][0ul]);
+    ConstArrayView<float> animatedMapToValuesForLOD{animatedMapToValues.data(), animatedMapLOD};
     ASSERT_EQ(animatedMapToValuesForLOD,
-              ConstArrayView<float>{decoded::conditionalToValues[params.maxLOD][params.currentLOD]});
+              ConstArrayView<float>{DecodedDNA::conditionalToValues[index][TAPICopyParameters::currentLOD()]});
 
     const auto animatedMapSlopeValues = reader->getAnimatedMapSlopeValues();
-    ASSERT_EQ(animatedMapSlopeValues.size(), decoded::animatedMapLODs[params.maxLOD][0ul]);
-    ConstArrayView<float> animatedMapSlopeValuesForLOD{animatedMapSlopeValues.data(), animatdMapLOD};
+    ASSERT_EQ(animatedMapSlopeValues.size(), DecodedDNA::animatedMapLODs[index][0ul]);
+    ConstArrayView<float> animatedMapSlopeValuesForLOD{animatedMapSlopeValues.data(), animatedMapLOD};
     ASSERT_EQ(animatedMapSlopeValuesForLOD,
-              ConstArrayView<float>{decoded::conditionalSlopeValues[params.maxLOD][params.currentLOD]});
+              ConstArrayView<float>{DecodedDNA::conditionalSlopeValues[index][TAPICopyParameters::currentLOD()]});
 
     const auto animatedMapCutValues = reader->getAnimatedMapCutValues();
-    ASSERT_EQ(animatedMapCutValues.size(), decoded::animatedMapLODs[params.maxLOD][0ul]);
-    ConstArrayView<float> animatedMapCutValuesForLOD{animatedMapCutValues.data(), animatdMapLOD};
+    ASSERT_EQ(animatedMapCutValues.size(), DecodedDNA::animatedMapLODs[index][0ul]);
+    ConstArrayView<float> animatedMapCutValuesForLOD{animatedMapCutValues.data(), animatedMapLOD};
     ASSERT_EQ(animatedMapCutValuesForLOD,
-              ConstArrayView<float>{decoded::conditionalCutValues[params.maxLOD][params.currentLOD]});
+              ConstArrayView<float>{DecodedDNA::conditionalCutValues[index][TAPICopyParameters::currentLOD()]});
 }
 
-static void verifyGeometry(GeometryReader* reader, const LODParameters& params) {
+template<class TAPICopyParameters>
+static void verifyGeometry(GeometryReader* reader) {
+    using DecodedDNA = typename TAPICopyParameters::DecodedData;
+    const auto index = DecodedDNA::lodConstraintToIndex(TAPICopyParameters::maxLOD(), TAPICopyParameters::minLOD());
+
     const auto meshCount = reader->getMeshCount();
-    ASSERT_EQ(meshCount, decoded::meshCount[params.maxLOD]);
+    ASSERT_EQ(meshCount, DecodedDNA::meshCount[index]);
     for (std::uint16_t meshIndex = 0u; meshIndex < meshCount; ++meshIndex) {
         const auto vertexPositionCount = reader->getVertexPositionCount(meshIndex);
-        ASSERT_EQ(vertexPositionCount, decoded::vertexPositions[params.maxLOD][meshIndex].size());
+        ASSERT_EQ(vertexPositionCount, DecodedDNA::vertexPositions[index][meshIndex].size());
         for (std::uint32_t vertexIndex = 0u; vertexIndex < vertexPositionCount; ++vertexIndex) {
             ASSERT_EQ(reader->getVertexPosition(meshIndex, vertexIndex),
-                      decoded::vertexPositions[params.maxLOD][meshIndex][vertexIndex]);
+                      DecodedDNA::vertexPositions[index][meshIndex][vertexIndex]);
         }
 
         const auto vertexTextureCoordinateCount = reader->getVertexTextureCoordinateCount(meshIndex);
-        ASSERT_EQ(vertexTextureCoordinateCount, decoded::vertexTextureCoordinates[params.maxLOD][meshIndex].size());
+        ASSERT_EQ(vertexTextureCoordinateCount, DecodedDNA::vertexTextureCoordinates[index][meshIndex].size());
         for (std::uint32_t texCoordIndex = 0u; texCoordIndex < vertexTextureCoordinateCount; ++texCoordIndex) {
             const auto& textureCoordinate = reader->getVertexTextureCoordinate(meshIndex, texCoordIndex);
             const auto& expectedTextureCoordinate =
-                decoded::vertexTextureCoordinates[params.maxLOD][meshIndex][texCoordIndex];
+                DecodedDNA::vertexTextureCoordinates[index][meshIndex][texCoordIndex];
             ASSERT_EQ(textureCoordinate.u, expectedTextureCoordinate.u);
             ASSERT_EQ(textureCoordinate.v, expectedTextureCoordinate.v);
         }
 
         const auto vertexNormalCount = reader->getVertexNormalCount(meshIndex);
-        ASSERT_EQ(vertexNormalCount, decoded::vertexNormals[params.maxLOD][meshIndex].size());
+        ASSERT_EQ(vertexNormalCount, DecodedDNA::vertexNormals[index][meshIndex].size());
         for (std::uint32_t normalIndex = 0u; normalIndex < vertexNormalCount; ++normalIndex) {
             ASSERT_EQ(reader->getVertexNormal(meshIndex, normalIndex),
-                      decoded::vertexNormals[params.maxLOD][meshIndex][normalIndex]);
+                      DecodedDNA::vertexNormals[index][meshIndex][normalIndex]);
         }
 
         const auto vertexLayoutCount = reader->getVertexLayoutCount(meshIndex);
-        ASSERT_EQ(vertexLayoutCount, decoded::vertexLayouts[params.maxLOD][meshIndex].size());
+        ASSERT_EQ(vertexLayoutCount, DecodedDNA::vertexLayouts[index][meshIndex].size());
         for (std::uint32_t layoutIndex = 0u; layoutIndex < vertexLayoutCount; ++layoutIndex) {
             const auto& layout = reader->getVertexLayout(meshIndex, layoutIndex);
-            const auto& expectedLayout = decoded::vertexLayouts[params.maxLOD][meshIndex][layoutIndex];
+            const auto& expectedLayout = DecodedDNA::vertexLayouts[index][meshIndex][layoutIndex];
             ASSERT_EQ(layout.position, expectedLayout.position);
             ASSERT_EQ(layout.textureCoordinate, expectedLayout.textureCoordinate);
             ASSERT_EQ(layout.normal, expectedLayout.normal);
         }
 
         const auto faceCount = reader->getFaceCount(meshIndex);
-        ASSERT_EQ(faceCount, decoded::faces[params.maxLOD][meshIndex].size());
+        ASSERT_EQ(faceCount, DecodedDNA::faces[index][meshIndex].size());
         for (std::uint32_t faceIndex = 0u; faceIndex < faceCount; ++faceIndex) {
             ASSERT_EQ(reader->getFaceVertexLayoutIndices(meshIndex, faceIndex),
-                      ConstArrayView<std::uint32_t>{decoded::faces[params.maxLOD][meshIndex][faceIndex]});
+                      ConstArrayView<std::uint32_t>{DecodedDNA::faces[index][meshIndex][faceIndex]});
         }
 
         ASSERT_EQ(reader->getMaximumInfluencePerVertex(meshIndex),
-                  decoded::maxInfluencePerVertex[params.maxLOD][meshIndex]);
+                  DecodedDNA::maxInfluencePerVertex[index][meshIndex]);
 
-        ASSERT_EQ(reader->getSkinWeightsCount(meshIndex), decoded::skinWeightsValues[params.maxLOD][meshIndex].size());
+        ASSERT_EQ(reader->getSkinWeightsCount(meshIndex), DecodedDNA::skinWeightsValues[index][meshIndex].size());
         for (std::uint32_t vertexIndex = 0u; vertexIndex < vertexPositionCount; ++vertexIndex) {
             const auto skinWeights = reader->getSkinWeightsValues(meshIndex, vertexIndex);
-            const auto& expectedSkinWeights = decoded::skinWeightsValues[params.maxLOD][meshIndex][vertexIndex];
+            const auto& expectedSkinWeights = DecodedDNA::skinWeightsValues[index][meshIndex][vertexIndex];
             ASSERT_EQ(skinWeights, ConstArrayView<float>{expectedSkinWeights});
 
             const auto jointIndices = reader->getSkinWeightsJointIndices(meshIndex, vertexIndex);
-            const auto& expectedJointIndices = decoded::skinWeightsJointIndices[params.maxLOD][meshIndex][vertexIndex];
+            const auto& expectedJointIndices = DecodedDNA::skinWeightsJointIndices[index][meshIndex][vertexIndex];
             ASSERT_EQ(jointIndices, ConstArrayView<std::uint16_t>{expectedJointIndices});
         }
 
         const auto blendShapeCount = reader->getBlendShapeTargetCount(meshIndex);
-        ASSERT_EQ(blendShapeCount, decoded::correctiveBlendShapeDeltas[params.maxLOD][meshIndex].size());
+        ASSERT_EQ(blendShapeCount, DecodedDNA::correctiveBlendShapeDeltas[index][meshIndex].size());
         for (std::uint16_t blendShapeTargetIndex = 0u; blendShapeTargetIndex < blendShapeCount; ++blendShapeTargetIndex) {
             const auto channelIndex = reader->getBlendShapeChannelIndex(meshIndex, blendShapeTargetIndex);
-            ASSERT_EQ(channelIndex, decoded::correctiveBlendShapeIndices[params.maxLOD][meshIndex][blendShapeTargetIndex]);
+            ASSERT_EQ(channelIndex, DecodedDNA::correctiveBlendShapeIndices[index][meshIndex][blendShapeTargetIndex]);
 
             const auto deltaCount = reader->getBlendShapeTargetDeltaCount(meshIndex, blendShapeTargetIndex);
             ASSERT_EQ(deltaCount,
-                      decoded::correctiveBlendShapeDeltas[params.maxLOD][meshIndex][blendShapeTargetIndex].size());
+                      DecodedDNA::correctiveBlendShapeDeltas[index][meshIndex][blendShapeTargetIndex].size());
 
             for (std::uint32_t deltaIndex = 0u; deltaIndex < deltaCount; ++deltaIndex) {
                 const auto& delta = reader->getBlendShapeTargetDelta(meshIndex, blendShapeTargetIndex, deltaIndex);
                 const auto& expectedDelta =
-                    decoded::correctiveBlendShapeDeltas[params.maxLOD][meshIndex][blendShapeTargetIndex][deltaIndex];
+                    DecodedDNA::correctiveBlendShapeDeltas[index][meshIndex][blendShapeTargetIndex][deltaIndex];
                 ASSERT_EQ(delta, expectedDelta);
             }
 
             const auto vertexIndices = reader->getBlendShapeTargetVertexIndices(meshIndex, blendShapeTargetIndex);
             const auto& expectedVertexIndices =
-                decoded::correctiveBlendShapeVertexIndices[params.maxLOD][meshIndex][blendShapeTargetIndex];
+                DecodedDNA::correctiveBlendShapeVertexIndices[index][meshIndex][blendShapeTargetIndex];
             ASSERT_EQ(vertexIndices, ConstArrayView<std::uint32_t>{expectedVertexIndices});
         }
     }
 }
 
-static void assertReaderHasAllData(Reader* reader, const LODParameters& params) {
-    verifyDescriptor(reader, params);
-    verifyDefinition(reader, params);
-    verifyBehavior(reader, params);
-    verifyGeometry(reader, params);
+template<class TAPICopyParameters>
+static void verifyMachineLearnedBehavior(MachineLearnedBehaviorReader* reader) {
+    using DecodedDNA = typename TAPICopyParameters::DecodedData;
+    const auto index = DecodedDNA::lodConstraintToIndex(TAPICopyParameters::maxLOD(), TAPICopyParameters::minLOD());
+
+    const auto mlControlCount = reader->getMLControlCount();
+    ASSERT_EQ(mlControlCount, DecodedDNA::mlControlNames.size());
+    for (std::uint16_t i = 0u; i < mlControlCount; ++i) {
+        ASSERT_EQ(reader->getMLControlName(i), StringView{DecodedDNA::mlControlNames[i]});
+    }
+
+    ASSERT_EQ(reader->getNeuralNetworkCount(), DecodedDNA::neuralNetworkLayerCount[index].size());
+
+    const auto& expectedRegionNames = DecodedDNA::regionNames[index];
+    ASSERT_EQ(reader->getMeshCount(), expectedRegionNames.size());
+    for (std::uint16_t mi = {}; mi < reader->getMeshCount(); ++mi) {
+        ASSERT_EQ(reader->getMeshRegionCount(mi), expectedRegionNames[mi].size());
+        for (std::uint16_t ri = {}; ri < expectedRegionNames[mi].size(); ++ri) {
+            ASSERT_EQ(reader->getMeshRegionName(mi, ri), StringView{expectedRegionNames[mi][ri]});
+        }
+    }
+
+    const auto& expectedNetIndices = DecodedDNA::neuralNetworkIndicesPerMeshRegion[index];
+    ASSERT_EQ(reader->getMeshCount(), expectedNetIndices.size());
+    for (std::uint16_t meshIdx = {}; meshIdx < expectedNetIndices.size(); ++meshIdx) {
+        ASSERT_EQ(reader->getMeshRegionCount(meshIdx), expectedNetIndices[meshIdx].size());
+        for (std::uint16_t regionIdx = {}; regionIdx < expectedNetIndices[meshIdx].size(); ++regionIdx) {
+            const auto indices = reader->getNeuralNetworkIndicesForMeshRegion(meshIdx, regionIdx);
+            ASSERT_EQ(indices.size(), expectedNetIndices[meshIdx][regionIdx].size());
+            ASSERT_ELEMENTS_EQ(indices, expectedNetIndices[meshIdx][regionIdx], expectedNetIndices[meshIdx][regionIdx].size());
+        }
+    }
+
+    for (std::uint16_t neuralNetIdx = {}; neuralNetIdx < reader->getNeuralNetworkCount(); ++neuralNetIdx) {
+        ASSERT_EQ(reader->getNeuralNetworkInputIndices(neuralNetIdx), DecodedDNA::neuralNetworkInputIndices[index][neuralNetIdx]);
+        ASSERT_EQ(reader->getNeuralNetworkOutputIndices(neuralNetIdx),
+                  DecodedDNA::neuralNetworkOutputIndices[index][neuralNetIdx]);
+        ASSERT_EQ(reader->getNeuralNetworkLayerCount(neuralNetIdx), DecodedDNA::neuralNetworkLayerCount[index][neuralNetIdx]);
+        for (std::uint16_t layerIdx = {}; layerIdx < reader->getNeuralNetworkLayerCount(neuralNetIdx); ++layerIdx) {
+            const auto expected =
+                static_cast<dna::ActivationFunction>(DecodedDNA::neuralNetworkActivationFunction[index][neuralNetIdx][layerIdx]);
+            ASSERT_EQ(reader->getNeuralNetworkLayerActivationFunction(neuralNetIdx, layerIdx), expected);
+            ASSERT_EQ(reader->getNeuralNetworkLayerActivationFunctionParameters(neuralNetIdx, layerIdx),
+                      DecodedDNA::neuralNetworkActivationFunctionParameters[index][neuralNetIdx][layerIdx]);
+            ASSERT_EQ(reader->getNeuralNetworkLayerBiases(neuralNetIdx, layerIdx),
+                      DecodedDNA::neuralNetworkBiases[index][neuralNetIdx][layerIdx]);
+            ASSERT_EQ(reader->getNeuralNetworkLayerWeights(neuralNetIdx, layerIdx),
+                      DecodedDNA::neuralNetworkWeights[index][neuralNetIdx][layerIdx]);
+        }
+    }
 }
 
-TEST(StreamReadWriteIntegrationTest, JSONVerifyAllDNADataAfterSetFromIsUsed) {
-    pma::AlignedMemoryResource memRes;
+template<class TAPICopyParameters>
+struct ReaderDataVerifier {
+
+    static void assertHasAllData(Reader* reader) {
+        verifyDescriptor<TAPICopyParameters>(reader);
+        verifyDefinition<TAPICopyParameters>(reader);
+        verifyBehavior<TAPICopyParameters>(reader);
+        verifyGeometry<TAPICopyParameters>(reader);
+    }
+
+};
+
+
+template<class Reader, class Writer, std::uint16_t MaxLOD, std::uint16_t MinLOD, std::uint16_t CurrentLOD>
+struct ReaderDataVerifier<APICopyParameters<Reader, Writer, RawV23, DecodedV23, MaxLOD, MinLOD, CurrentLOD> > {
+
+    static void assertHasAllData(Reader* reader) {
+        using TAPICopyParameters = APICopyParameters<Reader, Writer, RawV23, DecodedV23, MaxLOD, MinLOD, CurrentLOD>;
+        verifyDescriptor<TAPICopyParameters>(reader);
+        verifyDefinition<TAPICopyParameters>(reader);
+        verifyBehavior<TAPICopyParameters>(reader);
+        verifyGeometry<TAPICopyParameters>(reader);
+        verifyMachineLearnedBehavior<TAPICopyParameters>(reader);
+    }
+
+};
+
+using TAPICopyTestParameters = ::testing::Types<
+    APICopyParameters<dna::BinaryStreamReader, dna::BinaryStreamWriter, RawV21, DecodedV21, 0u, 1u, 0u>,
+    APICopyParameters<dna::BinaryStreamReader, dna::BinaryStreamWriter, RawV21, DecodedV21, 0u, 1u, 1u>,
+    APICopyParameters<dna::BinaryStreamReader, dna::BinaryStreamWriter, RawV21, DecodedV21, 0u, 0u, 0u>,
+    APICopyParameters<dna::BinaryStreamReader, dna::BinaryStreamWriter, RawV21, DecodedV21, 1u, 1u, 0u>,
+    APICopyParameters<dna::BinaryStreamReader, dna::BinaryStreamWriter, RawV22, DecodedV22, 0u, 1u, 0u>,
+    APICopyParameters<dna::BinaryStreamReader, dna::BinaryStreamWriter, RawV22, DecodedV22, 0u, 1u, 1u>,
+    APICopyParameters<dna::BinaryStreamReader, dna::BinaryStreamWriter, RawV22, DecodedV22, 0u, 0u, 0u>,
+    APICopyParameters<dna::BinaryStreamReader, dna::BinaryStreamWriter, RawV22, DecodedV22, 1u, 1u, 0u>,
+    APICopyParameters<dna::BinaryStreamReader, dna::BinaryStreamWriter, RawV23, DecodedV23, 0u, 1u, 0u>,
+    APICopyParameters<dna::BinaryStreamReader, dna::BinaryStreamWriter, RawV23, DecodedV23, 0u, 1u, 1u>,
+    APICopyParameters<dna::BinaryStreamReader, dna::BinaryStreamWriter, RawV23, DecodedV23, 0u, 0u, 0u>,
+    APICopyParameters<dna::BinaryStreamReader, dna::BinaryStreamWriter, RawV23, DecodedV23, 1u, 1u, 0u>,
+    APICopyParameters<dna::JSONStreamReader, dna::JSONStreamWriter, RawV21, DecodedV21, 0u, 1u, 0u>,
+    APICopyParameters<dna::JSONStreamReader, dna::JSONStreamWriter, RawV22, DecodedV22, 0u, 1u, 0u>,
+    APICopyParameters<dna::JSONStreamReader, dna::JSONStreamWriter, RawV23, DecodedV23, 0u, 1u, 0u>
+    >;
+TYPED_TEST_SUITE(StreamReadWriteAPICopyIntegrationTest, TAPICopyTestParameters, );
+
+template<typename TReader>
+struct ReaderFactory;
+
+template<>
+struct ReaderFactory<dna::BinaryStreamReader> {
+    static pma::ScopedPtr<dna::BinaryStreamReader> create(trio::BoundedIOStream* stream,
+                                                          dna::DataLayer layer,
+                                                          dna::UnknownLayerPolicy policy,
+                                                          std::uint16_t maxLOD,
+                                                          std::uint16_t minLOD) {
+        return pma::makeScoped<dna::BinaryStreamReader>(stream, layer, policy, maxLOD, minLOD);
+    }
+
+};
+
+template<>
+struct ReaderFactory<dna::JSONStreamReader> {
+    static pma::ScopedPtr<dna::JSONStreamReader> create(trio::BoundedIOStream* stream,
+                                                        dna::DataLayer  /*unused*/,
+                                                        dna::UnknownLayerPolicy  /*unused*/,
+                                                        std::uint16_t  /*unused*/,
+                                                        std::uint16_t  /*unused*/) {
+        return pma::makeScoped<dna::JSONStreamReader>(stream);
+    }
+
+};
 
-    const auto bytes = raw::getBytes();
+TYPED_TEST(StreamReadWriteAPICopyIntegrationTest, VerifyAllDNADataAfterSetFromThroughAPI) {
+    using CurrentParameters = typename TestFixture::Parameters;
+
+    const auto bytes = CurrentParameters::RawBytes::getBytes();
     auto source = pma::makeScoped<trio::MemoryStream>();
     source->write(reinterpret_cast<const char*>(bytes.data()), bytes.size());
     source->seek(0);
 
-    auto sourceReader = pma::makeScoped<BinaryStreamReader>(source.get(), DataLayer::All, static_cast<std::uint16_t>(0), &memRes);
+    auto sourceReader = pma::makeScoped<BinaryStreamReader>(source.get(),
+                                                            DataLayer::All,
+                                                            UnknownLayerPolicy::Preserve,
+                                                            static_cast<std::uint16_t>(0));
     sourceReader->read();
 
     auto clone = pma::makeScoped<trio::MemoryStream>();
-    auto cloneWriter = pma::makeScoped<JSONStreamWriter>(clone.get(), 4u, &memRes);
-    cloneWriter->setFrom(sourceReader.get());
+    auto cloneWriter = pma::makeScoped<typename CurrentParameters::Writer>(clone.get());
+    // Due to the abstract Reader type, the API copy method will be invoked
+    cloneWriter->setFrom(static_cast<Reader*>(sourceReader.get()));
     cloneWriter->write();
 
     clone->seek(0ul);
-    auto cloneReader = pma::makeScoped<JSONStreamReader>(clone.get(), &memRes);
+    using Factory = ReaderFactory<typename CurrentParameters::Reader>;
+    auto cloneReader = Factory::create(clone.get(),
+                                       DataLayer::All,
+                                       UnknownLayerPolicy::Preserve,
+                                       CurrentParameters::maxLOD(),
+                                       CurrentParameters::minLOD());
     cloneReader->read();
 
-    assertReaderHasAllData(cloneReader.get(), {0u, 0u});
+    ReaderDataVerifier<CurrentParameters>::assertHasAllData(cloneReader.get());
 }
 
-TEST_P(StreamReadWriteIntegrationTest, BinaryVerifyAllDNADataAfterSetFromIsUsed) {
-    pma::AlignedMemoryResource memRes;
-
-    const auto bytes = raw::getBytes();
+using TRawCopyTestParameters = ::testing::Types<
+    // Copy tests
+    RawCopyParameters<RawV21, RawV21, UnknownLayerPolicy::Preserve, 2, 1>,
+    RawCopyParameters<RawV21, RawV21, UnknownLayerPolicy::Ignore, 2, 1>,
+    RawCopyParameters<RawV22, RawV22, UnknownLayerPolicy::Preserve, 2, 2>,
+    RawCopyParameters<RawV22, RawV22WithUnknownDataIgnoredAndDNARewritten, UnknownLayerPolicy::Ignore, 2, 2>,
+    RawCopyParameters<RawV2xNewer, RawV2xNewerWithUnknownDataPreservedAndDNARewritten, UnknownLayerPolicy::Preserve, 2,
+                      static_cast<std::uint16_t>(-1)>,
+    RawCopyParameters<RawV2xNewer, RawV2xNewerWithUnknownDataIgnoredAndDNARewritten, UnknownLayerPolicy::Ignore, 2,
+                      static_cast<std::uint16_t>(-1)>,
+    RawCopyParameters<RawV23, RawV23, UnknownLayerPolicy::Preserve, 2, 3>,
+    RawCopyParameters<RawV23, RawV23, UnknownLayerPolicy::Ignore, 2, 3>,
+    // File format conversion tests
+    RawCopyParameters<RawV21, RawV22WithUnknownDataIgnoredAndDNARewritten, UnknownLayerPolicy::Preserve, 2, 2>,
+    RawCopyParameters<RawV21, RawV22WithUnknownDataIgnoredAndDNARewritten, UnknownLayerPolicy::Ignore, 2, 2>,
+    RawCopyParameters<RawV22, RawV21, UnknownLayerPolicy::Preserve, 2, 1>,
+    RawCopyParameters<RawV22, RawV21, UnknownLayerPolicy::Ignore, 2, 1>,
+    RawCopyParameters<RawV2xNewer, RawV22WithUnknownDataFromNewer2x, UnknownLayerPolicy::Preserve, 2, 2>,
+    RawCopyParameters<RawV2xNewer, RawV22Empty, UnknownLayerPolicy::Ignore, 2, 2>,
+    RawCopyParameters<RawV23, RawV22DowngradedFromV23, UnknownLayerPolicy::Preserve, 2, 2>,
+    RawCopyParameters<RawV23, RawV22WithUnknownDataIgnoredAndDNARewritten, UnknownLayerPolicy::Ignore, 2, 2>
+    >;
+TYPED_TEST_SUITE(StreamReadWriteRawCopyIntegrationTest, TRawCopyTestParameters, );
+
+TYPED_TEST(StreamReadWriteRawCopyIntegrationTest, VerifySetFromCopiesEvenUnknownData) {
+    using CurrentParameters = typename TestFixture::Parameters;
+
+    const auto bytes = CurrentParameters::RawBytes::getBytes();
     auto source = pma::makeScoped<trio::MemoryStream>();
     source->write(reinterpret_cast<const char*>(bytes.data()), bytes.size());
     source->seek(0);
 
-    auto sourceReader = pma::makeScoped<BinaryStreamReader>(source.get(), DataLayer::All, static_cast<std::uint16_t>(0), &memRes);
+    auto sourceReader = pma::makeScoped<BinaryStreamReader>(source.get(),
+                                                            DataLayer::All,
+                                                            CurrentParameters::policy(),
+                                                            static_cast<std::uint16_t>(0));
     sourceReader->read();
 
     auto clone = pma::makeScoped<trio::MemoryStream>();
-    auto cloneWriter = pma::makeScoped<BinaryStreamWriter>(clone.get(), &memRes);
-    cloneWriter->setFrom(sourceReader.get());
+    auto cloneWriter = pma::makeScoped<BinaryStreamWriter>(clone.get());
+    cloneWriter->setFrom(sourceReader.get(), DataLayer::All, CurrentParameters::policy());
+    cloneWriter->setFileFormatGeneration(CurrentParameters::generation());
+    cloneWriter->setFileFormatVersion(CurrentParameters::version());
     cloneWriter->write();
 
     clone->seek(0ul);
-    auto cloneReader = pma::makeScoped<BinaryStreamReader>(clone.get(), DataLayer::All, params.maxLOD, &memRes);
-    cloneReader->read();
 
-    assertReaderHasAllData(cloneReader.get(), params);
+    #if !defined(__clang__) && defined(__GNUC__)
+        #pragma GCC diagnostic push
+        #pragma GCC diagnostic ignored "-Wuseless-cast"
+    #endif
+    const auto cloneSize = static_cast<std::size_t>(clone->size());
+    #if !defined(__clang__) && defined(__GNUC__)
+        #pragma GCC diagnostic pop
+    #endif
+    std::vector<char> copiedBytes(cloneSize);
+    clone->read(copiedBytes.data(), cloneSize);
+
+    const auto expectedBytes = CurrentParameters::ExpectedBytes::getBytes();
+    ASSERT_EQ(expectedBytes.size(), copiedBytes.size());
+    ASSERT_EQ(expectedBytes, copiedBytes);
 }
 
-INSTANTIATE_TEST_SUITE_P(StreamReadWriteIntegrationTest, StreamReadWriteIntegrationTest, ::testing::Values(
-                             // maxLOD, currentLOD
-                             LODParameters{0u, 0u},
-                             LODParameters{0u, 1u},
-                             LODParameters{1u, 0u}
-                             ));
+TEST(StreamReadWriteIntegrationTest, ReadWriteJSON) {
+    auto stream = pma::makeScoped<trio::MemoryStream>();
+    auto writer = pma::makeScoped<JSONStreamWriter>(stream.get(), 4u);
+
+    writer->setMeshName(0, "mesh0");
+    const Position vertices[] = {Position{0.0f, 1.0f, 2.0}, Position{3.0f, 4.0f, 5.0}};
+    writer->setVertexPositions(0u, vertices, 2u);
+    writer->write();
+
+    #if !defined(__clang__) && defined(__GNUC__)
+        #pragma GCC diagnostic push
+        #pragma GCC diagnostic ignored "-Wuseless-cast"
+    #endif
+    pma::Vector<char> json(static_cast<std::size_t>(stream->size()));
+    #if !defined(__clang__) && defined(__GNUC__)
+        #pragma GCC diagnostic pop
+    #endif
+
+    pma::String<char> expected = jsonDNA;
+    stream->seek(0ul);
+    stream->read(json.data(), json.size());
+    ASSERT_EQ(json.size(), expected.size());
+    ASSERT_ELEMENTS_EQ(json.data(), expected.data(), expected.size());
+
+    stream->seek(0ul);
+    auto reader = pma::makeScoped<JSONStreamReader>(stream.get());
+    reader->read();
+    ASSERT_TRUE(dna::Status::isOk());
+}
 
 }  // namespace dna
diff --git a/Engine/Plugins/Animation/RigLogic/Source/RigLogicLibTest/Private/dnatests/TestStreamReadWriteIntegration.h b/Engine/Plugins/Animation/RigLogic/Source/RigLogicLibTest/Private/dnatests/TestStreamReadWriteIntegration.h
--- a/Engine/Plugins/Animation/RigLogic/Source/RigLogicLibTest/Private/dnatests/TestStreamReadWriteIntegration.h
+++ b/Engine/Plugins/Animation/RigLogic/Source/RigLogicLibTest/Private/dnatests/TestStreamReadWriteIntegration.h
@@ -4,20 +4,65 @@
 
 #include "dnatests/Defs.h"
 
-struct LODParameters {
-    std::uint16_t maxLOD;
-    std::uint16_t currentLOD;
-};
+#include "dna/DNA.h"
+
+template<typename TReader,
+         typename TWriter,
+         typename TRawBytes,
+         typename TDecodedData,
+         std::uint16_t MaxLOD,
+         std::uint16_t MinLOD,
+         std::uint16_t CurrentLOD>
+struct APICopyParameters {
+    using Reader = TReader;
+    using Writer = TWriter;
+    using RawBytes = TRawBytes;
+    using DecodedData = TDecodedData;
+
+    static constexpr std::uint16_t maxLOD() {
+        return MaxLOD;
+    }
+
+    static constexpr std::uint16_t minLOD() {
+        return MinLOD;
+    }
 
-class StreamReadWriteIntegrationTest : public ::testing::TestWithParam<LODParameters> {
-    public:
-        ~StreamReadWriteIntegrationTest();
+    static constexpr std::uint16_t currentLOD() {
+        return CurrentLOD;
+    }
 
+};
+
+template<typename TAPICopyParameters>
+class StreamReadWriteAPICopyIntegrationTest : public ::testing::Test {
     protected:
-        void SetUp() override {
-            params = GetParam();
-        }
+        using Parameters = TAPICopyParameters;
+
+};
 
+template<typename TRawBytes, typename TExpectedBytes, dna::UnknownLayerPolicy Policy, std::uint16_t SaveAsGeneration,
+         std::uint16_t SaveAsVersion>
+struct RawCopyParameters {
+    using RawBytes = TRawBytes;
+    using ExpectedBytes = TExpectedBytes;
+
+    static constexpr dna::UnknownLayerPolicy policy() {
+        return Policy;
+    }
+
+    static constexpr std::uint16_t generation() {
+        return SaveAsGeneration;
+    }
+
+    static constexpr std::uint16_t version() {
+        return SaveAsVersion;
+    }
+
+};
+
+template<typename TRawCopyParameters>
+class StreamReadWriteRawCopyIntegrationTest : public ::testing::Test {
     protected:
-        LODParameters params;
+        using Parameters = TRawCopyParameters;
+
 };
diff --git a/Engine/Plugins/Animation/RigLogic/Source/RigLogicLibTest/Private/rltests/animatedmaps/AnimatedMapFixtures.cpp b/Engine/Plugins/Animation/RigLogic/Source/RigLogicLibTest/Private/rltests/animatedmaps/AnimatedMapFixtures.cpp
--- a/Engine/Plugins/Animation/RigLogic/Source/RigLogicLibTest/Private/rltests/animatedmaps/AnimatedMapFixtures.cpp
+++ b/Engine/Plugins/Animation/RigLogic/Source/RigLogicLibTest/Private/rltests/animatedmaps/AnimatedMapFixtures.cpp
@@ -3,7 +3,8 @@
 #include "rltests/animatedmaps/AnimatedMapFixtures.h"
 
 #include "riglogic/TypeDefs.h"
-#include "riglogic/animatedmaps/AnimatedMaps.h"
+#include "riglogic/animatedmaps/AnimatedMapsImpl.h"
+#include "riglogic/animatedmaps/AnimatedMapsImplOutputInstance.h"
 #include "riglogic/conditionaltable/ConditionalTable.h"
 
 namespace defaults {
@@ -13,9 +14,18 @@ const std::uint16_t lods[lodCount] = {2};
 
 }  // namespace calculation
 
-rl4::AnimatedMaps AnimatedMapsFactory::withDefaults(rl4::ConditionalTable&& conditionals, rl4::MemoryResource* memRes) {
-    return rl4::AnimatedMaps{
+rl4::AnimatedMapsOutputInstance::Factory AnimatedMapsFactory::getInstanceFactory(std::uint16_t animatedMapCount) {
+    return [animatedMapCount](rl4::MemoryResource* memRes) {
+               return pma::UniqueInstance<rl4::AnimatedMapsImplOutputInstance,
+                                          rl4::AnimatedMapsOutputInstance>::with(memRes).create(animatedMapCount, memRes);
+    };
+}
+
+rl4::AnimatedMapsImpl AnimatedMapsFactory::withDefaults(rl4::ConditionalTable&& conditionals, rl4::MemoryResource* memRes) {
+    const auto animatedMapCount = conditionals.getOutputCount();
+    return rl4::AnimatedMapsImpl{
         rl4::Vector<std::uint16_t>{defaults::lods, defaults::lods + defaults::lodCount, memRes},
-        std::move(conditionals)
+        std::move(conditionals),
+        getInstanceFactory(animatedMapCount)
     };
 }
diff --git a/Engine/Plugins/Animation/RigLogic/Source/RigLogicLibTest/Private/rltests/animatedmaps/AnimatedMapFixtures.h b/Engine/Plugins/Animation/RigLogic/Source/RigLogicLibTest/Private/rltests/animatedmaps/AnimatedMapFixtures.h
--- a/Engine/Plugins/Animation/RigLogic/Source/RigLogicLibTest/Private/rltests/animatedmaps/AnimatedMapFixtures.h
+++ b/Engine/Plugins/Animation/RigLogic/Source/RigLogicLibTest/Private/rltests/animatedmaps/AnimatedMapFixtures.h
@@ -3,13 +3,13 @@
 #pragma once
 
 #include "riglogic/TypeDefs.h"
-#include "riglogic/animatedmaps/AnimatedMaps.h"
+#include "riglogic/animatedmaps/AnimatedMapsImpl.h"
 #include "riglogic/conditionaltable/ConditionalTable.h"
 
 #include <cstddef>
 
 struct AnimatedMapsFactory {
 
-    static rl4::AnimatedMaps withDefaults(rl4::ConditionalTable&& conditionals, rl4::MemoryResource* memRes);
-
+    static rl4::AnimatedMapsImpl withDefaults(rl4::ConditionalTable&& conditionals, rl4::MemoryResource* memRes);
+    static rl4::AnimatedMapsOutputInstance::Factory getInstanceFactory(std::uint16_t animatedMapCount);
 };
diff --git a/Engine/Plugins/Animation/RigLogic/Source/RigLogicLibTest/Private/rltests/animatedmaps/TestAnimatedMaps.cpp b/Engine/Plugins/Animation/RigLogic/Source/RigLogicLibTest/Private/rltests/animatedmaps/TestAnimatedMaps.cpp
--- a/Engine/Plugins/Animation/RigLogic/Source/RigLogicLibTest/Private/rltests/animatedmaps/TestAnimatedMaps.cpp
+++ b/Engine/Plugins/Animation/RigLogic/Source/RigLogicLibTest/Private/rltests/animatedmaps/TestAnimatedMaps.cpp
@@ -2,12 +2,11 @@
 
 #include "rltests/Defs.h"
 #include "rltests/animatedmaps/AnimatedMapFixtures.h"
+#include "rltests/controls/ControlFixtures.h"
 #include "rltests/conditionaltable/ConditionalTableFixtures.h"
 
+#include "riglogic/TypeDefs.h"
 #include "riglogic/animatedmaps/AnimatedMaps.h"
-#include "riglogic/types/Aliases.h"
-
-#include <pma/resources/AlignedMemoryResource.h>
 
 #ifdef _MSC_VER
     #pragma warning(push)
@@ -29,18 +28,25 @@ TEST_P(AnimatedMapsTest, LODLimitsCondTableSize) {
     pma::AlignedMemoryResource amr;
     rl4::Vector<std::uint16_t> lods{2, 1};
     auto conditionals = ConditionalTableFactory::withMultipleIODefaults(&amr);
-    rl4::AnimatedMaps animatedMaps{std::move(lods), std::move(conditionals)};
+    auto inputInstanceFactory = ControlsFactory::getInstanceFactory(0,
+                                                                    static_cast<std::uint16_t>(conditionalTableInputs.size()),
+                                                                    0,
+                                                                    0);
+    auto outputInstanceFactory = AnimatedMapsFactory::getInstanceFactory(conditionals.getOutputCount());
+    rl4::AnimatedMapsImpl animatedMaps{std::move(lods), std::move(conditionals), outputInstanceFactory};
 
     const auto lod = GetParam();
     const float expected[2][2] = {
         {0.3f, 0.6f},  // LOD0
         {0.3f, 0.0f}  // LOD1
     };
-    std::array<float, 2ul> outputs;
-    animatedMaps.calculate(rl4::ConstArrayView<float>{conditionalTableInputs},
-                           rl4::ArrayView<float>{outputs},
-                           lod);
-    ASSERT_ELEMENTS_EQ(outputs, expected[lod], 2ul);
+    auto inputInstance = inputInstanceFactory(&amr);
+    auto inputBuffer = inputInstance->getInputBuffer();
+    std::copy(conditionalTableInputs.begin(), conditionalTableInputs.end(), inputBuffer.begin());
+
+    auto outputInstance = animatedMaps.createInstance(&amr);
+    animatedMaps.calculate(inputInstance.get(), outputInstance.get(), lod);
+    ASSERT_ELEMENTS_EQ(outputInstance->getOutputBuffer(), expected[lod], 2ul);
 }
 
 INSTANTIATE_TEST_SUITE_P(AnimatedMapsTest, AnimatedMapsTest, ::testing::Values(0u, 1u));
diff --git a/Engine/Plugins/Animation/RigLogic/Source/RigLogicLibTest/Private/rltests/blendshapes/BlendShapeFixtures.cpp b/Engine/Plugins/Animation/RigLogic/Source/RigLogicLibTest/Private/rltests/blendshapes/BlendShapeFixtures.cpp
--- a/Engine/Plugins/Animation/RigLogic/Source/RigLogicLibTest/Private/rltests/blendshapes/BlendShapeFixtures.cpp
+++ b/Engine/Plugins/Animation/RigLogic/Source/RigLogicLibTest/Private/rltests/blendshapes/BlendShapeFixtures.cpp
@@ -3,15 +3,24 @@
 #include "rltests/blendshapes/BlendShapeFixtures.h"
 
 #include "riglogic/TypeDefs.h"
-#include "riglogic/blendshapes/BlendShapes.h"
+#include "riglogic/blendshapes/BlendShapesImpl.h"
+#include "riglogic/blendshapes/BlendShapesImplOutputInstance.h"
 
 namespace rltests {
 
-rl4::BlendShapes createTestBlendShapes(rl4::MemoryResource* memRes) {
-    return rl4::BlendShapes{
+rl4::BlendShapesOutputInstance::Factory BlendShapesFactory::getInstanceFactory(std::uint16_t bsCount) {
+    return [bsCount](rl4::MemoryResource* memRes) {
+               return pma::UniqueInstance<rl4::BlendShapesImplOutputInstance,
+                                          rl4::BlendShapesOutputInstance>::with(memRes).create(bsCount, memRes);
+    };
+}
+
+rl4::BlendShapesImpl BlendShapesFactory::createTestBlendShapes(rl4::MemoryResource* memRes) {
+    return rl4::BlendShapesImpl{
         rl4::Vector<std::uint16_t>{std::begin(LODs), std::end(LODs), memRes},
         rl4::Vector<std::uint16_t>{std::begin(blendShapeInputIndices), std::end(blendShapeInputIndices), memRes},
-        rl4::Vector<std::uint16_t>{std::begin(blendShapeOutputIndices), std::end(blendShapeOutputIndices), memRes}
+        rl4::Vector<std::uint16_t>{std::begin(blendShapeOutputIndices), std::end(blendShapeOutputIndices), memRes},
+        getInstanceFactory(blendShapeCount)
     };
 }
 
diff --git a/Engine/Plugins/Animation/RigLogic/Source/RigLogicLibTest/Private/rltests/blendshapes/BlendShapeFixtures.h b/Engine/Plugins/Animation/RigLogic/Source/RigLogicLibTest/Private/rltests/blendshapes/BlendShapeFixtures.h
--- a/Engine/Plugins/Animation/RigLogic/Source/RigLogicLibTest/Private/rltests/blendshapes/BlendShapeFixtures.h
+++ b/Engine/Plugins/Animation/RigLogic/Source/RigLogicLibTest/Private/rltests/blendshapes/BlendShapeFixtures.h
@@ -3,9 +3,7 @@
 #pragma once
 
 #include "riglogic/TypeDefs.h"
-#include "riglogic/blendshapes/BlendShapes.h"
-
-#include <pma/MemoryResource.h>
+#include "riglogic/blendshapes/BlendShapesImpl.h"
 
 #ifdef _MSC_VER
     #pragma warning(push)
@@ -30,10 +28,14 @@ static const std::uint16_t LODs[] = {LOD0, LOD1, LOD2, LOD3};
 static const std::array<float, 15ul> blendShapeInputs = {
     1.0f, 2.0f, 3.0f, 4.0f, 5.0f, 6.0f, 7.0f, 8.0f, 9.0f, 10.0f, 11.0f, 12.0f, 13.0f, 14.0f, 15.0f
 };
-static const std::uint16_t blendShapeInputIndices[] = {0, 1, 3, 4, 6, 7, 8, 10, 12, 14};
-static const std::uint16_t blendShapeOutputIndices[] = {0, 2, 1, 3, 4, 5, 6, 7, 8, 9};
+static const std::uint16_t blendShapeCount = 10u;
+static const std::uint16_t blendShapeInputIndices[blendShapeCount] = {0, 1, 3, 4, 6, 7, 8, 10, 12, 14};
+static const std::uint16_t blendShapeOutputIndices[blendShapeCount] = {0, 2, 1, 3, 4, 5, 6, 7, 8, 9};
 static const std::array<float, 10ul> blendShapeOutputs = {1.0f, 4.0f, 2.0f, 5.0f, 7.0f, 8.0f, 9.0f, 11.0f, 13.0f, 15.0f};
 
-rl4::BlendShapes createTestBlendShapes(rl4::MemoryResource* memRes);
+struct BlendShapesFactory {
+    static rl4::BlendShapesImpl createTestBlendShapes(rl4::MemoryResource* memRes);
+    static rl4::BlendShapesOutputInstance::Factory getInstanceFactory(std::uint16_t animatedMapCount);
+};
 
 }  // namespace rltests
diff --git a/Engine/Plugins/Animation/RigLogic/Source/RigLogicLibTest/Private/rltests/blendshapes/TestBlendShapes.cpp b/Engine/Plugins/Animation/RigLogic/Source/RigLogicLibTest/Private/rltests/blendshapes/TestBlendShapes.cpp
--- a/Engine/Plugins/Animation/RigLogic/Source/RigLogicLibTest/Private/rltests/blendshapes/TestBlendShapes.cpp
+++ b/Engine/Plugins/Animation/RigLogic/Source/RigLogicLibTest/Private/rltests/blendshapes/TestBlendShapes.cpp
@@ -2,13 +2,12 @@
 
 #include "rltests/Defs.h"
 #include "rltests/blendshapes/BlendShapeFixtures.h"
+#include "rltests/controls/ControlFixtures.h"
 
+#include "riglogic/TypeDefs.h"
 #include "riglogic/blendshapes/BlendShapes.h"
-#include "riglogic/types/Aliases.h"
 #include "riglogic/utils/Extd.h"
 
-#include <pma/resources/AlignedMemoryResource.h>
-
 #ifdef _MSC_VER
     #pragma warning(push)
     #pragma warning(disable : 4365 4987)
@@ -32,19 +31,24 @@ TEST_P(BlendShapesTest, Calculate) {
     const auto lod = GetParam();
 
     pma::AlignedMemoryResource amr;
-    auto blendShapes = createTestBlendShapes(&amr);
 
-    std::vector<float> outputs(blendShapeOutputs.size());
-    blendShapes.calculate(rl4::ConstArrayView<float>{blendShapeInputs},
-                          rl4::ArrayView<float>{outputs},
-                          lod);
+    auto inputInstanceFactory = ControlsFactory::getInstanceFactory(0, static_cast<std::uint16_t>(blendShapeInputs.size()), 0, 0);
+    auto inputInstance = inputInstanceFactory(&amr);
+    auto inputBuffer = inputInstance->getInputBuffer();
+    std::copy(blendShapeInputs.begin(), blendShapeInputs.end(), inputBuffer.begin());
+
+    auto blendShapes = BlendShapesFactory::createTestBlendShapes(&amr);
+    auto outputInstance = blendShapes.createInstance(&amr);
+
+    blendShapes.calculate(inputInstance.get(), outputInstance.get(), lod);
 
     std::vector<float> expectedData{blendShapeOutputs.begin(), blendShapeOutputs.end()};
     std::fill(extd::advanced(expectedData.begin(), LODs[lod]), expectedData.end(), 0.0f);
     rl4::ConstArrayView<float> expected{expectedData};
 
+    auto outputBuffer = outputInstance->getOutputBuffer();
     for (std::size_t i = 0ul; i < blendShapeOutputs.size(); ++i) {
-        ASSERT_EQ(outputs[i], expected[i]);
+        ASSERT_EQ(outputBuffer[i], expected[i]);
     }
 }
 
diff --git a/Engine/Plugins/Animation/RigLogic/Source/RigLogicLibTest/Private/rltests/conditionaltable/ConditionalTableFixtures.cpp b/Engine/Plugins/Animation/RigLogic/Source/RigLogicLibTest/Private/rltests/conditionaltable/ConditionalTableFixtures.cpp
--- a/Engine/Plugins/Animation/RigLogic/Source/RigLogicLibTest/Private/rltests/conditionaltable/ConditionalTableFixtures.cpp
+++ b/Engine/Plugins/Animation/RigLogic/Source/RigLogicLibTest/Private/rltests/conditionaltable/ConditionalTableFixtures.cpp
@@ -33,6 +33,48 @@ const std::uint16_t distinctOutputs = 2ul;
 
 }  // namespace multiple
 
+namespace interleaved {
+
+/*
+// Conditional table                                      Recommended input test ranges
+//
+I: 0 O: 0  F: 0    T: 1    S: 1        C: 0              -1.0  0.0  0.5  1.0  1.5
+I: 1 O: 1  F: 0    T: 1    S: 1        C: 0              -1.0  0.0  0.5  1.0  1.5
+I: 2 O: 4  F: -1   T: 0    S: -1       C: 0              -1.5 -1.0 -0.5  0.0  0.5  1.0  1.5
+I: 3 O: 5  F: -1   T: 0    S: -1       C: 0              -1.5 -1.0 -0.5  0.0  0.5  1.0  1.5
+I: 2 O: 2  F: 0    T: 1    S: 1        C: 0
+I: 3 O: 3  F: 0    T: 1    S: 1        C: 0
+I: 4 O: 10 F: 0    T: 1    S: 1        C: 0              -1.5 -1.0 -0.5  0.0  0.5  1.0  1.5
+I: 5 O: 11 F: 0    T: 1    S: 1        C: 0              -1.5 -1.0 -0.5  0.0  0.5  1.0  1.5
+I: 6 O: 12 F: 0    T: 1    S: 1        C: 0              -1.5 -1.0 -0.5  0.0  0.5  1.0  1.5
+I: 7 O: 13 F: 0    T: 1    S: 1        C: 0              -1.5 -1.0 -0.5  0.0  0.5  1.0  1.5
+I: 4 O: 6  F: -1   T: 0    S: -1       C: 0
+I: 5 O: 7  F: -1   T: 0    S: -1       C: 0
+I: 6 O: 8  F: -1   T: 0    S: -1       C: 0
+I: 7 O: 9  F: -1   T: 0    S: -1       C: 0
+I: 8 O: 14 F: 0    T: 0.33 S: 3.0303   C: 0              -0.5  0.0  0.15 0.33 0.5  0.66 0.75 1.0 1.5
+I: 8 O: 14 F: 0.33 T: 0.66 S: -3.0303  C: 2
+I: 8 O: 15 F: 0.33 T: 0.66 S: 3.0303   C: -1
+I: 8 O: 15 F: 0.66 T: 1    S: -2.94118 C: 2.94118
+I: 8 O: 16 F: 0.66 T: 1    S: 2.94118  C: -1.94118
+*/
+const std::uint16_t inputIndices[] = {0, 1, 2, 3, 2, 3, 4, 5, 6, 7, 4, 5, 6, 7, 8, 8, 8, 8, 8};
+const std::uint16_t outputIndices[] = {0, 1, 4, 5, 2, 3, 10, 11, 12, 13, 6, 7, 8, 9, 14, 14, 15, 15, 16};
+const float fromValues[] =
+{0.0f, 0.0f, -1.0f, -1.0f, 0.0f, 0.0f, 0.0f, 0.0f, 0.0f, 0.0f, -1.0f, -1.0f, -1.0f, -1.0f, 0.0f, 0.33f, 0.33f, 0.66f, 0.66f};
+const float toValues[] =
+{1.0f, 1.0f, 0.0f, 0.0f, 1.0f, 1.0f, 1.0f, 1.0f, 1.0f, 1.0f, 0.0f, 0.0f, 0.0f, 0.0f, 0.33f, 0.66f, 0.66f, 1.0f, 1.0f};
+const float slopeValues[] =
+{1.0f, 1.0f, -1.0f, -1.0f, 1.0f, 1.0f, 1.0f, 1.0f, 1.0f, 1.0f, -1.0f, -1.0f, -1.0f, -1.0f, 3.0303f, -3.0303f, 3.0303f,
+ -2.94118f, 2.94118f};
+const float cutValues[] =
+{0.0f, 0.0f, 0.0f, 0.0f, 0.0f, 0.0f, 0.0f, 0.0f, 0.0f, 0.0f, 0.0f, 0.0f, 0.0f, 0.0f, 0.0f, 2.0f, -1.0f, 2.94118f, -1.94118f};
+const std::uint16_t size = 19ul;
+const std::uint16_t distinctInputs = 9ul;
+const std::uint16_t distinctOutputs = 17ul;
+
+}  // namespace interleaved
+
 rl4::ConditionalTable ConditionalTableFactory::withSingleIO(rl4::Vector<float>&& fromValues,
                                                             rl4::Vector<float>&& toValues,
                                                             rl4::Vector<float>&& cutValues,
@@ -82,3 +124,17 @@ rl4::ConditionalTable ConditionalTableFactory::withMultipleIODefaults(rl4::Memor
         memRes
         );
 }
+
+rl4::ConditionalTable ConditionalTableFactory::withInterleavedIO(rl4::MemoryResource* memRes) {
+    return rl4::ConditionalTable{
+        rl4::Vector<std::uint16_t>{interleaved::inputIndices, interleaved::inputIndices + interleaved::size, memRes},
+        rl4::Vector<std::uint16_t>{interleaved::outputIndices, interleaved::outputIndices + interleaved::size, memRes},
+        rl4::Vector<float>{interleaved::fromValues, interleaved::fromValues + interleaved::size, memRes},
+        rl4::Vector<float>{interleaved::toValues, interleaved::toValues + interleaved::size, memRes},
+        rl4::Vector<float>{interleaved::slopeValues, interleaved::slopeValues + interleaved::size, memRes},
+        rl4::Vector<float>{interleaved::cutValues, interleaved::cutValues + interleaved::size, memRes},
+        interleaved::distinctInputs,
+        interleaved::distinctOutputs,
+        memRes
+    };
+}
diff --git a/Engine/Plugins/Animation/RigLogic/Source/RigLogicLibTest/Private/rltests/conditionaltable/ConditionalTableFixtures.h b/Engine/Plugins/Animation/RigLogic/Source/RigLogicLibTest/Private/rltests/conditionaltable/ConditionalTableFixtures.h
--- a/Engine/Plugins/Animation/RigLogic/Source/RigLogicLibTest/Private/rltests/conditionaltable/ConditionalTableFixtures.h
+++ b/Engine/Plugins/Animation/RigLogic/Source/RigLogicLibTest/Private/rltests/conditionaltable/ConditionalTableFixtures.h
@@ -5,8 +5,6 @@
 #include "riglogic/TypeDefs.h"
 #include "riglogic/conditionaltable/ConditionalTable.h"
 
-#include <pma/MemoryResource.h>
-
 #ifdef _MSC_VER
     #pragma warning(push)
     #pragma warning(disable : 4365 4987)
@@ -32,5 +30,6 @@ struct ConditionalTableFactory {
                                                 rl4::MemoryResource* memRes);
     static rl4::ConditionalTable withSingleIODefaults(rl4::MemoryResource* memRes);
     static rl4::ConditionalTable withMultipleIODefaults(rl4::MemoryResource* memRes);
+    static rl4::ConditionalTable withInterleavedIO(rl4::MemoryResource* memRes);
 
 };
diff --git a/Engine/Plugins/Animation/RigLogic/Source/RigLogicLibTest/Private/rltests/conditionaltable/TestConditionalTable.cpp b/Engine/Plugins/Animation/RigLogic/Source/RigLogicLibTest/Private/rltests/conditionaltable/TestConditionalTable.cpp
--- a/Engine/Plugins/Animation/RigLogic/Source/RigLogicLibTest/Private/rltests/conditionaltable/TestConditionalTable.cpp
+++ b/Engine/Plugins/Animation/RigLogic/Source/RigLogicLibTest/Private/rltests/conditionaltable/TestConditionalTable.cpp
@@ -3,19 +3,20 @@
 #include "rltests/Defs.h"
 #include "rltests/conditionaltable/ConditionalTableFixtures.h"
 
+#include "riglogic/TypeDefs.h"
 #include "riglogic/conditionaltable/ConditionalTable.h"
 
-#include <pma/resources/AlignedMemoryResource.h>
-
 namespace {
 
-struct CalcTestData {
-    float fromValues[2];
-    float toValues[2];
-    float cutValues[2];
-    float inValues[1];
-    float expected[1];
+template<std::size_t EntryCount, std::size_t InputCount, std::size_t OutputCount>
+struct TCalcTestData {
+    float fromValues[EntryCount];
+    float toValues[EntryCount];
+    float cutValues[EntryCount];
+    float inValues[InputCount];
+    float expected[OutputCount];
 };
+using CalcTestData = TCalcTestData<2, 1, 1>;
 
 class ConditionalTableTest : public ::testing::TestWithParam<CalcTestData> {
 };
@@ -33,7 +34,7 @@ TEST_P(ConditionalTableTest, CheckCalculationBorderCases) {
         &amr
         );
     float outputs[] = {0.0f};
-    conditionals.calculate(testData.inValues, outputs);
+    conditionals.calculateForward(testData.inValues, outputs);
     ASSERT_ELEMENTS_NEAR(outputs, testData.expected, 1, 0.00001f);
 }
 
@@ -77,7 +78,7 @@ TEST(ConditionalTableTest, OutputClamped) {
                                        inputCount,
                                        outputCount,
                                        &amr};
-    conditionals.calculate(conditionalTableInputs.data(), outputs);
+    conditionals.calculateForward(conditionalTableInputs.data(), outputs);
     const float expected[] = {1.0f};
     ASSERT_ELEMENTS_EQ(outputs, expected, 1ul);
 }
@@ -102,7 +103,7 @@ TEST(ConditionalTableTest, OutputIsAccumulated) {
                                        inputCount,
                                        outputCount,
                                        &amr};
-    conditionals.calculate(conditionalTableInputs.data(), outputs);
+    conditionals.calculateForward(conditionalTableInputs.data(), outputs);
     const float expected[] = {0.9f};
     ASSERT_ELEMENTS_NEAR(outputs, expected, 1ul, 0.00001f);
 }
@@ -114,9 +115,98 @@ TEST(ConditionalTableTest, OutputIsResetOnEachCalculation) {
     float outputs[1ul] = {};
     const float expected[] = {0.2f};
 
-    conditionals.calculate(conditionalTableInputs.data(), outputs);
+    conditionals.calculateForward(conditionalTableInputs.data(), outputs);
     ASSERT_ELEMENTS_NEAR(outputs, expected, 1ul, 0.00001f);
 
-    conditionals.calculate(conditionalTableInputs.data(), outputs);
+    conditionals.calculateForward(conditionalTableInputs.data(), outputs);
     ASSERT_ELEMENTS_NEAR(outputs, expected, 1ul, 0.00001f);
 }
+
+namespace {
+
+template<std::size_t InputCount, std::size_t OutputCount>
+struct TIOTestData {
+
+    static constexpr std::size_t inputCount() {
+        return InputCount;
+    }
+
+    static constexpr std::size_t outputCount() {
+        return OutputCount;
+    }
+
+    float inputs[InputCount];
+    float outputs[OutputCount];
+    float reverseOutputs[InputCount];
+};
+using IOTestData = TIOTestData<9, 17>;
+
+class IRLConditionalTableTest : public ::testing::TestWithParam<IOTestData> {
+};
+
+}  // namespace
+
+TEST_P(IRLConditionalTableTest, InterleavedInputsAndOutputs) {
+    pma::AlignedMemoryResource amr;
+    auto conditionals = ConditionalTableFactory::withInterleavedIO(&amr);
+
+    auto testData = GetParam();
+    std::vector<float> inputs{testData.inputs, testData.inputs + testData.inputCount()};
+    float outputs[testData.outputCount()] = {};
+
+    conditionals.calculateForward(inputs.data(), outputs);
+    ASSERT_ELEMENTS_NEAR(outputs, testData.outputs, testData.outputCount(), 0.00001f);
+
+    conditionals.calculateReverse(inputs.data(), outputs);
+    ASSERT_ELEMENTS_NEAR(inputs, testData.reverseOutputs, testData.inputCount(), 0.00001f);
+}
+
+INSTANTIATE_TEST_SUITE_P(IRLConditionalTableTest,
+                         IRLConditionalTableTest,
+                         ::testing::Values(
+                             IOTestData{
+    {-1.0f, -1.0f, -1.5f, -1.5f, -1.5f, -1.5f, -1.5f, -1.5f, -0.5f},
+    {0.0f, 0.0f, 0.0f, 0.0f, 0.0f, 0.0f, 0.0f, 0.0f, 0.0f, 0.0f, 0.0f, 0.0f, 0.0f, 0.0f, 0.0f, 0.0f, 0.0f},
+    {0.0f, 0.0f, 0.0f, 0.0f, 0.0f, 0.0f, 0.0f, 0.0f, 0.0f}
+},
+                             IOTestData{
+    {0.0f, 0.0f, -1.0f, -1.0f, -1.0f, -1.0f, -1.0f, -1.0f, 0.0f},
+    {0.0f, 0.0f, 0.0f, 0.0f, 1.0f, 1.0f, 1.0f, 1.0f, 1.0f, 1.0f, 0.0f, 0.0f, 0.0f, 0.0f, 0.0f, 0.0f, 0.0f},
+    {0.0f, 0.0f, -1.0f, -1.0f, -1.0f, -1.0f, -1.0f, -1.0f, 0.0f}
+},
+                             IOTestData{
+    {0.5f, 0.5f, -0.5f, -0.5f, -0.5f, -0.5f, -0.5f, -0.5f, 0.15f},
+    {0.5f, 0.5f, 0.0f, 0.0f, 0.5f, 0.5f, 0.5f, 0.5f, 0.5f, 0.5f, 0.0f, 0.0f, 0.0f, 0.0f, 0.454545f, 0.0f, 0.0f},
+    {0.5f, 0.5f, -0.5f, -0.5f, -0.5f, -0.5f, -0.5f, -0.5f, 0.15f}
+},
+                             IOTestData{
+    {1.0f, 1.0f, 0.0f, 0.0f, 0.0f, 0.0f, 0.0f, 0.0f, 0.33f},
+    {1.0f, 1.0f, 0.0f, 0.0f, 0.0f, 0.0f, 0.0f, 0.0f, 0.0f, 0.0f, 0.0f, 0.0f, 0.0f, 0.0f, 0.999999f, 0.0f, 0.0f},
+    {1.0f, 1.0f, 0.0f, 0.0f, 0.0f, 0.0f, 0.0f, 0.0f, 0.33f}
+},
+                             IOTestData{
+    {1.5f, 1.5f, 0.5f, 0.5f, 0.5f, 0.5f, 0.5f, 0.5f, 0.5f},
+    {0.0f, 0.0f, 0.5f, 0.5f, 0.0f, 0.0f, 0.0f, 0.0f, 0.0f, 0.0f, 0.5f, 0.5f, 0.5f, 0.5f, 0.48485f, 0.51515f, 0.0f},
+    {0.0f, 0.0f, 0.5f, 0.5f, 0.5f, 0.5f, 0.5f, 0.5f, 0.5f}
+},
+                             IOTestData{
+    {0.0f, 0.0f, 1.0f, 1.0f, 1.0f, 1.0f, 1.0f, 1.0f, 0.66f},
+    {0.0f, 0.0f, 1.0f, 1.0f, 0.0f, 0.0f, 0.0f, 0.0f, 0.0f, 0.0f, 1.0f, 1.0f, 1.0f, 1.0f, 0.000002f, 0.999998f, -0.0000012f},
+    {0.0f, 0.0f, 1.0f, 1.0f, 1.0f, 1.0f, 1.0f, 1.0f, 0.66f}
+},
+                             IOTestData{
+    {0.0f, 0.0f, 1.5f, 1.5f, 1.5f, 1.5f, 1.5f, 1.5f, 0.75f},
+    {0.0f, 0.0f, 0.0f, 0.0f, 0.0f, 0.0f, 0.0f, 0.0f, 0.0f, 0.0f, 0.0f, 0.0f, 0.0f, 0.0f, 0.0f, 0.735295f, 0.264705f},
+    {0.0f, 0.0f, 0.0f, 0.0f, 0.0f, 0.0f, 0.0f, 0.0f, 0.75f}
+},
+                             IOTestData{
+    {0.0f, 0.0f, 0.0f, 0.0f, 0.0f, 0.0f, 0.0f, 0.0f, 1.0f},
+    {0.0f, 0.0f, 0.0f, 0.0f, 0.0f, 0.0f, 0.0f, 0.0f, 0.0f, 0.0f, 0.0f, 0.0f, 0.0f, 0.0f, 0.0f, 0.0f, 1.0f},
+    {0.0f, 0.0f, 0.0f, 0.0f, 0.0f, 0.0f, 0.0f, 0.0f, 1.0f}
+},
+                             IOTestData{
+    {0.0f, 0.0f, 0.0f, 0.0f, 0.0f, 0.0f, 0.0f, 0.0f, 1.5f},
+    {0.0f, 0.0f, 0.0f, 0.0f, 0.0f, 0.0f, 0.0f, 0.0f, 0.0f, 0.0f, 0.0f, 0.0f, 0.0f, 0.0f, 0.0f, 0.0f, 0.0f},
+    {0.0f, 0.0f, 0.0f, 0.0f, 0.0f, 0.0f, 0.0f, 0.0f, 0.0f}
+}
+                             ));
diff --git a/Engine/Plugins/Animation/RigLogic/Source/RigLogicLibTest/Private/rltests/controls/ControlFixtures.cpp b/Engine/Plugins/Animation/RigLogic/Source/RigLogicLibTest/Private/rltests/controls/ControlFixtures.cpp
new file mode 100644
--- /dev/null
+++ b/Engine/Plugins/Animation/RigLogic/Source/RigLogicLibTest/Private/rltests/controls/ControlFixtures.cpp
@@ -0,0 +1,21 @@
+// Copyright Epic Games, Inc. All Rights Reserved.
+
+#include "rltests/controls/ControlFixtures.h"
+
+#include "riglogic/controls/instances/StandardControlsInputInstance.h"
+
+#include <cstddef>
+
+rl4::ControlsInputInstance::Factory ControlsFactory::getInstanceFactory(std::uint16_t guiControlCount,
+                                                                        std::uint16_t rawControlCount,
+                                                                        std::uint16_t psdControlCount,
+                                                                        std::uint16_t mlControlCount) {
+    return [ = ](rl4::MemoryResource* memRes) {
+               return pma::UniqueInstance<rl4::StandardControlsInputInstance, rl4::ControlsInputInstance>::with(memRes).create(
+                   guiControlCount,
+                   rawControlCount,
+                   psdControlCount,
+                   mlControlCount,
+                   memRes);
+    };
+}
diff --git a/Engine/Plugins/Animation/RigLogic/Source/RigLogicLibTest/Private/rltests/controls/TestControls.cpp b/Engine/Plugins/Animation/RigLogic/Source/RigLogicLibTest/Private/rltests/controls/TestControls.cpp
--- a/Engine/Plugins/Animation/RigLogic/Source/RigLogicLibTest/Private/rltests/controls/TestControls.cpp
+++ b/Engine/Plugins/Animation/RigLogic/Source/RigLogicLibTest/Private/rltests/controls/TestControls.cpp
@@ -2,11 +2,10 @@
 
 #include "rltests/Defs.h"
 #include "rltests/conditionaltable/ConditionalTableFixtures.h"
+#include "rltests/controls/ControlFixtures.h"
 
+#include "riglogic/TypeDefs.h"
 #include "riglogic/controls/Controls.h"
-#include "riglogic/types/Aliases.h"
-
-#include <pma/resources/AlignedMemoryResource.h>
 
 #ifdef _MSC_VER
     #pragma warning(push)
@@ -22,31 +21,48 @@ TEST(ControlsTest, GUIToRawMapping) {
     auto conditionals = ConditionalTableFactory::withMultipleIODefaults(&amr);
 
     rl4::PSDMatrix psds{0u, {}, {}, {}};
-    rl4::Controls controls{std::move(conditionals), std::move(psds)};
+    const std::uint16_t guiControlCount = conditionals.getInputCount();
+    const std::uint16_t rawControlCount = conditionals.getOutputCount();
+    const std::uint16_t psdControlCount = psds.getDistinctPSDCount();
+    auto instanceFactory = ControlsFactory::getInstanceFactory(guiControlCount, rawControlCount, psdControlCount, 0u);
+    rl4::Controls controls{std::move(conditionals), std::move(psds), instanceFactory};
 
     const rl4::Vector<float> guiControls{0.1f, 0.2f};
     const rl4::Vector<float> expected{0.3f, 0.6f};
 
-    std::array<float, 2ul> outputs;
-    controls.mapGUIToRaw(rl4::ConstArrayView<float>{guiControls}, rl4::ArrayView<float>{outputs});
+    auto instance = controls.createInstance(&amr);
+    auto guiBuffer = instance->getGUIControlBuffer();
+    auto rawBuffer = instance->getInputBuffer();
+    std::copy(guiControls.begin(), guiControls.end(), guiBuffer.begin());
+    controls.mapGUIToRaw(instance.get());
 
-    ASSERT_ELEMENTS_EQ(outputs, expected, 2ul);
+    ASSERT_EQ(rawBuffer.size(), expected.size());
+    ASSERT_ELEMENTS_EQ(rawBuffer, expected, expected.size());
 }
 
 TEST(ControlsTest, PSDsAppendToOutput) {
     pma::AlignedMemoryResource amr;
     auto conditionals = ConditionalTableFactory::withMultipleIODefaults(&amr);
 
+    rl4::Vector<float> rawControls{0.1f, 0.2f};
+    const rl4::Vector<float> expected{0.1f, 0.2f, 0.4f, 0.6f};
+
     rl4::Vector<std::uint16_t> rows{2u, 3u};
     rl4::Vector<std::uint16_t> cols{0u, 1u};
     rl4::Vector<float> values{4.0f, 3.0f};
     rl4::PSDMatrix psds{2u, std::move(rows), std::move(cols), std::move(values)};
 
-    rl4::Controls controls{std::move(conditionals), std::move(psds)};
+    const std::uint16_t guiControlCount = conditionals.getInputCount();
+    const std::uint16_t rawControlCount = conditionals.getOutputCount();
+    const std::uint16_t psdControlCount = psds.getDistinctPSDCount();
+    auto instanceFactory = ControlsFactory::getInstanceFactory(guiControlCount, rawControlCount, psdControlCount, 0u);
+    rl4::Controls controls{std::move(conditionals), std::move(psds), instanceFactory};
+    auto instance = controls.createInstance(&amr);
+    auto buffer = instance->getInputBuffer();
+    std::copy(rawControls.begin(), rawControls.end(), buffer.begin());
 
-    rl4::Vector<float> rawControls{0.1f, 0.2f, 0.0f, 0.0f};
-    const rl4::Vector<float> expected{0.1f, 0.2f, 0.4f, 0.6f};
-    controls.calculate(rl4::ArrayView<float>{rawControls});
+    controls.calculate(instance.get());
 
-    ASSERT_ELEMENTS_EQ(rawControls, expected, 4ul);
+    ASSERT_EQ(buffer.size(), expected.size());
+    ASSERT_ELEMENTS_EQ(buffer, expected, expected.size());
 }
