Update RigLogic to v9.1.1 and GeneSplicer to v9.4.0

diff --git a/Engine/Plugins/Animation/RigLogic/Source/RigLogicLibTest/Private/rltests/dna/DNAFixtures.cpp b/Engine/Plugins/Animation/RigLogic/Source/RigLogicLibTest/Private/rltests/dna/DNAFixtures.cpp
--- a/Engine/Plugins/Animation/RigLogic/Source/RigLogicLibTest/Private/rltests/dna/DNAFixtures.cpp
+++ b/Engine/Plugins/Animation/RigLogic/Source/RigLogicLibTest/Private/rltests/dna/DNAFixtures.cpp
@@ -1,9 +1,12 @@
 // Copyright Epic Games, Inc. All Rights Reserved.
 
+#ifdef _MSC_VER
+    #pragma warning(disable : 4503)
+#endif
+
 #include "rltests/dna/DNAFixtures.h"
 
-#include <dna/Reader.h>
-#include <pma/TypeDefs.h>
+#include "riglogic/TypeDefs.h"
 
 namespace rltests {
 
@@ -18,16 +21,29 @@ namespace raw {
 const unsigned char header[] = {
     0x44, 0x4e, 0x41,  // DNA signature
     0x00, 0x02,  // Generation
-    0x00, 0x01,  // Version
-    // Section Lookup Table
-    0x00, 0x00, 0x00, 0x27,  // Descriptor offset
-    0x00, 0x00, 0x00, 0x7e,  // Definition offset
-    0x00, 0x00, 0x03, 0x9a,  // Behavior offset
-    0x00, 0x00, 0x03, 0x9a,  // Controls offset
-    0x00, 0x00, 0x05, 0xac,  // Joints offset
-    0x00, 0x00, 0x07, 0x68,  // BlendShapes offset
-    0x00, 0x00, 0x07, 0x94,  // AnimatedMaps offset
-    0x00, 0x00, 0x08, 0xe0  // Geometry offset
+    0x00, 0x03,  // Version
+    // Index Table
+    0x00, 0x00, 0x00, 0x05,  // Index table entry count
+    0x64, 0x65, 0x73, 0x63,  // Descriptor id
+    0x00, 0x01, 0x00, 0x01,  // Descriptor version
+    0x00, 0x00, 0x00, 0x5b,  // Descriptor offset
+    0x00, 0x00, 0x00, 0x57,  // Descriptor size
+    0x64, 0x65, 0x66, 0x6e,  // Definition id
+    0x00, 0x01, 0x00, 0x01,  // Definition version
+    0x00, 0x00, 0x00, 0xb2,  // Definition offset
+    0x00, 0x00, 0x03, 0x1a,  // Definition size
+    0x62, 0x68, 0x76, 0x72,  // Behavior id
+    0x00, 0x01, 0x00, 0x01,  // Behavior version
+    0x00, 0x00, 0x03, 0xcc,  // Behavior offset
+    0x00, 0x00, 0x05, 0x46,  // Behavior size
+    0x67, 0x65, 0x6f, 0x6d,  // Geometry id
+    0x00, 0x01, 0x00, 0x01,  // Geometry version
+    0x00, 0x00, 0x09, 0x12,  // Geometry offset
+    0x00, 0x00, 0x04, 0x38,  // Geometry size
+    0x6d, 0x6c, 0x62, 0x68,  // Machine learned behavior id
+    0x00, 0x01, 0x00, 0x00,  // Machine learned behavior version
+    0x00, 0x00, 0x0d, 0x4a,  // Machine learned behavior offset
+    0x00, 0x00, 0x02, 0xfa  // Machine learned behavior size
 };
 
 const unsigned char descriptor[] = {
@@ -123,10 +139,9 @@ const unsigned char definition[] = {
     0x00, 0x00,  // Map from LOD-0 to row 0 in below defined matrix
     0x00, 0x01,  // Map from LOD-1 to row 1 in below defined matrix
     0x00, 0x00, 0x00, 0x02,  // Mesh name indices per LOD row count
-    0x00, 0x00, 0x00, 0x03,  // Indices matrix row-0
+    0x00, 0x00, 0x00, 0x02,  // Indices matrix row-0
     0x00, 0x00,  // Mesh name index: 0
     0x00, 0x01,  // Mesh name index: 1
-    0x00, 0x02,  // Mesh name index: 2
     0x00, 0x00, 0x00, 0x01,  // Indices matrix row-1
     0x00, 0x02,  // Mesh name index: 2
     0x00, 0x00, 0x00, 0x09,  // Gui control names length
@@ -261,7 +276,7 @@ const unsigned char definition[] = {
     0x00, 0x01,  // JE
     0x00, 0x04,  // JF
     0x00, 0x02,  // JG
-    0x00, 0x06,  // JH
+    0x00, 0x04,  // JH
     0x00, 0x02,  // JI
     0x00, 0x00, 0x00, 0x09,  // Neutral joint translation X values length
     0x3f, 0x80, 0x00, 0x00,  // 1.0f
@@ -708,7 +723,7 @@ const unsigned char animatedmaps[] = {
 const unsigned char geometry[] = {
     0x00, 0x00, 0x00, 0x03,  // Mesh count
     // Mesh-0
-    0x00, 0x00, 0x0a, 0x3a,  // Mesh-0 end offset
+    0x00, 0x00, 0x01, 0x52,  // Mesh-0 size
     0x00, 0x00, 0x00, 0x03,  // Vertex positions X values length
     0x40, 0xe0, 0x00, 0x00,  // 7.0f
     0x41, 0x00, 0x00, 0x00,  // 8.0f
@@ -761,22 +776,22 @@ const unsigned char geometry[] = {
     0x00, 0x08,  // Maximum influence per vertex
     0x00, 0x00, 0x00, 0x03,  // Skin weights structure count: 3 (for each vertex)
     0x00, 0x00, 0x00, 0x03,  // Weights length: 3 (for each influencing joint)
-    0x40, 0xe0, 0x00, 0x00,  // 7.0f
-    0x41, 0x00, 0x00, 0x00,  // 8.0f
-    0x41, 0x10, 0x00, 0x00,  // 9.0f
+    0x3f, 0x33, 0x33, 0x33,  // 0.7f
+    0x3d, 0xcc, 0xcc, 0xcd,  // 0.1f
+    0x3e, 0x4c, 0xcc, 0xcd,  // 0.2f
     0x00, 0x00, 0x00, 0x03,  // Influencing joint count: 3 (for each weight)
     0x00, 0x00,  // Joint: 0
     0x00, 0x01,  // Joint: 1
     0x00, 0x02,  // Joint: 2
     0x00, 0x00, 0x00, 0x02,  // Weights length: 2 (for each influencing joint)
-    0x40, 0xe0, 0x00, 0x00,  // 7.0f
-    0x41, 0x00, 0x00, 0x00,  // 8.0f
+    0x3f, 0x00, 0x00, 0x00,  // 0.5f
+    0x3f, 0x00, 0x00, 0x00,  // 0.5f
     0x00, 0x00, 0x00, 0x02,  // Influencing joint count: 2 (for each weight)
     0x00, 0x03,  // Joint: 3
     0x00, 0x04,  // Joint: 4
     0x00, 0x00, 0x00, 0x02,  // Weights length: 2 (for each influencing joint)
-    0x41, 0x00, 0x00, 0x00,  // 8.0f
-    0x41, 0x10, 0x00, 0x00,  // 9.0f
+    0x3e, 0xcc, 0xcc, 0xcd,  // 0.4f
+    0x3f, 0x19, 0x99, 0x9a,  // 0.6f
     0x00, 0x00, 0x00, 0x02,  // Influencing joint count: 2 (for each weight)
     0x00, 0x05,  // Joint: 5
     0x00, 0x06,  // Joint: 6
@@ -799,7 +814,7 @@ const unsigned char geometry[] = {
     0x00, 0x00, 0x00, 0x02,  // Vertex position: 2
     0x00, 0x06,  // Blend shape index in Definition
     // Mesh-1
-    0x00, 0x00, 0x0b, 0x90,  // Mesh-1 end offset
+    0x00, 0x00, 0x01, 0x52,  // Mesh-1 size
     0x00, 0x00, 0x00, 0x03,  // Vertex positions X values length
     0x40, 0x80, 0x00, 0x00,  // 4.0f
     0x40, 0xa0, 0x00, 0x00,  // 5.0f
@@ -852,22 +867,22 @@ const unsigned char geometry[] = {
     0x00, 0x08,  // Maximum influence per vertex
     0x00, 0x00, 0x00, 0x03,  // Skin weights structure count: 3 (for each vertex)
     0x00, 0x00, 0x00, 0x03,  // Weights length: 3 (for each influencing joint)
-    0x40, 0x80, 0x00, 0x00,  // 4.0f
-    0x40, 0xa0, 0x00, 0x00,  // 5.0f
-    0x40, 0xc0, 0x00, 0x00,  // 6.0f
+    0x3e, 0xcc, 0xcc, 0xcd,  // 0.4f
+    0x3e, 0x99, 0x99, 0x9a,  // 0.3f
+    0x3e, 0x99, 0x99, 0x9a,  // 0.3f
     0x00, 0x00, 0x00, 0x03,  // Influencing joint count: 3 (for each weight)
     0x00, 0x00,  // Joint: 0
     0x00, 0x01,  // Joint: 1
     0x00, 0x02,  // Joint: 2
     0x00, 0x00, 0x00, 0x02,  // Weights length: 2 (for each influencing joint)
-    0x40, 0x80, 0x00, 0x00,  // 4.0f
-    0x40, 0xa0, 0x00, 0x00,  // 5.0f
+    0x3f, 0x4c, 0xcc, 0xcd,  // 0.8f
+    0x3e, 0x4c, 0xcc, 0xcd,  // 0.2f
     0x00, 0x00, 0x00, 0x02,  // Influencing joint count: 2 (for each weight)
     0x00, 0x03,  // Joint: 3
     0x00, 0x04,  // Joint: 4
     0x00, 0x00, 0x00, 0x02,  // Weights length: 2 (for each influencing joint)
-    0x40, 0xa0, 0x00, 0x00,  // 5.0f
-    0x40, 0xc0, 0x00, 0x00,  // 6.0f
+    0x3d, 0xcc, 0xcc, 0xcd,  // 0.1f
+    0x3f, 0x66, 0x66, 0x66,  // 0.9f
     0x00, 0x00, 0x00, 0x02,  // Influencing joint count: 2 (for each weight)
     0x00, 0x05,  // Joint: 5
     0x00, 0x06,  // Joint: 6
@@ -890,7 +905,7 @@ const unsigned char geometry[] = {
     0x00, 0x00, 0x00, 0x02,  // Vertex position: 2
     0x00, 0x06,  // Blend shape index in Definition
     // Mesh-2
-    0x00, 0x00, 0x0d, 0x18,  // Mesh-2 offset
+    0x00, 0x00, 0x01, 0x84,  // Mesh-2 size
     0x00, 0x00, 0x00, 0x03,  // Vertex positions X values length
     0x3f, 0x80, 0x00, 0x00,  // 1.0f
     0x40, 0x00, 0x00, 0x00,  // 2.0f
@@ -943,22 +958,22 @@ const unsigned char geometry[] = {
     0x00, 0x08,  // Maximum influence per vertex
     0x00, 0x00, 0x00, 0x03,  // Skin weights structure count: 3 (for each vertex)
     0x00, 0x00, 0x00, 0x03,  // Weights length: 3 (for each influencing joint)
-    0x3f, 0x80, 0x00, 0x00,  // 1.0f
-    0x40, 0x00, 0x00, 0x00,  // 2.0f
-    0x40, 0x40, 0x00, 0x00,  // 3.0f
+    0x3d, 0xcc, 0xcc, 0xcd,  // 0.1f
+    0x3e, 0x99, 0x99, 0x9a,  // 0.3f
+    0x3f, 0x19, 0x99, 0x9a,  // 0.6f
     0x00, 0x00, 0x00, 0x03,  // Influencing joint count: 3 (for each weight)
     0x00, 0x00,  // Joint: 0
     0x00, 0x01,  // Joint: 1
     0x00, 0x02,  // Joint: 2
     0x00, 0x00, 0x00, 0x02,  // Weights length: 2 (for each influencing joint)
-    0x3f, 0x80, 0x00, 0x00,  // 1.0f
-    0x40, 0x00, 0x00, 0x00,  // 2.0f
+    0x3e, 0x99, 0x99, 0x9a,  // 0.3f
+    0x3f, 0x33, 0x33, 0x33,  // 0.7f
     0x00, 0x00, 0x00, 0x02,  // Influencing joint count: 2 (for each weight)
     0x00, 0x03,  // Joint: 3
     0x00, 0x04,  // Joint: 4
     0x00, 0x00, 0x00, 0x02,  // Weights length: 2 (for each influencing joint)
-    0x40, 0x00, 0x00, 0x00,  // 2.0f
-    0x40, 0x40, 0x00, 0x00,  // 3.0f
+    0x3e, 0x4c, 0xcc, 0xcd,  // 0.2f
+    0x3f, 0x4c, 0xcc, 0xcd,  // 0.8f
     0x00, 0x00, 0x00, 0x02,  // Influencing joint count: 2 (for each weight)
     0x00, 0x05,  // Joint: 5
     0x00, 0x06,  // Joint: 6
@@ -995,12 +1010,237 @@ const unsigned char geometry[] = {
     0x00, 0x07  // Blend shape index in Definition
 };
 
-const unsigned char footer[] = {
-    0x41, 0x4e, 0x44,  // AND end signature
+const unsigned char machineLearnedBehavior[] = {
+    0x00, 0x00, 0x00, 0x09,  // Raw control names length
+    0x00, 0x00, 0x00, 0x02,  // Raw control name 0 length
+    0x4d, 0x41,  // Raw control name 0 : MA
+    0x00, 0x00, 0x00, 0x02,  // Raw control name 1 length
+    0x4d, 0x42,  // Raw control name 1 : MB
+    0x00, 0x00, 0x00, 0x02,  // Raw control name 2 length
+    0x4d, 0x43,  // Raw control name 2 : MC
+    0x00, 0x00, 0x00, 0x02,  // Raw control name 3 length
+    0x4d, 0x44,  // Raw control name 3 : MD
+    0x00, 0x00, 0x00, 0x02,  // Raw control name 4 length
+    0x4d, 0x45,  // Raw control name 4 : ME
+    0x00, 0x00, 0x00, 0x02,  // Raw control name 5 length
+    0x4d, 0x46,  // Raw control name 5 : MF
+    0x00, 0x00, 0x00, 0x02,  // Raw control name 6 length
+    0x4d, 0x47,  // Raw control name 6 : MG
+    0x00, 0x00, 0x00, 0x02,  // Raw control name 7 length
+    0x4d, 0x48,  // Raw control name 7 : MH
+    0x00, 0x00, 0x00, 0x02,  // Raw control name 8 length
+    0x4d, 0x49,  // Raw control name 8 : MI
+    0x00, 0x00, 0x00, 0x02,  // Neural network indices lod to row mapping length
+    0x00, 0x00,  // Map from LOD-0 to row 0 in below defined matrix
+    0x00, 0x01,  // Map from LOD-1 to row 1 in below defined matrix
+    0x00, 0x00, 0x00, 0x02,  // Neural network indices per LOD row count
+    0x00, 0x00, 0x00, 0x04,  // Indices matrix row-0
+    0x00, 0x00,  // Neural network index: 0
+    0x00, 0x01,  // Neural network index: 1
+    0x00, 0x02,  // Neural network index: 2
+    0x00, 0x03,  // Neural network index: 3
+    0x00, 0x00, 0x00, 0x02,  // Indices matrix row-1
+    0x00, 0x04,  // Neural network index: 4
+    0x00, 0x05,  // Neural network index: 5
+    0x00, 0x00, 0x00, 0x03,  // Region names length
+    0x00, 0x00, 0x00, 0x02,  // Region names length for mesh 0
+    0x00, 0x00, 0x00, 0x02,  // Region name 0 length
+    0x52, 0x41,  // Region name 0 : RA
+    0x00, 0x00, 0x00, 0x02,  // Region name 1 length
+    0x52, 0x42,  // Region name 1 : RB
+    0x00, 0x00, 0x00, 0x02,  // Region names length for mesh 1
+    0x00, 0x00, 0x00, 0x02,  // Region name 0 length
+    0x52, 0x43,  // Region name 0 : RC
+    0x00, 0x00, 0x00, 0x02,  // Region name 1 length
+    0x52, 0x44,  // Region name 1 : RD
+    0x00, 0x00, 0x00, 0x02,  // Region names length for mesh 2
+    0x00, 0x00, 0x00, 0x02,  // Region name 0 length
+    0x52, 0x45,  // Region name 0 : RE
+    0x00, 0x00, 0x00, 0x02,  // Region name 1 length
+    0x52, 0x46,  // Region name 1 : RF
+    0x00, 0x00, 0x00, 0x03,  // Mesh count
+    0x00, 0x00, 0x00, 0x02,  // Region count for Mesh-0
+    0x00, 0x00, 0x00, 0x01,  // Neural network index count for Mesh-0 Region-0
+    0x00, 0x00,  // Neural network index: 0
+    0x00, 0x00, 0x00, 0x01,  // Neural network index count for Mesh-0 Region-1
+    0x00, 0x01,  // Neural network index: 1
+    0x00, 0x00, 0x00, 0x02,  // Region count for Mesh-1
+    0x00, 0x00, 0x00, 0x01,  // Neural network index count for Mesh-1 Region-0
+    0x00, 0x02,  // Neural network index: 2
+    0x00, 0x00, 0x00, 0x01,  // Neural network index count for Mesh-1 Region-1
+    0x00, 0x03,  // Neural network index: 3
+    0x00, 0x00, 0x00, 0x02,  // Region count for Mesh-2
+    0x00, 0x00, 0x00, 0x01,  // Neural network index count for Mesh-2 Region-0
+    0x00, 0x04,  // Neural network index: 4
+    0x00, 0x00, 0x00, 0x01,  // Neural network index count for Mesh-2 Region-1
+    0x00, 0x05,  // Neural network index: 5
+    0x00, 0x00, 0x00, 0x06,  // Neural network count
+    // Mesh-0 Region-0 neural network
+    0x00, 0x00, 0x00, 0x5a,  // Mesh-0 Region-0 neural network size
+    0x00, 0x00, 0x00, 0x01,  // Mesh-0 Region-0 neural network output index count
+    0x00, 0x09,  // Mesh-0 Region-0 neural network output index-9
+    0x00, 0x00, 0x00, 0x02,  // Mesh-0 Region-1 neural network input index count
+    0x00, 0x00,  // Mesh-0 Region-0 neural network input index-0
+    0x00, 0x01,  // Mesh-0 Region-0 neural network input index-1
+    0x00, 0x00, 0x00, 0x02,  // Mesh-0 Region-0 neural network layer count
+    0x00, 0x00, 0x00, 0x02,  // Mesh-0 Region-0 neural network layer-0 bias count
+    0x3f, 0x80, 0x00, 0x00,  // 1.0f
+    0x3f, 0x80, 0x00, 0x00,  // 1.0f
+    0x00, 0x00, 0x00, 0x04,  // Mesh-0 Region-0 neural network layer-0 weight count
+    0x3f, 0x00, 0x00, 0x00,  // 0.5f
+    0x3f, 0x00, 0x00, 0x00,  // 0.5f
+    0x3f, 0x00, 0x00, 0x00,  // 0.5f
+    0x3f, 0x00, 0x00, 0x00,  // 0.5f
+    0x00, 0x01,  // Mesh-0 Region-0 neural network layer-0 activation function ID
+    0x00, 0x00, 0x00, 0x01,  // Mesh-0 Region-0 neural network layer-1 activation function parameter count
+    0x3f, 0x00, 0x00, 0x00,  // 0.5f
+    0x00, 0x00, 0x00, 0x01,  // Mesh-0 Region-0 neural network layer-1 bias count
+    0x3f, 0x80, 0x00, 0x00,  // 1.0f
+    0x00, 0x00, 0x00, 0x02,  // Mesh-0 Region-0 neural network layer-1 weight count
+    0x3f, 0x00, 0x00, 0x00,  // 0.5f
+    0x3f, 0x00, 0x00, 0x00,  // 0.5f
+    0x00, 0x01,  // Mesh-0 Region-0 neural network layer-1 activation function ID
+    0x00, 0x00, 0x00, 0x01,  // Mesh-0 Region-0 neural network layer-1 activation function parameter count
+    0x3f, 0x00, 0x00, 0x00,  // 0.5f
+    // Mesh-0 Region-1 neural network
+    0x00, 0x00, 0x00, 0x5a,  // Mesh-0 Region-1 neural network size
+    0x00, 0x00, 0x00, 0x01,  // Mesh-0 Region-1 neural network output index count
+    0x00, 0x0a,  // Mesh-0 Region-1 neural network output index-10
+    0x00, 0x00, 0x00, 0x02,  // Mesh-0 Region-1 neural network input index count
+    0x00, 0x02,  // Mesh-0 Region-1 neural network input index-2
+    0x00, 0x03,  // Mesh-0 Region-1 neural network input index-3
+    0x00, 0x00, 0x00, 0x02,  // Mesh-0 Region-1 neural network layer count
+    0x00, 0x00, 0x00, 0x02,  // Mesh-0 Region-1 neural network layer-0 bias count
+    0x3f, 0x00, 0x00, 0x00,  // 0.5f
+    0x3f, 0x00, 0x00, 0x00,  // 0.5f
+    0x00, 0x00, 0x00, 0x04,  // Mesh-0 Region-1 neural network layer-0 weight count
+    0x3f, 0x80, 0x00, 0x00,  // 1.0f
+    0x3f, 0x80, 0x00, 0x00,  // 1.0f
+    0x3f, 0x80, 0x00, 0x00,  // 1.0f
+    0x3f, 0x80, 0x00, 0x00,  // 1.0f
+    0x00, 0x01,  // Mesh-0 Region-1 neural network layer-0 activation function ID
+    0x00, 0x00, 0x00, 0x01,  // Mesh-0 Region-1 neural network layer-1 activation function parameter count
+    0x3f, 0x80, 0x00, 0x00,  // 1.0f
+    0x00, 0x00, 0x00, 0x01,  // Mesh-0 Region-1 neural network layer-1 bias count
+    0x3f, 0x00, 0x00, 0x00,  // 0.5f
+    0x00, 0x00, 0x00, 0x02,  // Mesh-0 Region-1 neural network layer-1 weight count
+    0x3f, 0x80, 0x00, 0x00,  // 1.0f
+    0x3f, 0x80, 0x00, 0x00,  // 1.0f
+    0x00, 0x01,  // Mesh-0 Region-1 neural network layer-1 activation function ID
+    0x00, 0x00, 0x00, 0x01,  // Mesh-0 Region-1 neural network layer-1 activation function parameter count
+    0x3f, 0x80, 0x00, 0x00,  // 1.0f
+    // Mesh-1 Region-0 neural network
+    0x00, 0x00, 0x00, 0x5a,  // Mesh-1 Region-0 neural network size
+    0x00, 0x00, 0x00, 0x01,  // Mesh-1 Region-0 neural network output index count
+    0x00, 0x0b,  // Mesh-1 Region-0 neural network output index-11
+    0x00, 0x00, 0x00, 0x02,  // Mesh-1 Region-1 neural network input index count
+    0x00, 0x04,  // Mesh-1 Region-0 neural network input index-4
+    0x00, 0x05,  // Mesh-1 Region-0 neural network input index-5
+    0x00, 0x00, 0x00, 0x02,  // Mesh-1 Region-0 neural network layer count
+    0x00, 0x00, 0x00, 0x02,  // Mesh-1 Region-0 neural network layer-0 bias count
+    0x3f, 0x00, 0x00, 0x00,  // 0.5f
+    0x3f, 0x00, 0x00, 0x00,  // 0.5f
+    0x00, 0x00, 0x00, 0x04,  // Mesh-1 Region-0 neural network layer-0 weight count
+    0x3f, 0x80, 0x00, 0x00,  // 1.0f
+    0x3f, 0x80, 0x00, 0x00,  // 1.0f
+    0x3f, 0x80, 0x00, 0x00,  // 1.0f
+    0x3f, 0x80, 0x00, 0x00,  // 1.0f
+    0x00, 0x01,  // Mesh-1 Region-0 neural network layer-0 activation function ID
+    0x00, 0x00, 0x00, 0x01,  // Mesh-1 Region-0 neural network layer-1 activation function parameter count
+    0x3f, 0x80, 0x00, 0x00,  // 1.0f
+    0x00, 0x00, 0x00, 0x01,  // Mesh-1 Region-0 neural network layer-1 bias count
+    0x3f, 0x00, 0x00, 0x00,  // 0.5f
+    0x00, 0x00, 0x00, 0x02,  // Mesh-1 Region-0 neural network layer-1 weight count
+    0x3f, 0x80, 0x00, 0x00,  // 1.0f
+    0x3f, 0x80, 0x00, 0x00,  // 1.0f
+    0x00, 0x01,  // Mesh-1 Region-0 neural network layer-1 activation function ID
+    0x00, 0x00, 0x00, 0x01,  // Mesh-1 Region-0 neural network layer-1 activation function parameter count
+    0x3f, 0x80, 0x00, 0x00,  // 1.0f
+    // Mesh-1 Region-1 neural network
+    0x00, 0x00, 0x00, 0x5a,  // Mesh-1 Region-1 neural network size
+    0x00, 0x00, 0x00, 0x01,  // Mesh-1 Region-1 neural network output index count
+    0x00, 0x0c,  // Mesh-1 Region-1 neural network output index-12
+    0x00, 0x00, 0x00, 0x02,  // Mesh-1 Region-1 neural network input index count
+    0x00, 0x06,  // Mesh-1 Region-1 neural network input index-6
+    0x00, 0x07,  // Mesh-1 Region-1 neural network input index-7
+    0x00, 0x00, 0x00, 0x02,  // Mesh-1 Region-1 neural network layer count
+    0x00, 0x00, 0x00, 0x02,  // Mesh-1 Region-1 neural network layer-0 bias count
+    0x3f, 0x80, 0x00, 0x00,  // 1.0f
+    0x3f, 0x80, 0x00, 0x00,  // 1.0f
+    0x00, 0x00, 0x00, 0x04,  // Mesh-1 Region-1 neural network layer-0 weight count
+    0x3f, 0x00, 0x00, 0x00,  // 0.5f
+    0x3f, 0x00, 0x00, 0x00,  // 0.5f
+    0x3f, 0x00, 0x00, 0x00,  // 0.5f
+    0x3f, 0x00, 0x00, 0x00,  // 0.5f
+    0x00, 0x01,  // Mesh-1 Region-1 neural network layer-0 activation function ID
+    0x00, 0x00, 0x00, 0x01,  // Mesh-1 Region-1 neural network layer-1 activation function parameter count
+    0x3f, 0x00, 0x00, 0x00,  // 0.5f
+    0x00, 0x00, 0x00, 0x01,  // Mesh-1 Region-1 neural network layer-1 bias count
+    0x3f, 0x80, 0x00, 0x00,  // 1.0f
+    0x00, 0x00, 0x00, 0x02,  // Mesh-1 Region-1 neural network layer-1 weight count
+    0x3f, 0x00, 0x00, 0x00,  // 0.5f
+    0x3f, 0x00, 0x00, 0x00,  // 0.5f
+    0x00, 0x01,  // Mesh-1 Region-1 neural network layer-1 activation function ID
+    0x00, 0x00, 0x00, 0x01,  // Mesh-1 Region-1 neural network layer-1 activation function parameter count
+    0x3f, 0x00, 0x00, 0x00,  // 0.5f
+    // Mesh-2 Region-0 neural network
+    0x00, 0x00, 0x00, 0x5a,  // Mesh-2 Region-0 neural network size
+    0x00, 0x00, 0x00, 0x01,  // Mesh-2 Region-0 neural network output index count
+    0x00, 0x0d,  // Mesh-2 Region-0 neural network output index-13
+    0x00, 0x00, 0x00, 0x02,  // Mesh-2 Region-1 neural network input index count
+    0x00, 0x08,  // Mesh-2 Region-0 neural network input index-8
+    0x00, 0x00,  // Mesh-2 Region-0 neural network input index-0
+    0x00, 0x00, 0x00, 0x02,  // Mesh-2 Region-0 neural network layer count
+    0x00, 0x00, 0x00, 0x02,  // Mesh-2 Region-0 neural network layer-0 bias count
+    0x3f, 0x80, 0x00, 0x00,  // 1.0f
+    0x3f, 0x80, 0x00, 0x00,  // 1.0f
+    0x00, 0x00, 0x00, 0x04,  // Mesh-2 Region-0 neural network layer-0 weight count
+    0x3f, 0x00, 0x00, 0x00,  // 0.5f
+    0x3f, 0x00, 0x00, 0x00,  // 0.5f
+    0x3f, 0x00, 0x00, 0x00,  // 0.5f
+    0x3f, 0x00, 0x00, 0x00,  // 0.5f
+    0x00, 0x01,  // Mesh-2 Region-0 neural network layer-0 activation function ID
+    0x00, 0x00, 0x00, 0x01,  // Mesh-2 Region-0 neural network layer-1 activation function parameter count
+    0x3f, 0x00, 0x00, 0x00,  // 0.5f
+    0x00, 0x00, 0x00, 0x01,  // Mesh-2 Region-0 neural network layer-1 bias count
+    0x3f, 0x80, 0x00, 0x00,  // 1.0f
+    0x00, 0x00, 0x00, 0x02,  // Mesh-2 Region-0 neural network layer-1 weight count
+    0x3f, 0x00, 0x00, 0x00,  // 0.5f
+    0x3f, 0x00, 0x00, 0x00,  // 0.5f
+    0x00, 0x01,  // Mesh-2 Region-0 neural network layer-1 activation function ID
+    0x00, 0x00, 0x00, 0x01,  // Mesh-2 Region-0 neural network layer-1 activation function parameter count
+    0x3f, 0x00, 0x00, 0x00,  // 0.5f
+    // Mesh-2 Region-1 neural network
+    0x00, 0x00, 0x00, 0x5a,  // Mesh-2 Region-1 neural network size
+    0x00, 0x00, 0x00, 0x01,  // Mesh-2 Region-1 neural network output index count
+    0x00, 0x0e,  // Mesh-2 Region-1 neural network output index-14
+    0x00, 0x00, 0x00, 0x02,  // Mesh-2 Region-1 neural network input index count
+    0x00, 0x04,  // Mesh-2 Region-1 neural network input index-4
+    0x00, 0x07,  // Mesh-2 Region-1 neural network input index-7
+    0x00, 0x00, 0x00, 0x02,  // Mesh-2 Region-1 neural network layer count
+    0x00, 0x00, 0x00, 0x02,  // Mesh-2 Region-1 neural network layer-0 bias count
+    0x3f, 0x00, 0x00, 0x00,  // 0.5f
+    0x3f, 0x00, 0x00, 0x00,  // 0.5f
+    0x00, 0x00, 0x00, 0x04,  // Mesh-2 Region-1 neural network layer-0 weight count
+    0x3f, 0x80, 0x00, 0x00,  // 1.0f
+    0x3f, 0x80, 0x00, 0x00,  // 1.0f
+    0x3f, 0x80, 0x00, 0x00,  // 1.0f
+    0x3f, 0x80, 0x00, 0x00,  // 1.0f
+    0x00, 0x01,  // Mesh-2 Region-1 neural network layer-0 activation function ID
+    0x00, 0x00, 0x00, 0x01,  // Mesh-2 Region-1 neural network layer-1 activation function parameter count
+    0x3f, 0x80, 0x00, 0x00,  // 1.0f
+    0x00, 0x00, 0x00, 0x01,  // Mesh-2 Region-1 neural network layer-1 bias count
+    0x3f, 0x00, 0x00, 0x00,  // 0.5f
+    0x00, 0x00, 0x00, 0x02,  // Mesh-2 Region-1 neural network layer-1 weight count
+    0x3f, 0x80, 0x00, 0x00,  // 1.0f
+    0x3f, 0x80, 0x00, 0x00,  // 1.0f
+    0x00, 0x01,  // Mesh-2 Region-1 neural network layer-1 activation function ID
+    0x00, 0x00, 0x00, 0x01,  // Mesh-2 Region-1 neural network layer-1 activation function parameter count
+    0x3f, 0x80, 0x00, 0x00  // 1.0f
 };
 
-std::vector<unsigned char> getBytes() {
-    std::vector<unsigned char> bytes;
+std::vector<char> getBytes() {
+    std::vector<char> bytes;
     // Header
     bytes.insert(bytes.end(), header, header + sizeof(header));
     // Descriptor
@@ -1021,8 +1261,8 @@ std::vector<unsigned char> getBytes() {
     bytes.insert(bytes.end(), conditionals, conditionals + sizeof(conditionals));
     // Geometry
     bytes.insert(bytes.end(), geometry, geometry + sizeof(geometry));
-    // Footer
-    bytes.insert(bytes.end(), footer, footer + sizeof(footer));
+    // Machine learned behavior
+    bytes.insert(bytes.end(), machineLearnedBehavior, machineLearnedBehavior + sizeof(machineLearnedBehavior));
     return bytes;
 }
 
@@ -1047,12 +1287,14 @@ const CoordinateSystem coordinateSystem = {
     Direction::front
 };
 const std::uint16_t lodCount[] = {
-    2u,  // MaxLOD-0
-    1u  // MaxLOD-1
+    2u,  // MaxLOD-0 - MinLOD-1
+    1u,  // MaxLOD-1 - MinLOD-1
+    1u  // MaxLOD-0 - MinLOD-0
 };
-const std::uint16_t maxLOD[] = {
-    0u,  // MaxLOD-0
-    1u  // MaxLOD-1
+const std::uint16_t maxLODs[] = {
+    0u,  // MaxLOD-0 - MinLOD-1
+    1u,  // MaxLOD-1 - MinLOD-0
+    0u  // MaxLOD-0 - MinLOD-0
 };
 const String<char> complexity = "A";
 const String<char> dbName = "testDB";
@@ -1065,60 +1307,78 @@ const Vector<String<char> > rawControlNames = {
     "RA", "RB", "RC", "RD", "RE", "RF", "RG", "RH", "RI"
 };
 const VectorOfCharStringMatrix jointNames = {
-    {  // MaxLOD-0
+    {  // MaxLOD-0 - MinLOD-1
         {"JA", "JB", "JC", "JD", "JE", "JF", "JG", "JH", "JI"},
         {"JA", "JB", "JC", "JD", "JG", "JI"}
     },
-    {  // Max-LOD-1
+    {  // MaxLOD-1 - MinLOD-0
         {"JA", "JB", "JC", "JD", "JG", "JI"}
+    },
+    {  // MaxLOD-0 - MinLOD-0
+        {"JA", "JB", "JC", "JD", "JE", "JF", "JG", "JH", "JI"},
     }
 };
 const VectorOfCharStringMatrix blendShapeNames = {
-    {  // MaxLOD-0
+    {  // MaxLOD-0 - MinLOD-1
         {"BA", "BB", "BC", "BD", "BE", "BF", "BG", "BH", "BI"},
         {"BC", "BF", "BH", "BI"}
     },
-    {  // Max-LOD-1
+    {  // MaxLOD-1 - MinLOD-1
         {"BC", "BF", "BH", "BI"}
+    },
+    {  // MaxLOD-0 - MinLOD-0
+        {"BA", "BB", "BC", "BD", "BE", "BF", "BG", "BH", "BI"},
     }
 };
 const VectorOfCharStringMatrix animatedMapNames = {
-    {  // MaxLOD-0
+    {  // MaxLOD-0 - MinLOD-1
         {"AA", "AB", "AC", "AD", "AE", "AF", "AG", "AH", "AI", "AJ"},
         {"AC", "AF", "AH", "AI"}
     },
-    {  // Max-LOD-1
+    {  // MaxLOD-1 - MinLOD-1
         {"AC", "AF", "AH", "AI"}
+    },
+    {  // MaxLOD-0 - MinLOD-0
+        {"AA", "AB", "AC", "AD", "AE", "AF", "AG", "AH", "AI", "AJ"},
     }
 };
 const VectorOfCharStringMatrix meshNames = {
-    {  // MaxLOD-0
-        {"MA", "MB", "MC"},
+    {  // MaxLOD-0 - MinLOD-1
+        {"MA", "MB"},
         {"MC"}
     },
-    {  // MaxLOD-1
+    {  // MaxLOD-1 - MinLOD-1
         {"MC"}
+    },
+    {  // MaxLOD-0 - MinLOD-0
+        {"MA", "MB"}
     }
 };
 const Vector<Matrix<std::uint16_t> > meshBlendShapeIndices = {
-    {  // MaxLOD-0
-        {0, 1, 2, 3, 4, 5, 6, 7, 8},
+    {  // MaxLOD-0 - MinLOD-1
+        {0, 1, 2, 3, 4, 5, 6},
         {7, 8}
     },
-    {  // MaxLOD-1
+    {  // MaxLOD-1 - MinLOD-1
         {0, 1}
+    },
+    {  // MaxLOD-0 - MinLOD-0
+        {0, 1, 2, 3, 4, 5, 6},
     }
 };
 const Matrix<std::uint16_t> jointHierarchy = {
-    {  // MaxLOD-0
-        {0, 0, 0, 1, 1, 4, 2, 6, 2}
+    {  // MaxLOD-0 - MinLOD-1
+        {0, 0, 0, 1, 1, 4, 2, 4, 2}
     },
-    {  // Max-LOD-1
+    {  // MaxLOD-1 - MinLOD-1
         {0, 0, 0, 1, 2, 2}
+    },
+    {  // MaxLOD-0 - MinLOD-0
+        {0, 0, 0, 1, 1, 4, 2, 4, 2}
     }
 };
 const Vector<Matrix<Vector3> > neutralJointTranslations = {
-    {  // MaxLOD-0
+    {  // MaxLOD-0 - MinLOD-1
         {
             {1.0f, 1.0f, 1.0f},
             {2.0f, 2.0f, 2.0f},
@@ -1139,7 +1399,7 @@ const Vector<Matrix<Vector3> > neutralJointTranslations = {
             {9.0f, 9.0f, 9.0f}
         }
     },
-    {  // MaxLOD-1
+    {  // MaxLOD-1 - MinLOD-1
         {
             {1.0f, 1.0f, 1.0f},
             {2.0f, 2.0f, 2.0f},
@@ -1148,10 +1408,23 @@ const Vector<Matrix<Vector3> > neutralJointTranslations = {
             {7.0f, 7.0f, 7.0f},
             {9.0f, 9.0f, 9.0f}
         }
+    },
+    {  // MaxLOD-0 - MinLOD-0
+        {
+            {1.0f, 1.0f, 1.0f},
+            {2.0f, 2.0f, 2.0f},
+            {3.0f, 3.0f, 3.0f},
+            {4.0f, 4.0f, 4.0f},
+            {5.0f, 5.0f, 5.0f},
+            {6.0f, 6.0f, 6.0f},
+            {7.0f, 7.0f, 7.0f},
+            {8.0f, 8.0f, 8.0f},
+            {9.0f, 9.0f, 9.0f}
+        }
     }
 };
 const Vector<Matrix<Vector3> > neutralJointRotations = {
-    {  // MaxLOD-0
+    {  // MaxLOD-0 - MinLOD-1
         {
             {1.0f, 1.0f, 1.0f},
             {2.0f, 2.0f, 2.0f},
@@ -1172,7 +1445,7 @@ const Vector<Matrix<Vector3> > neutralJointRotations = {
             {9.0f, 9.0f, 9.0f}
         }
     },
-    {  // MaxLOD-1
+    {  // MaxLOD-1 - MinLOD-1
         {
             {1.0f, 1.0f, 1.0f},
             {2.0f, 2.0f, 2.0f},
@@ -1181,67 +1454,97 @@ const Vector<Matrix<Vector3> > neutralJointRotations = {
             {7.0f, 7.0f, 7.0f},
             {9.0f, 9.0f, 9.0f}
         }
+    },
+    {  // MaxLOD-0 - MinLOD-0
+        {
+            {1.0f, 1.0f, 1.0f},
+            {2.0f, 2.0f, 2.0f},
+            {3.0f, 3.0f, 3.0f},
+            {4.0f, 4.0f, 4.0f},
+            {5.0f, 5.0f, 5.0f},
+            {6.0f, 6.0f, 6.0f},
+            {7.0f, 7.0f, 7.0f},
+            {8.0f, 8.0f, 8.0f},
+            {9.0f, 9.0f, 9.0f}
+        }
     }
 };
 
 // Behavior
 const std::uint16_t guiControlCount = 9u;
 const std::uint16_t rawControlCount = 9u;
-const std::uint16_t psdCount = 12u;
+const std::uint16_t psdControlCount = 12u;
 // Behavior->Conditionals
 const Vector<Matrix<std::uint16_t> > conditionalInputIndices = {
-    {  // MaxLOD-0
+    {  // MaxLOD-0 - MinLOD-1
         {0, 1, 1, 2, 3, 3, 4, 4, 4, 5, 6, 7, 7, 8, 8},
         {0, 1, 1, 2, 3, 3},
     },
-    {  // MaxLOD-1
+    {  // MaxLOD-1 - MinLOD-1
         {0, 1, 1, 2, 3, 3},
+    },
+    {  // MaxLOD-0 - MinLOD-0
+        {0, 1, 1, 2, 3, 3, 4, 4, 4, 5, 6, 7, 7, 8, 8}
     }
 };
 const Vector<Matrix<std::uint16_t> > conditionalOutputIndices = {
-    {  // MaxLOD-0
+    {  // MaxLOD-0 - MinLOD-1
         {0, 1, 1, 2, 3, 3, 4, 4, 4, 5, 6, 7, 7, 8, 8},
         {0, 1, 1, 2, 3, 3},
     },
-    {  // MaxLOD-1
+    {  // MaxLOD-1 - MinLOD-1
         {0, 1, 1, 2, 3, 3},
+    },
+    {  // MaxLOD-0 - MinLOD-0
+        {0, 1, 1, 2, 3, 3, 4, 4, 4, 5, 6, 7, 7, 8, 8}
     }
 };
 const Vector<Matrix<float> > conditionalFromValues = {
-    {  // MaxLOD-0
+    {  // MaxLOD-0 - MinLOD-1
         {0.0f, 0.0f, 0.6f, 0.4f, 0.1f, 0.7f, 0.0f, 0.4f, 0.7f, 0.5f, 0.0f, 0.1f, 0.6f, 0.2f, 0.0f},
         {0.0f, 0.0f, 0.6f, 0.4f, 0.1f, 0.7f}
     },
-    {  // MaxLOD-1
+    {  // MaxLOD-1 - MinLOD-1
         {0.0f, 0.0f, 0.6f, 0.4f, 0.1f, 0.7f}
+    },
+    {  // MaxLOD-0 - MinLOD-0
+        {0.0f, 0.0f, 0.6f, 0.4f, 0.1f, 0.7f, 0.0f, 0.4f, 0.7f, 0.5f, 0.0f, 0.1f, 0.6f, 0.2f, 0.0f}
     }
 };
 const Vector<Matrix<float> > conditionalToValues = {
-    {  // MaxLOD-0
+    {  // MaxLOD-0 - MinLOD-1
         {1.0f, 0.6f, 1.0f, 0.9f, 0.7f, 1.0f, 0.4f, 0.7f, 1.0f, 1.0f, 1.0f, 0.6f, 1.0f, 0.8f, 1.0f},
         {1.0f, 0.6f, 1.0f, 0.9f, 0.7f, 1.0f}
     },
-    {  // MaxLOD-1
+    {  // MaxLOD-1 - MinLOD-1
         {1.0f, 0.6f, 1.0f, 0.9f, 0.7f, 1.0f}
+    },
+    {  // MaxLOD-1 - MinLOD-1
+        {1.0f, 0.6f, 1.0f, 0.9f, 0.7f, 1.0f, 0.4f, 0.7f, 1.0f, 1.0f, 1.0f, 0.6f, 1.0f, 0.8f, 1.0f}
     }
-
 };
 const Vector<Matrix<float> > conditionalSlopeValues = {
-    {  // MaxLOD-0
+    {  // MaxLOD-0 - MinLOD-1
         {1.0f, 0.9f, 0.9f, 0.8f, 0.7f, 0.7f, 0.6f, 0.6f, 0.6f, 0.5f, 0.6f, 0.7f, 0.7f, 0.8f, 0.9f},
         {1.0f, 0.9f, 0.9f, 0.8f, 0.7f, 0.7f}
     },
-    {  // MaxLOD-1
+    {  // MaxLOD-1 - MinLOD-1
         {1.0f, 0.9f, 0.9f, 0.8f, 0.7f, 0.7f}
+    },
+    {  // MaxLOD-0 - MinLOD-0
+        {1.0f, 0.9f, 0.9f, 0.8f, 0.7f, 0.7f, 0.6f, 0.6f, 0.6f, 0.5f, 0.6f, 0.7f, 0.7f, 0.8f, 0.9f}
     }
 };
 const Vector<Matrix<float> > conditionalCutValues = {
-    {  // MaxLOD-0
+    {  // MaxLOD-0 - MinLOD-1
         {0.0f, 0.5f, 0.5f, 0.4f, 0.3f, 0.3f, 1.0f, 1.0f, 1.0f, 0.2f, 0.4f, 0.8f, 0.8f, 1.0f, 0.2f},
         {0.0f, 0.5f, 0.5f, 0.4f, 0.3f, 0.3f}
     },
-    {  // MaxLOD-1
+    {  // MaxLOD-1 - MinLOD-1
         {0.0f, 0.5f, 0.5f, 0.4f, 0.3f, 0.3f}
+    },
+    {  // MaxLOD-0 - MinLOD-0
+        {0.0f, 0.5f, 0.5f, 0.4f, 0.3f, 0.3f, 1.0f, 1.0f, 1.0f, 0.2f, 0.4f, 0.8f, 0.8f, 1.0f, 0.2f}
     }
 };
 // Behavior->PSDs
@@ -1257,116 +1560,148 @@ const Vector<float> psdValues = {
 };
 // Behavior->Joints
 const Vector<std::uint16_t> jointRowCount = {
-    81u,  // MaxLOD-0
-    54u  // MaxLOD-1
+    81u,  // MaxLOD-0 - MinLOD-1
+    54u,  // MaxLOD-1 - MinLOD-1
+    81u  // MaxLOD-0 - MinLOD-0
 };
 const std::uint16_t jointColumnCount = 10u;
 const Vector<Matrix<std::uint16_t> > jointVariableIndices = {
-    {  // MaxLOD-0
+    {  // MaxLOD-0 - MinLOD-1
         {2, 3, 5, 18, 20, 36, 38, 55, 56, 63, 45, 46, 71},
         {2, 3, 5, 18, 20, 55, 56}
     },
-    {  // MaxLOD-1
+    {  // MaxLOD-1 - MinLOD-1
         {2, 3, 5, 18, 20, 37, 38}
+    },
+    {  // MaxLOD-0 - MinLOD-0
+        {2, 3, 5, 18, 20, 36, 38, 55, 56, 63, 45, 46, 71}
     }
 };
 const Vector<Matrix<std::uint16_t> > jointGroupLODs = {
     {  // Joint Group 0
-        {3, 3},  // MaxLOD-0
-        {3}  // MaxLOD-1
+        {3, 3},  // MaxLOD-0 - MaxLOD-1
+        {3},  // MaxLOD-1 - MaxLOD-1
+        {3},  // MaxLOD-0 - MaxLOD-1
     },
     {  // Joint group 1
-        {4, 2},  // MaxLOD-0
-        {2}  // MaxLOD-1
+        {4, 2},  // MaxLOD-0 - MaxLOD-1
+        {2},  // MaxLOD-1 - MaxLOD-1
+        {4}  // MaxLOD-0 - MaxLOD-0
     },
     {  // Joint group 2
-        {3, 2},  // MaxLOD-0
-        {2}  // MaxLOD-1
+        {3, 2},  // MaxLOD-0 - MinLOD-1
+        {2},  // MaxLOD-1 - MinLOD-1
+        {3}  // MaxLOD-0 - MinLOD-0
     },
     {  // Joint group 3
-        {3, 0},  // MaxLOD-0
-        {0}  // MaxLOD-1
+        {3, 0},  // MaxLOD-0 - MinLOD-1
+        {0},  // MaxLOD-1 - MinLOD-1
+        {3}  // MaxLOD-0 - MinLOD-0
     }
 };
 const Vector<Vector<Matrix<std::uint16_t> > > jointGroupInputIndices = {
     {  // Joint Group 0
-        {  // MaxLOD-0
+        {  // MaxLOD-0 - MaxLOD-1
             {0, 1, 2, 3, 6, 7, 8},
             {0, 1, 2, 3, 6, 7, 8}
         },
-        {  // MaxLOD-1
+        {  // MaxLOD-1 - MaxLOD-1
+            {0, 1, 2, 3, 6, 7, 8}
+        },
+        {  // MaxLOD-0 - MaxLOD-0
             {0, 1, 2, 3, 6, 7, 8}
         }
     },
     {  // Joint Group 1
-        {  // MaxLOD-0
+        {  // MaxLOD-0 - MaxLOD-1
             {3, 4, 7, 8, 9},
             {3, 4, 7, 8, 9}
         },
-        {  // MaxLOD-1
+        {  // MaxLOD-1 - MaxLOD-1
+            {3, 4, 7, 8, 9}
+        },
+        {  // MaxLOD-0 - MaxLOD-0
             {3, 4, 7, 8, 9}
         }
     },
     {  // Joint Group 2
-        {  // MaxLOD-0
-            {4, 5, 8, 9},
+        {  // MaxLOD-0 - MaxLOD-1
             {4, 5, 8, 9},
+            {4, 5, 8, 9}
+        },
+        {  // MaxLOD-1 - MaxLOD-1
+            {4, 5, 8, 9}
         },
-        {  // MaxLOD-1
+        {  // MaxLOD-0 - MaxLOD-0
             {4, 5, 8, 9}
         }
     },
     {  // Joint Group 3
-        {  // MaxLOD-0
-            {2, 5, 6, 8},
+        {  // MaxLOD-0 - MaxLOD-1
             {2, 5, 6, 8},
+            {2, 5, 6, 8}
         },
-        {  // MaxLOD-1
+        {  // MaxLOD-1 - MaxLOD-1
             {}
+        },
+        {  // MaxLOD-0 - MaxLOD-0
+            {2, 5, 6, 8}
         }
     }
 };
 const Vector<Vector<Matrix<std::uint16_t> > > jointGroupOutputIndices = {
     {  // Joint Group 0
-        {  // MaxLOD-0
+        {  // MaxLOD-0 - MaxLOD-1
             {2, 3, 5},
             {2, 3, 5}
         },
-        {  // MaxLOD-1
+        {  // MaxLOD-1 - MaxLOD-1
             {2, 3, 5}
+        },
+        {  // MaxLOD-0 - MaxLOD-0
+            {2, 3, 5},
         }
     },
     {  // Joint Group 1
-        {  // MaxLOD-0
+        {  // MaxLOD-0 - MaxLOD-1
             {18, 20, 36, 38},
             {18, 20}
         },
-        {  // MaxLOD-1
+        {  // MaxLOD-1 - MaxLOD-1
             {18, 20}
+        },
+        {  // MaxLOD-0 - MaxLOD-0
+            {18, 20, 36, 38}
         }
     },
     {  // Joint Group 2
-        {  // MaxLOD-0
+        {  // MaxLOD-0 - MaxLOD-1
             {55, 56, 63},
             {55, 56}
         },
-        {  // MaxLOD-1
+        {  // MaxLOD-1 - MaxLOD-1
             {37, 38}
+        },
+        {  // MaxLOD-0 - MaxLOD-0
+            {55, 56, 63}
         }
     },
     {  // Joint Group 3
-        {  // MaxLOD-0
+        {  // MaxLOD-0 - MaxLOD-1
             {45, 46, 71},
             {}
         },
-        {  // MaxLOD-1
+        {  // MaxLOD-1 - MaxLOD-1
             {}
+        },
+        {  // MaxLOD-0 - MaxLOD-0
+            {45, 46, 71}
         }
     }
 };
 const Vector<Vector<Matrix<float> > > jointGroupValues = {
     {  // Joint Group 0
-        {  // MaxLOD-0
+        {  // MaxLOD-0 - MaxLOD-1
             {
                 0.00f, 0.05f, 0.10f, 0.15f, 0.20f, 0.25f, 0.30f,
                 0.35f, 0.40f, 0.45f, 0.50f, 0.55f, 0.60f, 0.65f,
@@ -1378,7 +1713,14 @@ const Vector<Vector<Matrix<float> > > jointGroupValues = {
                 0.70f, 0.75f, 0.80f, 0.85f, 0.90f, 0.95f, 1.00f
             }
         },
-        {  // MaxLOD-1
+        {  // MaxLOD-1 - MinLOD-1
+            {
+                0.00f, 0.05f, 0.10f, 0.15f, 0.20f, 0.25f, 0.30f,
+                0.35f, 0.40f, 0.45f, 0.50f, 0.55f, 0.60f, 0.65f,
+                0.70f, 0.75f, 0.80f, 0.85f, 0.90f, 0.95f, 1.00f
+            }
+        },
+        {  // MaxLOD-0 - MinLOD-0
             {
                 0.00f, 0.05f, 0.10f, 0.15f, 0.20f, 0.25f, 0.30f,
                 0.35f, 0.40f, 0.45f, 0.50f, 0.55f, 0.60f, 0.65f,
@@ -1387,7 +1729,7 @@ const Vector<Vector<Matrix<float> > > jointGroupValues = {
         }
     },
     {  // Joint group 1
-        {  // MaxLOD-0
+        {  // MaxLOD-0 - MaxLOD-1
             {
                 0.01f, 0.02f, 0.03f, 0.04f, 0.05f,
                 0.06f, 0.07f, 0.08f, 0.09f, 0.10f,
@@ -1399,15 +1741,23 @@ const Vector<Vector<Matrix<float> > > jointGroupValues = {
                 0.06f, 0.07f, 0.08f, 0.09f, 0.10f
             }
         },
-        {  // MaxLOD-1
+        {  // MaxLOD-1 - MinLOD-1
             {
                 0.01f, 0.02f, 0.03f, 0.04f, 0.05f,
                 0.06f, 0.07f, 0.08f, 0.09f, 0.10f
             }
+        },
+        {  // MaxLOD-0 - MinLOD-0
+            {
+                0.01f, 0.02f, 0.03f, 0.04f, 0.05f,
+                0.06f, 0.07f, 0.08f, 0.09f, 0.10f,
+                0.11f, 0.12f, 0.13f, 0.14f, 0.15f,
+                0.16f, 0.17f, 0.18f, 0.19f, 0.20f
+            }
         }
     },
     {  // Joint group 2
-        {  // MaxLOD-0
+        {  // MaxLOD-0 - MaxLOD-1
             {
                 0.31f, 0.36f, 0.42f, 0.47f,
                 0.53f, 0.58f, 0.64f, 0.69f,
@@ -1418,15 +1768,22 @@ const Vector<Vector<Matrix<float> > > jointGroupValues = {
                 0.53f, 0.58f, 0.64f, 0.69f
             }
         },
-        {  // MaxLOD-1
+        {  // MaxLOD-1 - MinLOD-1
             {
                 0.31f, 0.36f, 0.42f, 0.47f,
                 0.53f, 0.58f, 0.64f, 0.69f
             }
+        },
+        {  // MaxLOD-0 - MinLOD-0
+            {
+                0.31f, 0.36f, 0.42f, 0.47f,
+                0.53f, 0.58f, 0.64f, 0.69f,
+                0.75f, 0.80f, 0.86f, 0.91f
+            }
         }
     },
     {  // Joint group 3
-        {  // MaxLOD-0
+        {  // MaxLOD-0 - MaxLOD-1
             {
                 0.31f, 0.36f, 0.42f, 0.47f,
                 0.53f, 0.58f, 0.64f, 0.69f,
@@ -1435,93 +1792,122 @@ const Vector<Vector<Matrix<float> > > jointGroupValues = {
             {
             }
         },
-        {  // MaxLOD-1
+        {  // MaxLOD-1 - MinLOD-1
             {
             }
+        },
+        {  // MaxLOD-0 - MinLOD-0
+            {
+                0.31f, 0.36f, 0.42f, 0.47f,
+                0.53f, 0.58f, 0.64f, 0.69f,
+                0.75f, 0.80f, 0.86f, 0.91f
+            }
         }
     }
 };
 const Vector<Vector<Matrix<std::uint16_t> > > jointGroupJointIndices = {
     {  // Joint Group 0
-        {  // MaxLOD-0
+        {  // MaxLOD-0 - MaxLOD-1
             {0},
             {0}
         },
-        {  // MaxLOD-1
+        {  // MaxLOD-1 - MinLOD-1
+            {0}
+        },
+        {  // MaxLOD-0 - MinLOD-0
             {0}
         }
     },
     {  // Joint Group 1
-        {  // MaxLOD-0
+        {  // MaxLOD-0 - MaxLOD-1
             {2, 4},
             {2}
         },
-        {  // MaxLOD-1
+        {  // MaxLOD-1 - MinLOD-1
             {2}
+        },
+        {  // MaxLOD-0 - MinLOD-0
+            {2, 4}
         }
     },
     {  // Joint Group 2
-        {  // MaxLOD-0
+        {  // MaxLOD-0 - MaxLOD-1
             {6, 7},
             {6}
         },
-        {  // MaxLOD-1
+        {  // MaxLOD-1 - MinLOD-1
             {4}
+        },
+        {  // MaxLOD-0 - MinLOD-0
+            {6, 7}
         }
     },
     {  // Joint Group 3
-        {  // MaxLOD-0
+        {  // MaxLOD-0 - MaxLOD-1
             {5, 7},
             {}
         },
-        {  // MaxLOD-1
+        {  // MaxLOD-1 - MinLOD-1
             {}
+        },
+        {  // MaxLOD-0 - MinLOD-0
+            {5, 7}
         }
     }
 };
 // Behavior->BlendShapes
 const Matrix<std::uint16_t> blendShapeLODs = {
     {
-        {7, 4},  // MaxLOD-0
-        {4}  // MaxLOD-1
+        {7, 4},  // MaxLOD-0 - MaxLOD-1
+        {4},  // MaxLOD-1 - MinLOD-1
+        {7}  // MaxLOD-0 - MinLOD-0
     }
 };
 const Vector<Matrix<std::uint16_t> > blendShapeInputIndices = {
-    {  // MaxLOD-0
+    {  // MaxLOD-0 - MaxLOD-1
         {0, 1, 2, 3, 6, 7, 8},
         {0, 1, 2, 3}
     },
-    {  // MaxLOD-1
+    {  // MaxLOD-1 - MinLOD-1
         {0, 1, 2, 3}
+    },
+    {  // MaxLOD-0 - MinLOD-0
+        {0, 1, 2, 3, 6, 7, 8}
     }
 };
 const Vector<Matrix<std::uint16_t> > blendShapeOutputIndices = {
-    {  // MaxLOD-0
+    {  // MaxLOD-0 - MaxLOD-1
         {0, 1, 2, 3, 6, 7, 8},
         {0, 1, 2, 3}
     },
-    {  // MaxLOD-1
+    {  // MaxLOD-1 - MinLOD-1
         {0, 1, 2, 3}
+    },
+    {  // MaxLOD-0 - MinLOD-0
+        {0, 1, 2, 3, 6, 7, 8}
     }
 };
 // Behavior->AnimatedMaps
 const Vector<std::uint16_t> animatedMapCount = {
-    10,  // MaxLOD-0
-    4  // MaxLOD-1
+    10,  // MaxLOD-0 - MaxLOD-1
+    4,  // MaxLOD-1 - MinLOD-1
+    10  // MaxLOD-0 - MinLOD-0
 };
 const Matrix<std::uint16_t> animatedMapLODs = {
     {
-        {15, 6},  // MaxLOD-0
-        {6}  // MaxLOD-1
+        {15, 6},  // MaxLOD-0 - MaxLOD-1
+        {6},  // MaxLOD-1 - MinLOD-1
+        {15}  // MaxLOD-0 - MinLOD-0
     }
 };
 // Geometry
 const Vector<std::uint32_t> meshCount = {
-    3u,  // MaxLOD-0
-    1u  // MaxLOD-1
+    3u,  // MaxLOD-0 - MaxLOD-1
+    1u,  // MaxLOD-1 - MinLOD-1
+    2u  // MaxLOD-0 - MinLOD-0
 };
 const Vector<Matrix<Vector3> > vertexPositions = {
-    {  // MaxLOD-0
+    {  // MaxLOD-0 - MaxLOD-1
         {  // Mesh-0
             {7.0f, 7.0f, 7.0f},
             {8.0f, 8.0f, 8.0f},
@@ -1538,16 +1924,28 @@ const Vector<Matrix<Vector3> > vertexPositions = {
             {3.0f, 3.0f, 3.0f}
         }
     },
-    {  // MaxLOD-1
+    {  // MaxLOD-1 - MinLOD-1
         {  // Mesh-0 (Mesh-2 under MaxLOD-0)
             {1.0f, 1.0f, 1.0f},
             {2.0f, 2.0f, 2.0f},
             {3.0f, 3.0f, 3.0f}
         }
+    },
+    {  // MaxLOD-0 - MinLOD-0
+        {  // Mesh-0
+            {7.0f, 7.0f, 7.0f},
+            {8.0f, 8.0f, 8.0f},
+            {9.0f, 9.0f, 9.0f}
+        },
+        {  // Mesh-1
+            {4.0f, 4.0f, 4.0f},
+            {5.0f, 5.0f, 5.0f},
+            {6.0f, 6.0f, 6.0f}
+        }
     }
 };
 const Vector<Matrix<TextureCoordinate> > vertexTextureCoordinates = {
-    {  // MaxLOD-0
+    {  // MaxLOD-0 - MaxLOD-1
         {  // Mesh-0
             {7.0f, 7.0f},
             {8.0f, 8.0f},
@@ -1564,16 +1962,28 @@ const Vector<Matrix<TextureCoordinate> > vertexTextureCoordinates = {
             {3.0f, 3.0f}
         }
     },
-    {  // MaxLOD-1
+    {  // MaxLOD-1 - MinLOD-1
         {  // Mesh-0 (Mesh-2 under MaxLOD-0)
             {1.0f, 1.0f},
             {2.0f, 2.0f},
             {3.0f, 3.0f}
         }
+    },
+    {  // MaxLOD-0 - MinLOD-0
+        {  // Mesh-0
+            {7.0f, 7.0f},
+            {8.0f, 8.0f},
+            {9.0f, 9.0f}
+        },
+        {  // Mesh-1
+            {4.0f, 4.0f},
+            {5.0f, 5.0f},
+            {6.0f, 6.0f}
+        }
     }
 };
 const Vector<Matrix<Vector3> > vertexNormals = {
-    {  // MaxLOD-0
+    {  // MaxLOD-0 - MaxLOD-1
         {  // Mesh-0
             {7.0f, 7.0f, 7.0f},
             {8.0f, 8.0f, 8.0f},
@@ -1590,16 +2000,28 @@ const Vector<Matrix<Vector3> > vertexNormals = {
             {3.0f, 3.0f, 3.0f}
         }
     },
-    {  // MaxLOD-1
+    {  // MaxLOD-1 - MinLOD-1
         {  // Mesh-0 (Mesh-2 under MaxLOD-0)
             {1.0f, 1.0f, 1.0f},
             {2.0f, 2.0f, 2.0f},
             {3.0f, 3.0f, 3.0f}
         }
+    },
+    {  // MaxLOD-0 - MinLOD-0
+        {  // Mesh-0
+            {7.0f, 7.0f, 7.0f},
+            {8.0f, 8.0f, 8.0f},
+            {9.0f, 9.0f, 9.0f}
+        },
+        {  // Mesh-1
+            {4.0f, 4.0f, 4.0f},
+            {5.0f, 5.0f, 5.0f},
+            {6.0f, 6.0f, 6.0f}
+        }
     }
 };
 const Vector<Matrix<VertexLayout> > vertexLayouts = {
-    {  // MaxLOD-0
+    {  // MaxLOD-0 - MaxLOD-1
         {  // Mesh-0
             {0, 0, 0},
             {1, 1, 1},
@@ -1616,16 +2038,27 @@ const Vector<Matrix<VertexLayout> > vertexLayouts = {
             {2, 2, 2}
         }
     },
-    {  // MaxLOD-1
+    {  // MaxLOD-1 - MinLOD-1
         {  // Mesh-0 (Mesh-2 under MaxLOD-0)
             {0, 0, 0},
             {1, 1, 1},
             {2, 2, 2}
         }
-    }
-};
+    },
+    {  // MaxLOD-0 - MinLOD-0
+        {  // Mesh-0
+            {0, 0, 0},
+            {1, 1, 1},
+            {2, 2, 2}
+        },
+        {  // Mesh-1
+            {0, 0, 0},
+            {1, 1, 1},
+            {2, 2, 2}
+        }
+    }};
 const Matrix<Matrix<std::uint32_t> > faces = {
-    {  // MaxLOD-0
+    {  // MaxLOD-0 - MaxLOD-1
         {  // Mesh-0
             {0, 1, 2}
         },
@@ -1636,50 +2069,73 @@ const Matrix<Matrix<std::uint32_t> > faces = {
             {0, 1, 2}
         }
     },
-    {  // MaxLOD-1
+    {  // MaxLOD-1 - MinLOD-1
         {  // Mesh-0 (Mesh-2 under MaxLOD-0)
             {0, 1, 2}
         }
-    }
-};
+    },
+    {  // MaxLOD-0 - MinLOD-0
+        {  // Mesh-0
+            {0, 1, 2}
+        },
+        {  // Mesh-1
+            {0, 1, 2}
+        }
+    }};
 const Matrix<std::uint16_t> maxInfluencePerVertex = {
-    {  // MaxLOD-0
+    {  // MaxLOD-0 - MaxLOD-1
         8u,  // Mesh-0
         8u,  // Mesh-1
         8u  // Mesh-2
     },
-    {  // MaxLOD-1
+    {  // MaxLOD-1 - MinLOD-1
         8u  // Mesh-0 (Mesh-2 under MaxLOD-0)
+    },
+    {  // MaxLOD-0 - MinLOD-0
+        8u,  // Mesh-0
+        8u  // Mesh-1
     }
 };
 const Matrix<Matrix<float> > skinWeightsValues = {
-    {  // MaxLOD-0
+    {  // MaxLOD-0 - MinLOD-1
         {  // Mesh-0
-            {7.0f, 8.0f, 9.0f},
-            {7.0f, 8.0f},
-            {8.0f, 9.0f}
+            {0.7f, 0.1f, 0.2f},
+            {0.5f, 0.5f},
+            {0.4f, 0.6f}
         },
         {  // Mesh-1
-            {4.0f, 5.0f, 6.0f},
-            {4.0f, 5.0f},
-            {5.0f, 6.0f}
+            {0.4f, 0.3f, 0.3f},
+            {0.8f, 0.2f},
+            {0.1f, 0.9f}
         },
         {  // Mesh-2
-            {1.0f, 2.0f, 3.0f},
-            {1.0f, 2.0f},
-            {2.0f, 3.0f}
+            {0.1f, 0.3f, 0.6f},
+            {0.3f, 0.7f},
+            {0.2f, 0.8f}
         }
     },
-    {  // MaxLOD-1
+    {  // MaxLOD-1 - MinLOD-1
         {  // Mesh-0 (Mesh-2 under MaxLOD-0)
-            {1.0f, 2.0f, 3.0f},
-            {1.0f},
-            {3.0f}
+            {0.1f, 0.3f, 0.6f},
+            {1.0f},  // 0.3f normalized to 1.0f
+            {1.0f}  // 0.8f normalized to 1.0f
+        }
+    },
+    {  // MaxLOD-0 - MinLOD-0
+        {  // Mesh-0
+            {0.7f, 0.1f, 0.2f},
+            {0.5f, 0.5f},
+            {0.4f, 0.6f}
+        },
+        {  // Mesh-1
+            {0.4f, 0.3f, 0.3f},
+            {0.8f, 0.2f},
+            {0.1f, 0.9f}
         }
     }
 };
 const Matrix<Matrix<std::uint16_t> > skinWeightsJointIndices = {
-    {  // MaxLOD-0
+    {  // MaxLOD-0 - MaxLOD-1
         {  // Mesh-0
             {0, 1, 2},
             {3, 4},
@@ -1696,16 +2152,28 @@ const Matrix<Matrix<std::uint16_t> > skinWeightsJointIndices = {
             {5, 6}
         }
     },
-    {  // MaxLOD-1
+    {  // MaxLOD-1 - MinLOD-1
         {  // Mesh-0 (Mesh-2 under MaxLOD-0)
             {0, 1, 2},
             {3},
             {4}
         }
+    },
+    {  // MaxLOD-0 - MinLOD-0
+        {  // Mesh-0
+            {0, 1, 2},
+            {3, 4},
+            {5, 6}
+        },
+        {  // Mesh-1
+            {0, 1, 2},
+            {3, 4},
+            {5, 6}
+        }
     }
 };
 const Vector<Matrix<std::uint16_t> > correctiveBlendShapeIndices = {
-    {  // MaxLOD-0
+    {  // MaxLOD-0 - MaxLOD-1
         {  // Mesh-0
             6
         },
@@ -1716,14 +2184,22 @@ const Vector<Matrix<std::uint16_t> > correctiveBlendShapeIndices = {
             6, 7
         }
     },
-    {  // MaxLOD-1
+    {  // MaxLOD-1 - MinLOD-1
         {  // Mesh-0 (Mesh-2 under MaxLOD-0)
             7
         }
+    },
+    {  // MaxLOD-0 - MinLOD-0
+        {  // Mesh-0
+            6
+        },
+        {  // Mesh-1
+            6
+        }
     }
 };
 const Matrix<Matrix<Vector3> > correctiveBlendShapeDeltas = {
-    {  // MaxLOD-0
+    {  // MaxLOD-0 - MaxLOD-1
         {  // Mesh-0
             {  // Blendshape-0
                 {7.0f, 7.0f, 7.0f},
@@ -1750,17 +2226,33 @@ const Matrix<Matrix<Vector3> > correctiveBlendShapeDeltas = {
             }
         }
     },
-    {  // MaxLOD-1
+    {  // MaxLOD-1 - MinLOD-1
         {  // Mesh-0 (Mesh-2 under MaxLOD-0)
             {  // Blendshape-1
                 {4.0f, 4.0f, 4.0f},
                 {5.0f, 5.0f, 5.0f}
             }
         }
+    },
+    {  // MaxLOD-0 - MinLOD-0
+        {  // Mesh-0
+            {  // Blendshape-0
+                {7.0f, 7.0f, 7.0f},
+                {8.0f, 8.0f, 8.0f},
+                {9.0f, 9.0f, 9.0f}
+            }
+        },
+        {  // Mesh-1
+            {  // Blendshape-0
+                {4.0f, 4.0f, 4.0f},
+                {5.0f, 5.0f, 5.0f},
+                {6.0f, 6.0f, 6.0f}
+            }
+        }
     }
 };
 const Matrix<Matrix<std::uint32_t> > correctiveBlendShapeVertexIndices = {
-    {  // MaxLOD-0
+    {  // MaxLOD-0 - MaxLOD-1
         {  // Mesh-0
             {0, 1, 2},  // Blendshape-0
         },
@@ -1772,10 +2264,444 @@ const Matrix<Matrix<std::uint32_t> > correctiveBlendShapeVertexIndices = {
             {0, 2}  // Blendshape-1
         }
     },
-    {  // MaxLOD-1
+    {  // MaxLOD-1 - MinLOD-1
         {  // Mesh-0 (Mesh-2 under MaxLOD-0)
             {0, 2}  // Blendshape-1
         }
+    },
+    {  // MaxLOD-0 - MinLOD-0
+        {  // Mesh-0
+            {0, 1, 2},  // Blendshape-0
+        },
+        {  // Mesh-1
+            {0, 1, 2},  // Blendshape-0
+        }
+    }
+};
+// Machine learned behavior
+const Vector<String<char> > mlControlNames = {
+    "MA", "MB", "MC", "MD", "ME", "MF", "MG", "MH", "MI"
+};
+const Matrix<std::uint16_t> neuralNetworkIndicesPerLOD = {
+    {  // MaxLOD-0 - MaxLOD-1
+        0,  // Mesh-0 Region-0
+        1,  // Mesh-0 Region-1
+        2,  // Mesh-1 Region-0
+        3,  // Mesh-1 Region-1
+        4,  // Mesh-2 Region-0
+        5  // Mesh-2 Region-1
+    },
+    {  // MaxLOD-1 - MinLOD-1
+        0,  // Mesh-0 (Mesh-2 under MaxLOD-0) Region-0
+        1  // Mesh-0 (Mesh-2 under MaxLOD-0) Region-1
+    },
+    {  // MaxLOD-0 - MinLOD-0
+        0,  // Mesh-0 Region-0
+        1,  // Mesh-0 Region-1
+        2,  // Mesh-1 Region-0
+        3  // Mesh-1 Region-1
+    }
+};
+const VectorOfCharStringMatrix regionNames = {
+    {  // MaxLOD-0 - MinLOD-1
+        {  // Mesh-0
+            "RA", "RB"
+        },
+        {  // Mesh-1
+            "RC", "RD"
+        },
+        {  // Mesh-2
+            "RE", "RF"
+        }
+    },
+    {  // MaxLOD-1 - MinLOD-1
+        {  // Mesh-0 (Mesh-2 under MaxLOD-0)
+            "RE", "RF"
+        }
+    },
+    {  // MaxLOD-0 - MinLOD-0
+        {  // Mesh-0
+            "RA", "RB"
+        },
+        {  // Mesh-1
+            "RC", "RD"
+        }
+    }
+};
+const Matrix<Matrix<std::uint16_t> > neuralNetworkIndicesPerMeshRegion = {
+    {  // MaxLOD-0 - MaxLOD-1
+        {  // Mesh-0
+            {  // Region-0
+                0
+            },
+            {  // Region-1
+                1
+            }
+        },
+        {  // Mesh-1
+            {  // Region-0
+                2
+            },
+            {  // Region-1
+                3
+            }
+        },
+        {  // Mesh-2
+            {  // Region-0
+                4
+            },
+            {  // Region-1
+                5
+            }
+        }
+    },
+    {  // MaxLOD-1 - MinLOD-1
+        {  // Mesh-0 (Mesh-2 under MaxLOD-0)
+            {  // Region-0
+                0  // (4 under MaxLOD-0)
+            },
+            {  // Region-1
+                1  // (5 under MaxLOD-0)
+            }
+        }
+    },
+    {  // MaxLOD-0 - MinLOD-0
+        {  // Mesh-0
+            {  // Region-0
+                0
+            },
+            {  // Region-1
+                1
+            }
+        },
+        {  // Mesh-1
+            {  // Region-0
+                2
+            },
+            {  // Region-1
+                3
+            }
+        }
+    }
+};
+const Vector<Matrix<std::uint16_t> > neuralNetworkInputIndices = {
+    {  // MaxLOD-0 - MaxLOD-1
+        {  // Mesh-0 Region-0
+            0, 1
+        },
+        {  // Mesh-0 Region-1
+            2, 3
+        },
+        {  // Mesh-1 Region-0
+            4, 5
+        },
+        {  // Mesh-1 Region-1
+            6, 7
+        },
+        {  // Mesh-2 Region-0
+            8, 0
+        },
+        {  // Mesh-2 Region-1
+            4, 7
+        }
+    },
+    {  // MaxLOD-1 - MinLOD-1
+        {  // Mesh-0 (Mesh-2 under MaxLOD-0) Region-0
+            8, 0
+        },
+        {  // Mesh-0 (Mesh-2 under MaxLOD-0) Region-1
+            4, 7
+        }
+    },
+    {  // MaxLOD-0 - MinLOD-0
+        {  // Mesh-0 Region-0
+            0, 1
+        },
+        {  // Mesh-0 Region-1
+            2, 3
+        },
+        {  // Mesh-1 Region-0
+            4, 5
+        },
+        {  // Mesh-1 Region-1
+            6, 7
+        }
+    }
+};
+const Vector<Matrix<std::uint16_t> > neuralNetworkOutputIndices = {
+    {  // MaxLOD-0 - MaxLOD-1
+        {  // Mesh-0 Region-0
+            9
+        },
+        {  // Mesh-0 Region-1
+            10
+        },
+        {  // Mesh-1 Region-0
+            11
+        },
+        {  // Mesh-1 Region-1
+            12
+        },
+        {  // Mesh-2 Region-0
+            13
+        },
+        {  // Mesh-2 Region-1
+            14
+        }
+    },
+    {  // MaxLOD-1 - MinLOD-1
+        {  // Mesh-0 (Mesh-2 under MaxLOD-0) Region-0
+            13
+        },
+        {  // Mesh-0 (Mesh-2 under MaxLOD-0) Region-1
+            14
+        }
+    },
+    {  // MaxLOD-0 - MinLOD-0
+        {  // Mesh-0 Region-0
+            9
+        },
+        {  // Mesh-0 Region-1
+            10
+        },
+        {  // Mesh-1 Region-0
+            11
+        },
+        {  // Mesh-1 Region-1
+            12
+        }
+    }
+};
+const Matrix<std::uint16_t> neuralNetworkLayerCount = {
+    {  // MaxLOD-0 - MaxLOD-1
+        2,  // Mesh-0 Region-0
+        2,  // Mesh-0 Region-1
+        2,  // Mesh-1 Region-0
+        2,  // Mesh-1 Region-1
+        2,  // Mesh-2 Region-0
+        2  // Mesh-2 Region-1
+    },
+    {  // MaxLOD-1 - MinLOD-1
+        2,  // Mesh-0 (Mesh-2 under MaxLOD-0) Region-0
+        2  // Mesh-0 (Mesh-2 under MaxLOD-0) Region-1
+    },
+    {  // MaxLOD-0 - MinLOD-0
+        2,  // Mesh-0 Region-0
+        2,  // Mesh-0 Region-1
+        2,  // Mesh-1 Region-0
+        2  // Mesh-1 Region-1
+    }
+};
+const Vector<Matrix<std::uint16_t> > neuralNetworkActivationFunction = {
+    {  // MaxLOD-0 - MaxLOD-1
+        {  // Mesh-0 Region-0
+            1, 1
+        },
+        {  // Mesh-0 Region-1
+            1, 1
+        },
+        {  // Mesh-1  Region-0
+            1, 1
+        },
+        {  // Mesh-1 Region-1
+            1, 1
+        },
+        {  // Mesh-2 Region-0
+            1, 1
+        },
+        {  // Mesh-2 Region-1
+            1, 1
+        }
+    },
+    {  // MaxLOD-1 - MinLOD-1
+        {  // Mesh-0 (Mesh-2 under MaxLOD-0) Region-0
+            1, 1
+        },
+        {  // Mesh-0 (Mesh-2 under MaxLOD-0) Region-1
+            1, 1
+        }
+    },
+    {  // MaxLOD-0 - MinLOD-0
+        {  // Mesh-0  Region-0
+            1, 1
+        },
+        {  // Mesh-0 Region-1
+            1, 1
+        },
+        {  // Mesh-1 Region-0
+            1, 1
+        },
+        {  // Mesh-1 Region-1
+            1, 1
+        }
+    }
+};
+const Matrix<Matrix<float> > neuralNetworkActivationFunctionParameters = {
+    {  // MaxLOD-0 - MaxLOD-1
+        {  // Mesh-0  Region-0
+            {0.5f},  // Layer-0
+            {0.5f}  // Layer-1
+        },
+        {  // Mesh-0 Region-1
+            {1.0f},  // Layer-0
+            {1.0f}  // Layer-1
+        },
+        {  // Mesh-1 Region-0
+            {1.0f},  // Layer-0
+            {1.0f}  // Layer-1
+        },
+        {  // Mesh-1 Region-1
+            {0.5f},  // Layer-0
+            {0.5f}  // Layer-1
+        },
+        {  // Mesh-2 Region-0
+            {0.5f},  // Layer-0
+            {0.5f}  // Layer-1
+        },
+        {  // Mesh-2 Region-1
+            {1.0f},  // Layer-0
+            {1.0f}  // Layer-1
+        }
+    },
+    {  // MaxLOD-1 - MinLOD-1
+        {  // Mesh-0 (Mesh-2 under MaxLOD-0) Region-0
+            {0.5f},  // Layer-0
+            {0.5f}  // Layer-1
+        },
+        {  // Mesh-0 (Mesh-2 under MaxLOD-0) Region-1
+            {1.0f},  // Layer-0
+            {1.0f}  // Layer-1
+        }
+    },
+    {  // MaxLOD-0 - MinLOD-0
+        {  // Mesh-0 Region-0
+            {0.5f},  // Layer-0
+            {0.5f}  // Layer-1
+        },
+        {  // Mesh-0 Region-1
+            {1.0f},  // Layer-0
+            {1.0f}  // Layer-1
+        },
+        {  // Mesh-1 Region-0
+            {1.0f},  // Layer-0
+            {1.0f}  // Layer-1
+        },
+        {  // Mesh-1 Region-1
+            {0.5f},  // Layer-0
+            {0.5f}  // Layer-1
+        }
+    }
+};
+const Matrix<Matrix<float> > neuralNetworkBiases = {
+    {  // MaxLOD-0 - MaxLOD-1
+        {  // Mesh-0 Region-0
+            {1.0f, 1.0f},  // Layer-0
+            {1.0f}  // Layer-1
+        },
+        {  // Mesh-0 Region-1
+            {0.5f, 0.5f},  // Layer-0
+            {0.5f}  // Layer-1
+        },
+        {  // Mesh-1 Region-0
+            {0.5f, 0.5f},  // Layer-0
+            {0.5f}  // Layer-1
+        },
+        {  // Mesh-1 Region-1
+            {1.0f, 1.0f},  // Layer-0
+            {1.0f}  // Layer-1
+        },
+        {  // Mesh-2 Region-0
+            {1.0f, 1.0f},  // Layer-0
+            {1.0f}  // Layer-1
+        },
+        {  // Mesh-2 Region-1
+            {0.5f, 0.5f},  // Layer-0
+            {0.5f}  // Layer-1
+        }
+    },
+    {  // MaxLOD-1 - MinLOD-1
+        {  // Mesh-0 (Mesh-2 under MaxLOD-0) Region-0
+            {1.0f, 1.0f},  // Layer-0
+            {1.0f}  // Layer-1
+        },
+        {  // Mesh-0 (Mesh-2 under MaxLOD-0) Region-1
+            {0.5f, 0.5f},  // Layer-0
+            {0.5f}  // Layer-1
+        }
+    },
+    {  // MaxLOD-0 - MinLOD-0
+        {  // Mesh-0 Region-0
+            {1.0f, 1.0f},  // Layer-0
+            {1.0f}  // Layer-1
+        },
+        {  // Mesh-0 Region-1
+            {0.5f, 0.5f},  // Layer-0
+            {0.5f}  // Layer-1
+        },
+        {  // Mesh-1 Region-0
+            {0.5f, 0.5f},  // Layer-0
+            {0.5f}  // Layer-1
+        },
+        {  // Mesh-1 Region-1
+            {1.0f, 1.0f},  // Layer-0
+            {1.0f}  // Layer-1
+        }
+    }
+};
+const Matrix<Matrix<float> > neuralNetworkWeights = {
+    {  // MaxLOD-0 - MaxLOD-1
+        {  // Mesh-0 Region-0
+            {0.5f, 0.5f, 0.5f, 0.5f},  // Layer-0
+            {0.5f, 0.5f}  // Layer-1
+        },
+        {  // Mesh-0 Region-1
+            {1.0f, 1.0f, 1.0f, 1.0f},  // Layer-0
+            {1.0f, 1.0f}  // Layer-1
+        },
+        {  // Mesh-1 Region-0
+            {1.0f, 1.0f, 1.0f, 1.0f},  // Layer-0
+            {1.0f, 1.0f}  // Layer-1
+        },
+        {  // Mesh-1 Region-1
+            {0.5f, 0.5f, 0.5f, 0.5f},  // Layer-0
+            {0.5f, 0.5f}  // Layer-1
+        },
+        {  // Mesh-2 Region-0
+            {0.5f, 0.5f, 0.5f, 0.5f},  // Layer-0
+            {0.5f, 0.5f}  // Layer-1
+        },
+        {  // Mesh-2 Region-1
+            {1.0f, 1.0f, 1.0f, 1.0f},  // Layer-0
+            {1.0f, 1.0f}  // Layer-1
+        }
+    },
+    {  // MaxLOD-1 - MinLOD-1
+        {  // Mesh-0 (Mesh-2 under MaxLOD-0) Region-0
+            {0.5f, 0.5f, 0.5f, 0.5f},  // Layer-0
+            {0.5f, 0.5f}  // Layer-1
+        },
+        {  // Mesh-0 Region-1
+            {1.0f, 1.0f, 1.0f, 1.0f},  // Layer-0
+            {1.0f, 1.0f}  // Layer-1
+        }
+    },
+    {  // MaxLOD-0 - MinLOD-0
+        {  // Mesh-0 Region-0
+            {0.5f, 0.5f, 0.5f, 0.5f},  // Layer-0
+            {0.5f, 0.5f}  // Layer-1
+        },
+        {  // Mesh-0 Region-1
+            {1.0f, 1.0f, 1.0f, 1.0f},  // Layer-0
+            {1.0f, 1.0f}  // Layer-1
+        },
+        {  // Mesh-1 Region-0
+            {1.0f, 1.0f, 1.0f, 1.0f},  // Layer-0
+            {1.0f, 1.0f}  // Layer-1
+        },
+        {  // Mesh-1 Region-1
+            {0.5f, 0.5f, 0.5f, 0.5f},  // Layer-0
+            {0.5f, 0.5f}  // Layer-1
+        }
     }
 };
 
diff --git a/Engine/Plugins/Animation/RigLogic/Source/RigLogicLibTest/Private/rltests/dna/DNAFixtures.h b/Engine/Plugins/Animation/RigLogic/Source/RigLogicLibTest/Private/rltests/dna/DNAFixtures.h
--- a/Engine/Plugins/Animation/RigLogic/Source/RigLogicLibTest/Private/rltests/dna/DNAFixtures.h
+++ b/Engine/Plugins/Animation/RigLogic/Source/RigLogicLibTest/Private/rltests/dna/DNAFixtures.h
@@ -2,8 +2,7 @@
 
 #pragma once
 
-#include <dna/Reader.h>
-#include <pma/TypeDefs.h>
+#include "riglogic/TypeDefs.h"
 
 namespace rltests {
 
@@ -12,9 +11,9 @@ using namespace pma;
 
 namespace raw {
 
-extern const unsigned char header[39];
+extern const unsigned char header[91];
 extern const unsigned char descriptor[87];
-extern const unsigned char definition[796];
+extern const unsigned char definition[794];
 extern const unsigned char controls[2];
 extern const unsigned char conditionals[324];
 extern const unsigned char psds[204];
@@ -22,9 +21,9 @@ extern const unsigned char joints[444];
 extern const unsigned char blendshapes[44];
 extern const unsigned char animatedmaps[8];
 extern const unsigned char geometry[1080];
-extern const unsigned char footer[3];
+extern const unsigned char machineLearnedBehavior[762];
 
-std::vector<unsigned char> getBytes();
+std::vector<char> getBytes();
 
 }  // namespace raw
 
@@ -43,8 +42,8 @@ extern const Vector<StringPair> metadata;
 extern const TranslationUnit translationUnit;
 extern const RotationUnit rotationUnit;
 extern const CoordinateSystem coordinateSystem;
-extern const std::uint16_t lodCount[2ul];
-extern const std::uint16_t maxLOD[2ul];
+extern const std::uint16_t lodCount[3ul];
+extern const std::uint16_t maxLODs[3ul];
 extern const String<char> complexity;
 extern const String<char> dbName;
 
@@ -63,7 +62,7 @@ extern const Vector<Matrix<Vector3> > neutralJointRotations;
 // Behavior
 extern const std::uint16_t guiControlCount;
 extern const std::uint16_t rawControlCount;
-extern const std::uint16_t psdCount;
+extern const std::uint16_t psdControlCount;
 extern const Vector<Matrix<std::uint16_t> > conditionalInputIndices;
 extern const Vector<Matrix<std::uint16_t> > conditionalOutputIndices;
 extern const Vector<Matrix<float> > conditionalFromValues;
@@ -101,6 +100,19 @@ extern const Vector<Matrix<std::uint16_t> > correctiveBlendShapeIndices;
 extern const Matrix<Matrix<Vector3> > correctiveBlendShapeDeltas;
 extern const Matrix<Matrix<std::uint32_t> > correctiveBlendShapeVertexIndices;
 
+// Machine learned behavior
+extern const Vector<String<char> > mlControlNames;
+extern const VectorOfCharStringMatrix regionNames;
+extern const Matrix<std::uint16_t> neuralNetworkIndicesPerLOD;
+extern const Matrix<Matrix<std::uint16_t> > neuralNetworkIndicesPerMeshRegion;
+extern const Vector<Matrix<std::uint16_t> > neuralNetworkInputIndices;
+extern const Vector<Matrix<std::uint16_t> > neuralNetworkOutputIndices;
+extern const Matrix<std::uint16_t> neuralNetworkLayerCount;
+extern const Vector<Matrix<std::uint16_t> > neuralNetworkActivationFunction;
+extern const Matrix<Matrix<float> > neuralNetworkActivationFunctionParameters;
+extern const Matrix<Matrix<float> > neuralNetworkBiases;
+extern const Matrix<Matrix<float> > neuralNetworkWeights;
+
 }  // namespace decoded
 
 }  // namespace rltests
diff --git a/Engine/Plugins/Animation/RigLogic/Source/RigLogicLibTest/Private/rltests/dna/FakeReader.h b/Engine/Plugins/Animation/RigLogic/Source/RigLogicLibTest/Private/rltests/dna/FakeReader.h
--- a/Engine/Plugins/Animation/RigLogic/Source/RigLogicLibTest/Private/rltests/dna/FakeReader.h
+++ b/Engine/Plugins/Animation/RigLogic/Source/RigLogicLibTest/Private/rltests/dna/FakeReader.h
@@ -23,6 +23,15 @@ class FakeReader : public Reader {
     public:
         ~FakeReader();
 
+        // HeaderReader methods
+        std::uint16_t getFileFormatGeneration() const override {
+            return static_cast<std::uint16_t>(2);
+        }
+
+        std::uint16_t getFileFormatVersion() const override {
+            return static_cast<std::uint16_t>(3);
+        }
+
         // DescriptorReader methods start
         StringView getName() const override {
             return {};
@@ -80,7 +89,7 @@ class FakeReader : public Reader {
             return {};
         }
 
-        // DefinitionReader methods start
+        // DefinitionReader methods
         std::uint16_t getGUIControlCount() const override {
             return {};
         }
@@ -209,7 +218,7 @@ class FakeReader : public Reader {
             return {};
         }
 
-        // BehaviorReader methods start
+        // BehaviorReader methods
         ConstArrayView<std::uint16_t> getGUIToRawInputIndices() const override {
             return {};
         }
@@ -326,7 +335,7 @@ class FakeReader : public Reader {
             return {};
         }
 
-        // GeometryReader methods start
+        // GeometryReader methods
         std::uint32_t getVertexPositionCount(std::uint16_t meshIndex) const override {
             return {};
         }
@@ -467,8 +476,72 @@ class FakeReader : public Reader {
             return {};
         }
 
-        void unload(DataLayer  /**/) override {
+        // MachineLearnedBehaviorReader methods
+        std::uint16_t getMLControlCount() const override {
+            return {};
+        }
+
+        StringView getMLControlName(std::uint16_t index) const override {
+            return {};
+        }
+
+        std::uint16_t getNeuralNetworkCount() const override {
+            return {};
+        }
+
+        std::uint16_t getNeuralNetworkIndexListCount() const override {
+            return {};
+        }
+
+        ConstArrayView<std::uint16_t> getNeuralNetworkIndicesForLOD(std::uint16_t lod) const override {
+            return {};
+        }
+
+        std::uint16_t getMeshRegionCount(std::uint16_t meshIndex) const override {
+            return {};
+        }
+
+        StringView getMeshRegionName(std::uint16_t meshIndex, std::uint16_t regionIndex) const override {
+            return {};
+        }
+
+        ConstArrayView<std::uint16_t> getNeuralNetworkIndicesForMeshRegion(std::uint16_t meshIndex,
+                                                                           std::uint16_t regionIndex) const override {
+            return {};
+        }
+
+        ConstArrayView<std::uint16_t> getNeuralNetworkInputIndices(std::uint16_t neuralNetIndex) const override {
+            return {};
+        }
+
+        ConstArrayView<std::uint16_t> getNeuralNetworkOutputIndices(std::uint16_t neuralNetIndex) const override {
+            return {};
+        }
+
+        std::uint16_t getNeuralNetworkLayerCount(std::uint16_t neuralNetIndex) const override {
+            return {};
+        }
+
+        ActivationFunction getNeuralNetworkLayerActivationFunction(std::uint16_t neuralNetIndex,
+                                                                   std::uint16_t layerIndex) const override {
+            return {};
+        }
+
+        ConstArrayView<float> getNeuralNetworkLayerActivationFunctionParameters(std::uint16_t neuralNetIndex,
+                                                                                std::uint16_t layerIndex) const override {
+            return {};
+        }
+
+        ConstArrayView<float> getNeuralNetworkLayerBiases(std::uint16_t neuralNetIndex, std::uint16_t layerIndex) const override {
+            return {};
+        }
+
+        ConstArrayView<float> getNeuralNetworkLayerWeights(std::uint16_t neuralNetIndex,
+                                                           std::uint16_t layerIndex) const override {
+            return {};
+        }
 
+        void unload(DataLayer  /*unused*/) override {
         }
 
 };
diff --git a/Engine/Plugins/Animation/RigLogic/Source/RigLogicLibTest/Private/rltests/joints/TestJointsFactory.cpp b/Engine/Plugins/Animation/RigLogic/Source/RigLogicLibTest/Private/rltests/joints/TestJointsFactory.cpp
--- a/Engine/Plugins/Animation/RigLogic/Source/RigLogicLibTest/Private/rltests/joints/TestJointsFactory.cpp
+++ b/Engine/Plugins/Animation/RigLogic/Source/RigLogicLibTest/Private/rltests/joints/TestJointsFactory.cpp
@@ -1,17 +1,19 @@
 // Copyright Epic Games, Inc. All Rights Reserved.
 
 #include "rltests/Defs.h"
-#include "rltests/joints/bpcm/JointFixturesBlock4.h"
+#include "rltests/joints/bpcm/BPCMFixturesBlock4.h"
 
 #include "riglogic/RigLogic.h"
+#include "riglogic/TypeDefs.h"
 #include "riglogic/joints/JointsFactory.h"
-
-#include <pma/resources/AlignedMemoryResource.h>
+#include "riglogic/riglogic/Configuration.h"
+#include "riglogic/riglogic/RigMetrics.h"
 
 TEST(ScalarJointsFactoryTest, NeutralJointsAreCopied) {
     pma::AlignedMemoryResource memRes;
     block4::CanonicalReader reader;
-    rl4::Configuration config{rl4::RigLogic::CalculationType::Scalar};
+    rl4::Configuration config{};
+    config.calculationType = rl4::CalculationType::Scalar;
     auto joints = rl4::JointsFactory::create(config, &reader, &memRes);
     const float expected[] = {
         0.0f, 1.0f, 2.0f, 3.0f, 4.0f, 5.0f, 1.0f, 1.0f, 1.0f,
diff --git a/Engine/Plugins/Animation/RigLogic/Source/RigLogicLibTest/Private/rltests/joints/bpcm/Assertions.h b/Engine/Plugins/Animation/RigLogic/Source/RigLogicLibTest/Private/rltests/joints/bpcm/Assertions.h
--- a/Engine/Plugins/Animation/RigLogic/Source/RigLogicLibTest/Private/rltests/joints/bpcm/Assertions.h
+++ b/Engine/Plugins/Animation/RigLogic/Source/RigLogicLibTest/Private/rltests/joints/bpcm/Assertions.h
@@ -4,7 +4,7 @@
 
 #include "rltests/Defs.h"
 
-#include "riglogic/joints/bpcm/JointsEvaluator.h"
+#include "riglogic/joints/bpcm/Evaluator.h"
 
 namespace rl4 {
 
