Update RigLogic to v9.1.1 and GeneSplicer to v9.4.0

diff --git a/Engine/Plugins/Animation/RigLogic/Source/RigLogicLibTest/Private/rltests/ml/cpu/FixturesBlock4.cpp b/Engine/Plugins/Animation/RigLogic/Source/RigLogicLibTest/Private/rltests/ml/cpu/FixturesBlock4.cpp
new file mode 100644
--- /dev/null
+++ b/Engine/Plugins/Animation/RigLogic/Source/RigLogicLibTest/Private/rltests/ml/cpu/FixturesBlock4.cpp
@@ -0,0 +1,1102 @@
+// Copyright Epic Games, Inc. All Rights Reserved.
+
+#include "rltests/ml/cpu/FixturesBlock4.h"
+
+#include "riglogic/TypeDefs.h"
+#include "riglogic/joints/bpcm/BPCMOutputInstance.h"
+#include "riglogic/joints/bpcm/Evaluator.h"
+#include "riglogic/types/Extent.h"
+
+namespace rltests {
+
+namespace ml {
+
+namespace block4 {
+
+using namespace rl4;
+
+#ifdef __clang__
+    #pragma clang diagnostic push
+    #pragma clang diagnostic ignored "-Wglobal-constructors"
+    #pragma clang diagnostic ignored "-Wexit-time-destructors"
+#endif
+
+namespace unoptimized {
+
+const std::uint16_t rawControlCount = 13;  // rltests::ml::block4::input::values.size()
+const std::uint16_t mlControlCount = 14;  // number of all mlbNetInputIndices elements
+const std::uint16_t lodCount = 4u;
+const std::uint16_t neuralNetworkCount = 4u;
+const Matrix<std::uint16_t> neuralNetworkIndicesPerLOD = {
+    {0u, 1u, 2u, 3u},
+    {0u, 1u, 2u},
+    {},
+    {}
+};
+const Matrix<std::uint16_t> mlbNetInputIndices = {
+    {  // Mesh-0 Region-0
+        1u, 4u
+    },
+    {  // Mesh-0 Region-1
+        0u, 2u, 3u
+    },
+    {  // Mesh-1 Region-0
+        5u, 7u, 8u, 9u, 12u
+    },
+    {  // Mesh-1 Region-1
+        6u, 10u, 11u
+    }
+};
+const Matrix<std::uint16_t> mlbNetOutputIndices = {
+    {  // Mesh-0 Region-0
+        13u
+    },
+    {  // Mesh-0 Region-1
+        14u, 15u, 16u
+    },
+    {  // Mesh-1 Region-0
+        17u, 18u, 19u, 20u, 21u, 22u
+    },
+    {  // Mesh-1 Region-1
+        23u, 24u, 25u, 26u
+    }
+};
+const Matrix<dna::ActivationFunction> mlbNetActivationFunctions = {
+    {  // Mesh-0 Region-0
+        dna::ActivationFunction::sigmoid,  // Layer 0
+        dna::ActivationFunction::sigmoid,  // Layer 1
+        dna::ActivationFunction::linear  // Layer 2
+    },
+    {  // Mesh-0 Region-1
+        dna::ActivationFunction::leakyrelu,  // Layer 0
+        dna::ActivationFunction::leakyrelu,  // Layer 1
+        dna::ActivationFunction::leakyrelu,  // Layer 2
+        dna::ActivationFunction::linear  // Layer 3
+    },
+    {  // Mesh-1 Region-0
+        dna::ActivationFunction::relu,  // Layer 0
+        dna::ActivationFunction::relu,  // Layer 1
+        dna::ActivationFunction::relu,  // Layer 2
+        dna::ActivationFunction::linear  // Layer 3
+    },
+    {  // Mesh-1 Region-1
+        dna::ActivationFunction::tanh,  // Layer 0
+        dna::ActivationFunction::linear  // Layer 1
+    }
+};
+const Vector<Matrix<float> > mlbNetActivationFunctionParameters = {
+    {  // Mesh-0 Region-0
+        {},  // Layer 0
+        {},  // Layer 1
+        {}  // Layer 2
+    },
+    {  // Mesh-0 Region-1
+        {0.5f},  // Layer 0
+        {0.5f},  // Layer 1
+        {0.5f},  // Layer 2
+        {}  // Layer 3
+    },
+    {  // Mesh-1 Region-0
+        {},  // Layer 0
+        {},  // Layer 1
+        {},  // Layer 2
+        {}  // Layer 3
+    },
+    {  // Mesh-1 Region-1
+        {},  // Layer 0
+        {}  // Layer 1
+    }
+};
+const Vector<Matrix<float> > mlbNetWeights = {
+    {  // Mesh-0 Region-0
+        {  // Layer 0 - Pad to block-4 => block-4
+            0.11f, 0.11f,  // Unit-0
+            0.11f, 0.11f  // Unit-1
+        },
+        {  // Layer 1 - No padding => block-4
+            0.22f, 0.22f,  // Unit-0
+            0.33f, 0.33f,  // Unit-1
+            0.44f, 0.44f,  // Unit-2
+            0.55f, 0.55f  // Unit-3
+        },
+        {  // Layer 2 - Pad to block-4 => block-4
+            0.66f, 0.66f, 0.66f, 0.66f  // Unit-0
+        }
+    },
+    {  // Mesh-0 Region-1
+        {  // Layer 0 - Pad to block-8 => block-8
+            0.66f, 0.66f, 0.66f,  // Unit-0
+            0.77f, 0.77f, 0.77f,  // Unit-1
+            0.88f, 0.88f, 0.88f,  // Unit-2
+            0.99f, 0.99f, 0.99f,  // Unit-3
+            1.0f, 1.0f, 1.0f,  // Unit-4
+            1.11f, 1.11f, 1.11f,  // Unit-5
+            1.22f, 1.22f, 1.22f  // Unit-6
+        },
+        {  // Layer 1 - No padding => block-8
+            1.33f, 1.33f, 1.33f, 1.33f, 1.33f, 1.33f, 1.33f,
+            1.44f, 1.44f, 1.44f, 1.44f, 1.44f, 1.44f, 1.44f,
+            1.55f, 1.55f, 1.55f, 1.55f, 1.55f, 1.55f, 1.55f,
+            1.66f, 1.66f, 1.66f, 1.66f, 1.66f, 1.66f, 1.66f,
+            1.77f, 1.77f, 1.77f, 1.77f, 1.77f, 1.77f, 1.77f,
+            1.88f, 1.88f, 1.88f, 1.88f, 1.88f, 1.88f, 1.88f,
+            1.99f, 1.99f, 1.99f, 1.99f, 1.99f, 1.99f, 1.99f,
+            2.0f, 2.0f, 2.0f, 2.0f, 2.0f, 2.0f, 2.0f
+        },
+        {  // Layer 2 - Pad to block-4 => block-8 + block-4
+            -0.09f, -0.09f, -0.09f, -0.09f, -0.09f, -0.09f, -0.09f, -0.09f,
+            -0.08f, -0.08f, -0.08f, -0.08f, -0.08f, -0.08f, -0.08f, -0.08f,
+            -0.07f, -0.07f, -0.07f, -0.07f, -0.07f, -0.07f, -0.07f, -0.07f,
+            -0.06f, -0.06f, -0.06f, -0.06f, -0.06f, -0.06f, -0.06f, -0.06f,
+            -0.05f, -0.05f, -0.05f, -0.05f, -0.05f, -0.05f, -0.05f, -0.05f,
+            -0.04f, -0.04f, -0.04f, -0.04f, -0.04f, -0.04f, -0.04f, -0.04f,
+            -0.03f, -0.03f, -0.03f, -0.03f, -0.03f, -0.03f, -0.03f, -0.03f,
+            -0.02f, -0.02f, -0.02f, -0.02f, -0.02f, -0.02f, -0.02f, -0.02f,
+            -0.01f, -0.01f, -0.01f, -0.01f, -0.01f, -0.01f, -0.01f, -0.01f
+        },
+        {  // Layer 3 - Pad to block-4 => block-4
+            -0.30f, -0.30f, -0.30f, -0.30f, -0.30f, -0.30f, -0.30f, -0.30f, -0.30f,
+            -0.31f, -0.31f, -0.31f, -0.31f, -0.31f, -0.31f, -0.31f, -0.31f, -0.31f,
+            -0.32f, -0.32f, -0.32f, -0.32f, -0.32f, -0.32f, -0.32f, -0.32f, -0.32f
+        }
+    },
+    {  // Mesh-1 Region-0
+        {  // Layer 0 - No padding => block-8 + block-4
+            0.30f, 0.30f, 0.30f, 0.30f, 0.30f,
+            0.31f, 0.31f, 0.31f, 0.31f, 0.31f,
+            0.32f, 0.32f, 0.32f, 0.32f, 0.32f,
+            0.33f, 0.33f, 0.33f, 0.33f, 0.33f,
+            0.34f, 0.34f, 0.34f, 0.34f, 0.34f,
+            0.35f, 0.35f, 0.35f, 0.35f, 0.35f,
+            0.36f, 0.36f, 0.36f, 0.36f, 0.36f,
+            0.37f, 0.37f, 0.37f, 0.37f, 0.37f,
+            0.38f, 0.38f, 0.38f, 0.38f, 0.38f,
+            0.39f, 0.39f, 0.39f, 0.39f, 0.39f,
+            0.40f, 0.40f, 0.40f, 0.40f, 0.40f,
+            0.41f, 0.41f, 0.41f, 0.41f, 0.41f
+        },
+        {  // Layer 1 - Pad to block-8 => block-8 + block-8
+            0.42f, 0.42f, 0.42f, 0.42f, 0.42f, 0.42f, 0.42f, 0.42f, 0.42f, 0.42f, 0.42f, 0.42f,
+            0.43f, 0.43f, 0.43f, 0.43f, 0.43f, 0.43f, 0.43f, 0.43f, 0.43f, 0.43f, 0.43f, 0.43f,
+            0.44f, 0.44f, 0.44f, 0.44f, 0.44f, 0.44f, 0.44f, 0.44f, 0.44f, 0.44f, 0.44f, 0.44f,
+            0.45f, 0.45f, 0.45f, 0.45f, 0.45f, 0.45f, 0.45f, 0.45f, 0.45f, 0.45f, 0.45f, 0.45f,
+            0.46f, 0.46f, 0.46f, 0.46f, 0.46f, 0.46f, 0.46f, 0.46f, 0.46f, 0.46f, 0.46f, 0.46f,
+            0.47f, 0.47f, 0.47f, 0.47f, 0.47f, 0.47f, 0.47f, 0.47f, 0.47f, 0.47f, 0.47f, 0.47f,
+            0.48f, 0.48f, 0.48f, 0.48f, 0.48f, 0.48f, 0.48f, 0.48f, 0.48f, 0.48f, 0.48f, 0.48f,
+            0.49f, 0.49f, 0.49f, 0.49f, 0.49f, 0.49f, 0.49f, 0.49f, 0.49f, 0.49f, 0.49f, 0.49f,
+            0.50f, 0.50f, 0.50f, 0.50f, 0.50f, 0.50f, 0.50f, 0.50f, 0.50f, 0.50f, 0.50f, 0.50f,
+            0.51f, 0.51f, 0.51f, 0.51f, 0.51f, 0.51f, 0.51f, 0.51f, 0.51f, 0.51f, 0.51f, 0.51f,
+            0.52f, 0.52f, 0.52f, 0.52f, 0.52f, 0.52f, 0.52f, 0.52f, 0.52f, 0.52f, 0.52f, 0.52f,
+            0.53f, 0.53f, 0.53f, 0.53f, 0.53f, 0.53f, 0.53f, 0.53f, 0.53f, 0.53f, 0.53f, 0.53f,
+            0.54f, 0.54f, 0.54f, 0.54f, 0.54f, 0.54f, 0.54f, 0.54f, 0.54f, 0.54f, 0.54f, 0.54f,
+            0.55f, 0.55f, 0.55f, 0.55f, 0.55f, 0.55f, 0.55f, 0.55f, 0.55f, 0.55f, 0.55f, 0.55f
+        },
+        {  // Layer 2 - No padding => block-8 + block-8
+            0.56f, 0.56f, 0.56f, 0.56f, 0.56f, 0.56f, 0.56f, 0.56f, 0.56f, 0.56f, 0.56f, 0.56f, 0.56f, 0.56f,
+            0.57f, 0.57f, 0.57f, 0.57f, 0.57f, 0.57f, 0.57f, 0.57f, 0.57f, 0.57f, 0.57f, 0.57f, 0.57f, 0.57f,
+            0.58f, 0.58f, 0.58f, 0.58f, 0.58f, 0.58f, 0.58f, 0.58f, 0.58f, 0.58f, 0.58f, 0.58f, 0.58f, 0.58f,
+            0.59f, 0.59f, 0.59f, 0.59f, 0.59f, 0.59f, 0.59f, 0.59f, 0.59f, 0.59f, 0.59f, 0.59f, 0.59f, 0.59f,
+            0.60f, 0.60f, 0.60f, 0.60f, 0.60f, 0.60f, 0.60f, 0.60f, 0.60f, 0.60f, 0.60f, 0.60f, 0.60f, 0.60f,
+            0.61f, 0.61f, 0.61f, 0.61f, 0.61f, 0.61f, 0.61f, 0.61f, 0.61f, 0.61f, 0.61f, 0.61f, 0.61f, 0.61f,
+            0.62f, 0.62f, 0.62f, 0.62f, 0.62f, 0.62f, 0.62f, 0.62f, 0.62f, 0.62f, 0.62f, 0.62f, 0.62f, 0.62f,
+            0.63f, 0.63f, 0.63f, 0.63f, 0.63f, 0.63f, 0.63f, 0.63f, 0.63f, 0.63f, 0.63f, 0.63f, 0.63f, 0.63f,
+            0.64f, 0.64f, 0.64f, 0.64f, 0.64f, 0.64f, 0.64f, 0.64f, 0.64f, 0.64f, 0.64f, 0.64f, 0.64f, 0.64f,
+            0.65f, 0.65f, 0.65f, 0.65f, 0.65f, 0.65f, 0.65f, 0.65f, 0.65f, 0.65f, 0.65f, 0.65f, 0.65f, 0.65f,
+            0.66f, 0.66f, 0.66f, 0.66f, 0.66f, 0.66f, 0.66f, 0.66f, 0.66f, 0.66f, 0.66f, 0.66f, 0.66f, 0.66f,
+            0.67f, 0.67f, 0.67f, 0.67f, 0.67f, 0.67f, 0.67f, 0.67f, 0.67f, 0.67f, 0.67f, 0.67f, 0.67f, 0.67f,
+            0.68f, 0.68f, 0.68f, 0.68f, 0.68f, 0.68f, 0.68f, 0.68f, 0.68f, 0.68f, 0.68f, 0.68f, 0.68f, 0.68f,
+            0.69f, 0.69f, 0.69f, 0.69f, 0.69f, 0.69f, 0.69f, 0.69f, 0.69f, 0.69f, 0.69f, 0.69f, 0.69f, 0.69f,
+            0.70f, 0.70f, 0.70f, 0.70f, 0.70f, 0.70f, 0.70f, 0.70f, 0.70f, 0.70f, 0.70f, 0.70f, 0.70f, 0.70f,
+            0.71f, 0.71f, 0.71f, 0.71f, 0.71f, 0.71f, 0.71f, 0.71f, 0.71f, 0.71f, 0.71f, 0.71f, 0.71f, 0.71f
+        },
+        {  // Layer 3 - Pad to block-4 => block-8
+            -0.01f, -0.01f, -0.01f, -0.01f, -0.01f, -0.01f, -0.01f, -0.01f, -0.01f, -0.01f, -0.01f, -0.01f, -0.01f, -0.01f,
+            -0.01f, -0.01f,
+            -0.02f, -0.02f, -0.02f, -0.02f, -0.02f, -0.02f, -0.02f, -0.02f, -0.02f, -0.02f, -0.02f, -0.02f, -0.02f, -0.02f,
+            -0.02f, -0.02f,
+            -0.03f, -0.03f, -0.03f, -0.03f, -0.03f, -0.03f, -0.03f, -0.03f, -0.03f, -0.03f, -0.03f, -0.03f, -0.03f, -0.03f,
+            -0.03f, -0.03f,
+            -0.04f, -0.04f, -0.04f, -0.04f, -0.04f, -0.04f, -0.04f, -0.04f, -0.04f, -0.04f, -0.04f, -0.04f, -0.04f, -0.04f,
+            -0.04f, -0.04f,
+            -0.05f, -0.05f, -0.05f, -0.05f, -0.05f, -0.05f, -0.05f, -0.05f, -0.05f, -0.05f, -0.05f, -0.05f, -0.05f, -0.05f,
+            -0.05f, -0.05f,
+            -0.06f, -0.06f, -0.06f, -0.06f, -0.06f, -0.06f, -0.06f, -0.06f, -0.06f, -0.06f, -0.06f, -0.06f, -0.06f, -0.06f,
+            -0.06f, -0.06f
+        }
+    },
+    {  // Mesh-1 Region-1
+        {  // Layer 0 - Pad to block-4 => block-4
+            1.0f, 1.0f, 1.0f,
+            2.0f, 2.0f, 2.0f
+        },
+        {  // Layer 1 - Pad to block-4 => block-4
+            3.0f, 3.0f,
+            4.0f, 4.0f,
+            5.0f, 5.0f,
+            6.0f, 6.0f
+        }
+    }
+};
+
+const Vector<Matrix<float> > mlbNetBiases = {
+    {  // Mesh-0 Region-0
+        {  // Layer 0
+            0.1f,
+            0.1f
+        },
+        {  // Layer 1
+            0.2f,
+            0.2f,
+            0.2f,
+            0.2f
+        },
+        {  // Layer 2
+            0.3f,
+        }
+    },
+    {  // Mesh-0 Region-1
+        {  // Layer 0
+            0.4f,
+            0.4f,
+            0.4f,
+            0.4f,
+            0.4f,
+            0.4f,
+            0.4f,
+        },
+        {  // Layer 1
+            0.5f,
+            0.5f,
+            0.5f,
+            0.5f,
+            0.5f,
+            0.5f,
+            0.5f,
+            0.5f
+        },
+        {  // Layer 2
+            0.6f,
+            0.6f,
+            0.6f,
+            0.6f,
+            0.6f,
+            0.6f,
+            0.6f,
+            0.6f,
+            0.6f
+        },
+        {  // Layer 3
+            0.7f,
+            0.7f,
+            0.7f
+        }
+    },
+    {  // Mesh-1 Region-0
+        {  // Layer 0
+            0.8f,
+            0.8f,
+            0.8f,
+            0.8f,
+            0.8f,
+            0.8f,
+            0.8f,
+            0.8f,
+            0.8f,
+            0.8f,
+            0.8f,
+            0.8f,
+        },
+        {  // Layer 1
+            0.9f,
+            0.9f,
+            0.9f,
+            0.9f,
+            0.9f,
+            0.9f,
+            0.9f,
+            0.9f,
+            0.9f,
+            0.9f,
+            0.9f,
+            0.9f,
+            0.9f,
+            0.9f
+        },
+        {  // Layer 2
+            1.1f,
+            1.1f,
+            1.1f,
+            1.1f,
+            1.1f,
+            1.1f,
+            1.1f,
+            1.1f,
+            1.1f,
+            1.1f,
+            1.1f,
+            1.1f,
+            1.1f,
+            1.1f,
+            1.1f,
+            1.1f
+        },
+        {  // Layer 3
+            1.2f,
+            1.2f,
+            1.2f,
+            1.2f,
+            1.2f,
+            1.2f
+        }
+    },
+    {  // Mesh-1 Region-1
+        {  // Layer 0
+            1.3f,
+            1.3f
+        },
+        {  // Layer 1
+            1.4f,
+            1.4f,
+            1.4f,
+            1.4f
+        }
+    }
+};
+
+}  // namespace unoptimized
+
+namespace optimized {
+
+const Vector<AlignedMatrix<float> > mlbNetWeightsFloat = {
+    {  // Mesh-0 Region-0
+        {  // Layer-0
+            0.11f, 0.11f, 0.0f, 0.0f,
+            0.11f, 0.11f, 0.0f, 0.0f
+        },
+        {  // Layer-1
+            0.22f, 0.33f, 0.44f, 0.55f,
+            0.22f, 0.33f, 0.44f, 0.55f
+        },
+        {  // Layer-2
+            0.66f, 0.0f, 0.0f, 0.0f,
+            0.66f, 0.0f, 0.0f, 0.0f,
+            0.66f, 0.0f, 0.0f, 0.0f,
+            0.66f, 0.0f, 0.0f, 0.0f
+        }
+    },
+    {  // Mesh-0 Region-1
+        {  // Layer-0
+            0.66f, 0.77f, 0.88f, 0.99f, 1.0f, 1.11f, 1.22f, 0.0f,
+            0.66f, 0.77f, 0.88f, 0.99f, 1.0f, 1.11f, 1.22f, 0.0f,
+            0.66f, 0.77f, 0.88f, 0.99f, 1.0f, 1.11f, 1.22f, 0.0f
+        },
+        {  // Layer-1
+            1.33f, 1.44f, 1.55f, 1.66f, 1.77f, 1.88f, 1.99f, 2.0f,
+            1.33f, 1.44f, 1.55f, 1.66f, 1.77f, 1.88f, 1.99f, 2.0f,
+            1.33f, 1.44f, 1.55f, 1.66f, 1.77f, 1.88f, 1.99f, 2.0f,
+            1.33f, 1.44f, 1.55f, 1.66f, 1.77f, 1.88f, 1.99f, 2.0f,
+            1.33f, 1.44f, 1.55f, 1.66f, 1.77f, 1.88f, 1.99f, 2.0f,
+            1.33f, 1.44f, 1.55f, 1.66f, 1.77f, 1.88f, 1.99f, 2.0f,
+            1.33f, 1.44f, 1.55f, 1.66f, 1.77f, 1.88f, 1.99f, 2.0f
+        },
+        {  // Layer-2
+            -0.09f, -0.08f, -0.07f, -0.06f, -0.05f, -0.04f, -0.03f, -0.02f,
+            -0.09f, -0.08f, -0.07f, -0.06f, -0.05f, -0.04f, -0.03f, -0.02f,
+            -0.09f, -0.08f, -0.07f, -0.06f, -0.05f, -0.04f, -0.03f, -0.02f,
+            -0.09f, -0.08f, -0.07f, -0.06f, -0.05f, -0.04f, -0.03f, -0.02f,
+            -0.09f, -0.08f, -0.07f, -0.06f, -0.05f, -0.04f, -0.03f, -0.02f,
+            -0.09f, -0.08f, -0.07f, -0.06f, -0.05f, -0.04f, -0.03f, -0.02f,
+            -0.09f, -0.08f, -0.07f, -0.06f, -0.05f, -0.04f, -0.03f, -0.02f,
+            -0.09f, -0.08f, -0.07f, -0.06f, -0.05f, -0.04f, -0.03f, -0.02f,
+            -0.01f, 0.0f, 0.0f, 0.0f,
+            -0.01f, 0.0f, 0.0f, 0.0f,
+            -0.01f, 0.0f, 0.0f, 0.0f,
+            -0.01f, 0.0f, 0.0f, 0.0f,
+            -0.01f, 0.0f, 0.0f, 0.0f,
+            -0.01f, 0.0f, 0.0f, 0.0f,
+            -0.01f, 0.0f, 0.0f, 0.0f,
+            -0.01f, 0.0f, 0.0f, 0.0f
+        },
+        {  // Layer-3
+            -0.30f, -0.31f, -0.32f, 0.0f,
+            -0.30f, -0.31f, -0.32f, 0.0f,
+            -0.30f, -0.31f, -0.32f, 0.0f,
+            -0.30f, -0.31f, -0.32f, 0.0f,
+            -0.30f, -0.31f, -0.32f, 0.0f,
+            -0.30f, -0.31f, -0.32f, 0.0f,
+            -0.30f, -0.31f, -0.32f, 0.0f,
+            -0.30f, -0.31f, -0.32f, 0.0f,
+            -0.30f, -0.31f, -0.32f, 0.0f
+        }
+    },
+    {  // Mesh-1 Region-0
+        {  // Layer-0
+            0.30f, 0.31f, 0.32f, 0.33f, 0.34f, 0.35f, 0.36f, 0.37f,
+            0.30f, 0.31f, 0.32f, 0.33f, 0.34f, 0.35f, 0.36f, 0.37f,
+            0.30f, 0.31f, 0.32f, 0.33f, 0.34f, 0.35f, 0.36f, 0.37f,
+            0.30f, 0.31f, 0.32f, 0.33f, 0.34f, 0.35f, 0.36f, 0.37f,
+            0.30f, 0.31f, 0.32f, 0.33f, 0.34f, 0.35f, 0.36f, 0.37f,
+            0.38f, 0.39f, 0.40f, 0.41f,
+            0.38f, 0.39f, 0.40f, 0.41f,
+            0.38f, 0.39f, 0.40f, 0.41f,
+            0.38f, 0.39f, 0.40f, 0.41f,
+            0.38f, 0.39f, 0.40f, 0.41f
+        },
+        {  // Layer-1
+            0.42f, 0.43f, 0.44f, 0.45f, 0.46f, 0.47f, 0.48f, 0.49f,
+            0.42f, 0.43f, 0.44f, 0.45f, 0.46f, 0.47f, 0.48f, 0.49f,
+            0.42f, 0.43f, 0.44f, 0.45f, 0.46f, 0.47f, 0.48f, 0.49f,
+            0.42f, 0.43f, 0.44f, 0.45f, 0.46f, 0.47f, 0.48f, 0.49f,
+            0.42f, 0.43f, 0.44f, 0.45f, 0.46f, 0.47f, 0.48f, 0.49f,
+            0.42f, 0.43f, 0.44f, 0.45f, 0.46f, 0.47f, 0.48f, 0.49f,
+            0.42f, 0.43f, 0.44f, 0.45f, 0.46f, 0.47f, 0.48f, 0.49f,
+            0.42f, 0.43f, 0.44f, 0.45f, 0.46f, 0.47f, 0.48f, 0.49f,
+            0.42f, 0.43f, 0.44f, 0.45f, 0.46f, 0.47f, 0.48f, 0.49f,
+            0.42f, 0.43f, 0.44f, 0.45f, 0.46f, 0.47f, 0.48f, 0.49f,
+            0.42f, 0.43f, 0.44f, 0.45f, 0.46f, 0.47f, 0.48f, 0.49f,
+            0.42f, 0.43f, 0.44f, 0.45f, 0.46f, 0.47f, 0.48f, 0.49f,
+            0.50f, 0.51f, 0.52f, 0.53f, 0.54f, 0.55f, 0.0f, 0.0f,
+            0.50f, 0.51f, 0.52f, 0.53f, 0.54f, 0.55f, 0.0f, 0.0f,
+            0.50f, 0.51f, 0.52f, 0.53f, 0.54f, 0.55f, 0.0f, 0.0f,
+            0.50f, 0.51f, 0.52f, 0.53f, 0.54f, 0.55f, 0.0f, 0.0f,
+            0.50f, 0.51f, 0.52f, 0.53f, 0.54f, 0.55f, 0.0f, 0.0f,
+            0.50f, 0.51f, 0.52f, 0.53f, 0.54f, 0.55f, 0.0f, 0.0f,
+            0.50f, 0.51f, 0.52f, 0.53f, 0.54f, 0.55f, 0.0f, 0.0f,
+            0.50f, 0.51f, 0.52f, 0.53f, 0.54f, 0.55f, 0.0f, 0.0f,
+            0.50f, 0.51f, 0.52f, 0.53f, 0.54f, 0.55f, 0.0f, 0.0f,
+            0.50f, 0.51f, 0.52f, 0.53f, 0.54f, 0.55f, 0.0f, 0.0f,
+            0.50f, 0.51f, 0.52f, 0.53f, 0.54f, 0.55f, 0.0f, 0.0f,
+            0.50f, 0.51f, 0.52f, 0.53f, 0.54f, 0.55f, 0.0f, 0.0f
+        },
+        {  // Layer-2
+            0.56f, 0.57f, 0.58f, 0.59f, 0.60f, 0.61f, 0.62f, 0.63f,
+            0.56f, 0.57f, 0.58f, 0.59f, 0.60f, 0.61f, 0.62f, 0.63f,
+            0.56f, 0.57f, 0.58f, 0.59f, 0.60f, 0.61f, 0.62f, 0.63f,
+            0.56f, 0.57f, 0.58f, 0.59f, 0.60f, 0.61f, 0.62f, 0.63f,
+            0.56f, 0.57f, 0.58f, 0.59f, 0.60f, 0.61f, 0.62f, 0.63f,
+            0.56f, 0.57f, 0.58f, 0.59f, 0.60f, 0.61f, 0.62f, 0.63f,
+            0.56f, 0.57f, 0.58f, 0.59f, 0.60f, 0.61f, 0.62f, 0.63f,
+            0.56f, 0.57f, 0.58f, 0.59f, 0.60f, 0.61f, 0.62f, 0.63f,
+            0.56f, 0.57f, 0.58f, 0.59f, 0.60f, 0.61f, 0.62f, 0.63f,
+            0.56f, 0.57f, 0.58f, 0.59f, 0.60f, 0.61f, 0.62f, 0.63f,
+            0.56f, 0.57f, 0.58f, 0.59f, 0.60f, 0.61f, 0.62f, 0.63f,
+            0.56f, 0.57f, 0.58f, 0.59f, 0.60f, 0.61f, 0.62f, 0.63f,
+            0.56f, 0.57f, 0.58f, 0.59f, 0.60f, 0.61f, 0.62f, 0.63f,
+            0.56f, 0.57f, 0.58f, 0.59f, 0.60f, 0.61f, 0.62f, 0.63f,
+            0.64f, 0.65f, 0.66f, 0.67f, 0.68f, 0.69f, 0.70f, 0.71f,
+            0.64f, 0.65f, 0.66f, 0.67f, 0.68f, 0.69f, 0.70f, 0.71f,
+            0.64f, 0.65f, 0.66f, 0.67f, 0.68f, 0.69f, 0.70f, 0.71f,
+            0.64f, 0.65f, 0.66f, 0.67f, 0.68f, 0.69f, 0.70f, 0.71f,
+            0.64f, 0.65f, 0.66f, 0.67f, 0.68f, 0.69f, 0.70f, 0.71f,
+            0.64f, 0.65f, 0.66f, 0.67f, 0.68f, 0.69f, 0.70f, 0.71f,
+            0.64f, 0.65f, 0.66f, 0.67f, 0.68f, 0.69f, 0.70f, 0.71f,
+            0.64f, 0.65f, 0.66f, 0.67f, 0.68f, 0.69f, 0.70f, 0.71f,
+            0.64f, 0.65f, 0.66f, 0.67f, 0.68f, 0.69f, 0.70f, 0.71f,
+            0.64f, 0.65f, 0.66f, 0.67f, 0.68f, 0.69f, 0.70f, 0.71f,
+            0.64f, 0.65f, 0.66f, 0.67f, 0.68f, 0.69f, 0.70f, 0.71f,
+            0.64f, 0.65f, 0.66f, 0.67f, 0.68f, 0.69f, 0.70f, 0.71f,
+            0.64f, 0.65f, 0.66f, 0.67f, 0.68f, 0.69f, 0.70f, 0.71f,
+            0.64f, 0.65f, 0.66f, 0.67f, 0.68f, 0.69f, 0.70f, 0.71f,
+        },
+        {  // Layer-3
+            -0.01f, -0.02f, -0.03f, -0.04f, -0.05f, -0.06f, 0.0f, 0.0f,
+            -0.01f, -0.02f, -0.03f, -0.04f, -0.05f, -0.06f, 0.0f, 0.0f,
+            -0.01f, -0.02f, -0.03f, -0.04f, -0.05f, -0.06f, 0.0f, 0.0f,
+            -0.01f, -0.02f, -0.03f, -0.04f, -0.05f, -0.06f, 0.0f, 0.0f,
+            -0.01f, -0.02f, -0.03f, -0.04f, -0.05f, -0.06f, 0.0f, 0.0f,
+            -0.01f, -0.02f, -0.03f, -0.04f, -0.05f, -0.06f, 0.0f, 0.0f,
+            -0.01f, -0.02f, -0.03f, -0.04f, -0.05f, -0.06f, 0.0f, 0.0f,
+            -0.01f, -0.02f, -0.03f, -0.04f, -0.05f, -0.06f, 0.0f, 0.0f,
+            -0.01f, -0.02f, -0.03f, -0.04f, -0.05f, -0.06f, 0.0f, 0.0f,
+            -0.01f, -0.02f, -0.03f, -0.04f, -0.05f, -0.06f, 0.0f, 0.0f,
+            -0.01f, -0.02f, -0.03f, -0.04f, -0.05f, -0.06f, 0.0f, 0.0f,
+            -0.01f, -0.02f, -0.03f, -0.04f, -0.05f, -0.06f, 0.0f, 0.0f,
+            -0.01f, -0.02f, -0.03f, -0.04f, -0.05f, -0.06f, 0.0f, 0.0f,
+            -0.01f, -0.02f, -0.03f, -0.04f, -0.05f, -0.06f, 0.0f, 0.0f,
+            -0.01f, -0.02f, -0.03f, -0.04f, -0.05f, -0.06f, 0.0f, 0.0f,
+            -0.01f, -0.02f, -0.03f, -0.04f, -0.05f, -0.06f, 0.0f, 0.0f
+        }
+    },
+    {  // Mesh-1 Region-1
+        {  // Layer-0
+            1.0f, 2.0f, 0.0f, 0.0f,
+            1.0f, 2.0f, 0.0f, 0.0f,
+            1.0f, 2.0f, 0.0f, 0.0f
+        },
+        {  // Layer-1
+            3.0f, 4.0f, 5.0f, 6.0f,
+            3.0f, 4.0f, 5.0f, 6.0f
+        }
+    }
+};
+
+const Vector<AlignedMatrix<std::uint16_t> > mlbNetWeightsHalfFloat = {
+    {  // Mesh-0 Region-0
+        {  // Layer-0
+            12042, 12042, 0, 0,
+            12042, 12042, 0, 0
+        },
+        {  // Layer-1
+            13066, 13640, 14090, 14438,
+            13066, 13640, 14090, 14438
+        },
+        {  // Layer-2
+            14664, 0, 0, 0,
+            14664, 0, 0, 0,
+            14664, 0, 0, 0,
+            14664, 0, 0, 0
+        }
+    },
+    {  // Mesh-0 Region-1
+        {  // Layer-0
+            14664, 14889, 15114, 15340, 15360, 15473, 15585, 0,
+            14664, 14889, 15114, 15340, 15360, 15473, 15585, 0,
+            14664, 14889, 15114, 15340, 15360, 15473, 15585, 0
+        },
+        {  // Layer-1
+            15698, 15811, 15923, 16036, 16148, 16261, 16374, 16384,
+            15698, 15811, 15923, 16036, 16148, 16261, 16374, 16384,
+            15698, 15811, 15923, 16036, 16148, 16261, 16374, 16384,
+            15698, 15811, 15923, 16036, 16148, 16261, 16374, 16384,
+            15698, 15811, 15923, 16036, 16148, 16261, 16374, 16384,
+            15698, 15811, 15923, 16036, 16148, 16261, 16374, 16384,
+            15698, 15811, 15923, 16036, 16148, 16261, 16374, 16384
+        },
+        {  // Layer-2
+            44483, 44319, 44155, 43950, 43622, 43295, 42926, 42271,
+            44483, 44319, 44155, 43950, 43622, 43295, 42926, 42271,
+            44483, 44319, 44155, 43950, 43622, 43295, 42926, 42271,
+            44483, 44319, 44155, 43950, 43622, 43295, 42926, 42271,
+            44483, 44319, 44155, 43950, 43622, 43295, 42926, 42271,
+            44483, 44319, 44155, 43950, 43622, 43295, 42926, 42271,
+            44483, 44319, 44155, 43950, 43622, 43295, 42926, 42271,
+            44483, 44319, 44155, 43950, 43622, 43295, 42926, 42271,
+            41247, 0, 0, 0,
+            41247, 0, 0, 0,
+            41247, 0, 0, 0,
+            41247, 0, 0, 0,
+            41247, 0, 0, 0,
+            41247, 0, 0, 0,
+            41247, 0, 0, 0,
+            41247, 0, 0, 0
+        },
+        {  // Layer-3
+            46285, 46326, 46367, 0,
+            46285, 46326, 46367, 0,
+            46285, 46326, 46367, 0,
+            46285, 46326, 46367, 0,
+            46285, 46326, 46367, 0,
+            46285, 46326, 46367, 0,
+            46285, 46326, 46367, 0,
+            46285, 46326, 46367, 0,
+            46285, 46326, 46367, 0
+        }
+    },
+    {  // Mesh-1 Region-0
+        {  // Layer-0
+            13517, 13558, 13599, 13640, 13681, 13722, 13763, 13804,
+            13517, 13558, 13599, 13640, 13681, 13722, 13763, 13804,
+            13517, 13558, 13599, 13640, 13681, 13722, 13763, 13804,
+            13517, 13558, 13599, 13640, 13681, 13722, 13763, 13804,
+            13517, 13558, 13599, 13640, 13681, 13722, 13763, 13804,
+            13844, 13885, 13926, 13967,
+            13844, 13885, 13926, 13967,
+            13844, 13885, 13926, 13967,
+            13844, 13885, 13926, 13967,
+            13844, 13885, 13926, 13967
+        },
+        {  // Layer-1
+            14008, 14049, 14090, 14131, 14172, 14213, 14254, 14295,
+            14008, 14049, 14090, 14131, 14172, 14213, 14254, 14295,
+            14008, 14049, 14090, 14131, 14172, 14213, 14254, 14295,
+            14008, 14049, 14090, 14131, 14172, 14213, 14254, 14295,
+            14008, 14049, 14090, 14131, 14172, 14213, 14254, 14295,
+            14008, 14049, 14090, 14131, 14172, 14213, 14254, 14295,
+            14008, 14049, 14090, 14131, 14172, 14213, 14254, 14295,
+            14008, 14049, 14090, 14131, 14172, 14213, 14254, 14295,
+            14008, 14049, 14090, 14131, 14172, 14213, 14254, 14295,
+            14008, 14049, 14090, 14131, 14172, 14213, 14254, 14295,
+            14008, 14049, 14090, 14131, 14172, 14213, 14254, 14295,
+            14008, 14049, 14090, 14131, 14172, 14213, 14254, 14295,
+            14336, 14356, 14377, 14397, 14418, 14438, 0, 0,
+            14336, 14356, 14377, 14397, 14418, 14438, 0, 0,
+            14336, 14356, 14377, 14397, 14418, 14438, 0, 0,
+            14336, 14356, 14377, 14397, 14418, 14438, 0, 0,
+            14336, 14356, 14377, 14397, 14418, 14438, 0, 0,
+            14336, 14356, 14377, 14397, 14418, 14438, 0, 0,
+            14336, 14356, 14377, 14397, 14418, 14438, 0, 0,
+            14336, 14356, 14377, 14397, 14418, 14438, 0, 0,
+            14336, 14356, 14377, 14397, 14418, 14438, 0, 0,
+            14336, 14356, 14377, 14397, 14418, 14438, 0, 0,
+            14336, 14356, 14377, 14397, 14418, 14438, 0, 0,
+            14336, 14356, 14377, 14397, 14418, 14438, 0, 0
+        },
+        {  // Layer-2
+            14459, 14479, 14500, 14520, 14541, 14561, 14582, 14602,
+            14459, 14479, 14500, 14520, 14541, 14561, 14582, 14602,
+            14459, 14479, 14500, 14520, 14541, 14561, 14582, 14602,
+            14459, 14479, 14500, 14520, 14541, 14561, 14582, 14602,
+            14459, 14479, 14500, 14520, 14541, 14561, 14582, 14602,
+            14459, 14479, 14500, 14520, 14541, 14561, 14582, 14602,
+            14459, 14479, 14500, 14520, 14541, 14561, 14582, 14602,
+            14459, 14479, 14500, 14520, 14541, 14561, 14582, 14602,
+            14459, 14479, 14500, 14520, 14541, 14561, 14582, 14602,
+            14459, 14479, 14500, 14520, 14541, 14561, 14582, 14602,
+            14459, 14479, 14500, 14520, 14541, 14561, 14582, 14602,
+            14459, 14479, 14500, 14520, 14541, 14561, 14582, 14602,
+            14459, 14479, 14500, 14520, 14541, 14561, 14582, 14602,
+            14459, 14479, 14500, 14520, 14541, 14561, 14582, 14602,
+            14623, 14643, 14664, 14684, 14705, 14725, 14746, 14766,
+            14623, 14643, 14664, 14684, 14705, 14725, 14746, 14766,
+            14623, 14643, 14664, 14684, 14705, 14725, 14746, 14766,
+            14623, 14643, 14664, 14684, 14705, 14725, 14746, 14766,
+            14623, 14643, 14664, 14684, 14705, 14725, 14746, 14766,
+            14623, 14643, 14664, 14684, 14705, 14725, 14746, 14766,
+            14623, 14643, 14664, 14684, 14705, 14725, 14746, 14766,
+            14623, 14643, 14664, 14684, 14705, 14725, 14746, 14766,
+            14623, 14643, 14664, 14684, 14705, 14725, 14746, 14766,
+            14623, 14643, 14664, 14684, 14705, 14725, 14746, 14766,
+            14623, 14643, 14664, 14684, 14705, 14725, 14746, 14766,
+            14623, 14643, 14664, 14684, 14705, 14725, 14746, 14766,
+            14623, 14643, 14664, 14684, 14705, 14725, 14746, 14766,
+            14623, 14643, 14664, 14684, 14705, 14725, 14746, 14766
+        },
+        {  // Layer-3
+            41247, 42271, 42926, 43295, 43622, 43950, 0, 0,
+            41247, 42271, 42926, 43295, 43622, 43950, 0, 0,
+            41247, 42271, 42926, 43295, 43622, 43950, 0, 0,
+            41247, 42271, 42926, 43295, 43622, 43950, 0, 0,
+            41247, 42271, 42926, 43295, 43622, 43950, 0, 0,
+            41247, 42271, 42926, 43295, 43622, 43950, 0, 0,
+            41247, 42271, 42926, 43295, 43622, 43950, 0, 0,
+            41247, 42271, 42926, 43295, 43622, 43950, 0, 0,
+            41247, 42271, 42926, 43295, 43622, 43950, 0, 0,
+            41247, 42271, 42926, 43295, 43622, 43950, 0, 0,
+            41247, 42271, 42926, 43295, 43622, 43950, 0, 0,
+            41247, 42271, 42926, 43295, 43622, 43950, 0, 0,
+            41247, 42271, 42926, 43295, 43622, 43950, 0, 0,
+            41247, 42271, 42926, 43295, 43622, 43950, 0, 0,
+            41247, 42271, 42926, 43295, 43622, 43950, 0, 0,
+            41247, 42271, 42926, 43295, 43622, 43950, 0, 0
+        }
+    },
+    {  // Mesh-1 Region-1
+        {  // Layer-0
+            15360, 16384, 0, 0,
+            15360, 16384, 0, 0,
+            15360, 16384, 0, 0
+        },
+        {  // Layer-1
+            16896, 17408, 17664, 17920,
+            16896, 17408, 17664, 17920
+        }
+    }
+};
+
+const Vector<AlignedMatrix<float> > mlbNetBiasesFloat = {
+    {  // Mesh-0 Region-0
+        {  // Layer 0
+            0.1f,
+            0.1f,
+            0.0f,
+            0.0f
+        },
+        {  // Layer 1
+            0.2f,
+            0.2f,
+            0.2f,
+            0.2f
+        },
+        {  // Layer 2
+            0.3f,
+            0.0f,
+            0.0f,
+            0.0f
+        }
+    },
+    {  // Mesh-0 Region-1
+        {  // Layer 0
+            0.4f,
+            0.4f,
+            0.4f,
+            0.4f,
+            0.4f,
+            0.4f,
+            0.4f,
+            0.0f
+        },
+        {  // Layer 1
+            0.5f,
+            0.5f,
+            0.5f,
+            0.5f,
+            0.5f,
+            0.5f,
+            0.5f,
+            0.5f
+        },
+        {  // Layer 2
+            0.6f,
+            0.6f,
+            0.6f,
+            0.6f,
+            0.6f,
+            0.6f,
+            0.6f,
+            0.6f,
+            0.6f,
+            0.0f,
+            0.0f,
+            0.0f
+        },
+        {  // Layer 3
+            0.7f,
+            0.7f,
+            0.7f,
+            0.0f
+        }
+    },
+    {  // Mesh-1 Region-0
+        {  // Layer 0
+            0.8f,
+            0.8f,
+            0.8f,
+            0.8f,
+            0.8f,
+            0.8f,
+            0.8f,
+            0.8f,
+            0.8f,
+            0.8f,
+            0.8f,
+            0.8f,
+        },
+        {  // Layer 1
+            0.9f,
+            0.9f,
+            0.9f,
+            0.9f,
+            0.9f,
+            0.9f,
+            0.9f,
+            0.9f,
+            0.9f,
+            0.9f,
+            0.9f,
+            0.9f,
+            0.9f,
+            0.9f,
+            0.0f,
+            0.0f
+        },
+        {  // Layer 2
+            1.1f,
+            1.1f,
+            1.1f,
+            1.1f,
+            1.1f,
+            1.1f,
+            1.1f,
+            1.1f,
+            1.1f,
+            1.1f,
+            1.1f,
+            1.1f,
+            1.1f,
+            1.1f,
+            1.1f,
+            1.1f
+        },
+        {  // Layer 3
+            1.2f,
+            1.2f,
+            1.2f,
+            1.2f,
+            1.2f,
+            1.2f,
+            0.0f,
+            0.0f
+        }
+    },
+    {  // Mesh-1 Region-1
+        {  // Layer 0
+            1.3f,
+            1.3f,
+            0.0f,
+            0.0f
+        },
+        {  // Layer 1
+            1.4f,
+            1.4f,
+            1.4f,
+            1.4f
+        }
+    }
+};
+
+const Vector<AlignedMatrix<std::uint16_t> > mlbNetBiasesHalfFloat = {
+    {  // Mesh-0 Region-0
+        {  // Layer 0
+            11878,
+            11878,
+            0,
+            0
+        },
+        {  // Layer 1
+            12902,
+            12902,
+            12902,
+            12902
+        },
+        {  // Layer 2
+            13517,
+            0,
+            0,
+            0
+        }
+    },
+    {  // Mesh-0 Region-1
+        {  // Layer 0
+            13926,
+            13926,
+            13926,
+            13926,
+            13926,
+            13926,
+            13926,
+            0
+        },
+        {  // Layer 1
+            14336,
+            14336,
+            14336,
+            14336,
+            14336,
+            14336,
+            14336,
+            14336
+        },
+        {  // Layer 2
+            14541,
+            14541,
+            14541,
+            14541,
+            14541,
+            14541,
+            14541,
+            14541,
+            14541,
+            0,
+            0,
+            0
+        },
+        {  // Layer 3
+            14746,
+            14746,
+            14746,
+            0
+        }
+    },
+    {  // Mesh-1 Region-0
+        {  // Layer 0
+            14950,
+            14950,
+            14950,
+            14950,
+            14950,
+            14950,
+            14950,
+            14950,
+            14950,
+            14950,
+            14950,
+            14950,
+        },
+        {  // Layer 1
+            15155,
+            15155,
+            15155,
+            15155,
+            15155,
+            15155,
+            15155,
+            15155,
+            15155,
+            15155,
+            15155,
+            15155,
+            15155,
+            15155,
+            0,
+            0
+        },
+        {  // Layer 2
+            15462,
+            15462,
+            15462,
+            15462,
+            15462,
+            15462,
+            15462,
+            15462,
+            15462,
+            15462,
+            15462,
+            15462,
+            15462,
+            15462,
+            15462,
+            15462
+        },
+        {  // Layer 3
+            15565,
+            15565,
+            15565,
+            15565,
+            15565,
+            15565,
+            0,
+            0
+        }
+    },
+    {  // Mesh-1 Region-1
+        {  // Layer 0
+            15667,
+            15667,
+            0,
+            0
+        },
+        {  // Layer 1
+            15770,
+            15770,
+            15770,
+            15770
+        }
+    }
+};
+
+const LODSpec lods = {
+    {  // NetIndices
+        {  // LOD-0
+            0u, 1u, 2u, 3u
+        },
+        {  // LOD-1
+            0u, 1u, 2u
+        },
+        {  // LOD-2
+        },
+        {  // LOD-3
+        }
+    },
+    4u  // Net count
+};
+
+template<>
+const rl4::Vector<rl4::AlignedMatrix<float> >& Values<float>::weights() {
+    return optimized::mlbNetWeightsFloat;
+}
+
+template<>
+const rl4::Vector<rl4::AlignedMatrix<float> >& Values<float>::biases() {
+    return optimized::mlbNetBiasesFloat;
+}
+
+template<>
+const rl4::Vector<rl4::AlignedMatrix<std::uint16_t> >& Values<std::uint16_t>::weights() {
+    return optimized::mlbNetWeightsHalfFloat;
+}
+
+template<>
+const rl4::Vector<rl4::AlignedMatrix<std::uint16_t> >& Values<std::uint16_t>::biases() {
+    return optimized::mlbNetBiasesHalfFloat;
+}
+
+}  // namespace optimized
+
+namespace input {
+
+// Calculation input values
+const rl4::Vector<float> values = {0.1f, 0.2f, 0.3f, 0.4f, 0.0f, 0.6f, 0.7f, 0.8f, 0.9f, 0.0f, 0.11f, 0.12f, 0.13f};
+
+}  // namespace input
+
+namespace output {
+
+// Expected output results for each LOD
+const rl4::Matrix<float> valuesPerLOD = {
+    {  // LOD-0
+        2.00812592326754f,
+        7.6104124f,
+        7.84075948f,
+        8.07110656f,
+        -14.0201885552f,
+        -29.2403771104f,
+        -44.4605656656f,
+        -59.6807542208f,
+        -74.900942776f,
+        -90.1211313312f,
+        7.3206385240577f,
+        9.29418469874361f,
+        11.2677308734295f,
+        13.2412770481154f
+    }, {  // LOD-1
+        2.00812592326754f,
+        7.6104124f,
+        7.84075948f,
+        8.07110656f,
+        -14.0201885552f,
+        -29.2403771104f,
+        -44.4605656656f,
+        -59.6807542208f,
+        -74.900942776f,
+        -90.1211313312f,
+        0.0f,
+        0.0f,
+        0.0f,
+        0.0f
+    }, {  // LOD-2
+        0.0f,
+        0.0f,
+        0.0f,
+        0.0f,
+        0.0f,
+        0.0f,
+        0.0f,
+        0.0f,
+        0.0f,
+        0.0f,
+        0.0f,
+        0.0f,
+        0.0f,
+        0.0f
+    }, {  // LOD-3
+        0.0f,
+        0.0f,
+        0.0f,
+        0.0f,
+        0.0f,
+        0.0f,
+        0.0f,
+        0.0f,
+        0.0f,
+        0.0f,
+        0.0f,
+        0.0f,
+        0.0f,
+        0.0f
+    }
+};
+
+}  // namespace output
+
+#ifdef __clang__
+    #pragma clang diagnostic pop
+#endif
+
+CanonicalReader::~CanonicalReader() = default;
+
+}  // namespace block4
+
+}  // namespace ml
+
+}  // namespace rltests
diff --git a/Engine/Plugins/Animation/RigLogic/Source/RigLogicLibTest/Private/rltests/ml/cpu/FixturesBlock4.h b/Engine/Plugins/Animation/RigLogic/Source/RigLogicLibTest/Private/rltests/ml/cpu/FixturesBlock4.h
new file mode 100644
--- /dev/null
+++ b/Engine/Plugins/Animation/RigLogic/Source/RigLogicLibTest/Private/rltests/ml/cpu/FixturesBlock4.h
@@ -0,0 +1,137 @@
+// Copyright Epic Games, Inc. All Rights Reserved.
+
+#pragma once
+
+#include "rltests/Defs.h"
+#include "rltests/StorageValueType.h"
+#include "rltests/dna/FakeReader.h"
+
+#include "riglogic/TypeDefs.h"
+#include "riglogic/ml/LODSpec.h"
+
+#ifdef _MSC_VER
+    #pragma warning(push)
+    #pragma warning(disable : 4365 4987)
+#endif
+#include <cstdint>
+#include <functional>
+#include <memory>
+#ifdef _MSC_VER
+    #pragma warning(pop)
+#endif
+
+namespace rltests {
+
+namespace ml {
+
+namespace block4 {
+
+using namespace rl4;
+
+namespace unoptimized {
+
+extern const std::uint16_t rawControlCount;
+extern const std::uint16_t mlControlCount;
+extern const std::uint16_t lodCount;
+extern const std::uint16_t neuralNetworkCount;
+extern const Matrix<std::uint16_t> neuralNetworkIndicesPerLOD;
+
+// [neuralNetIndex][inputIndex]
+extern const Matrix<std::uint16_t> mlbNetInputIndices;
+// [neuralNetIndex][outputIndex]
+extern const Matrix<std::uint16_t> mlbNetOutputIndices;
+// [neuralNetIndex][layerIndex]
+extern const Matrix<dna::ActivationFunction> mlbNetActivationFunctions;
+// [neuralNetIndex][layerIndex][paramIndex]
+extern const Vector<Matrix<float> > mlbNetActivationFunctionParameters;
+// [neuralNetIndex][layerIndex][weightIndex]
+extern const Vector<Matrix<float> > mlbNetWeights;
+// [neuralNetIndex][layerIndex][biasIndex]
+extern const Vector<Matrix<float> > mlbNetBiases;
+
+}  // namespace unoptimized
+
+namespace optimized {
+
+extern const Vector<AlignedMatrix<float> > mlbNetWeightsFloat;
+extern const Vector<AlignedMatrix<std::uint16_t> > mlbNetWeightsHalfFloat;
+extern const Vector<AlignedMatrix<float> > mlbNetBiasesFloat;
+extern const Vector<AlignedMatrix<std::uint16_t> > mlbNetBiasesHalfFloat;
+extern const LODSpec lods;
+
+template<typename TValue>
+struct Values {
+    static const rl4::Vector<rl4::AlignedMatrix<TValue> >& weights();
+    static const rl4::Vector<rl4::AlignedMatrix<TValue> >& biases();
+};
+
+}  // namespace optimized
+
+namespace input {
+
+// Calculation input values
+extern const Vector<float> values;
+
+}  // namespace input
+
+namespace output {
+
+// Calculation output values
+extern const Matrix<float> valuesPerLOD;
+
+}  // namespace output
+
+class CanonicalReader : public dna::FakeReader {
+    public:
+        ~CanonicalReader();
+
+        std::uint16_t getLODCount() const override {
+            return unoptimized::lodCount;
+        }
+
+        std::uint16_t getNeuralNetworkCount() const override {
+            return unoptimized::neuralNetworkCount;
+        }
+
+        ConstArrayView<std::uint16_t> getNeuralNetworkIndicesForLOD(std::uint16_t lod) const override {
+            return ConstArrayView<std::uint16_t>{unoptimized::neuralNetworkIndicesPerLOD[lod]};
+        }
+
+        ConstArrayView<std::uint16_t> getNeuralNetworkInputIndices(std::uint16_t neuralNetIndex) const override {
+            return ConstArrayView<std::uint16_t>{unoptimized::mlbNetInputIndices[neuralNetIndex]};
+        }
+
+        ConstArrayView<std::uint16_t> getNeuralNetworkOutputIndices(std::uint16_t neuralNetIndex) const override {
+            return ConstArrayView<std::uint16_t>{unoptimized::mlbNetOutputIndices[neuralNetIndex]};
+        }
+
+        std::uint16_t getNeuralNetworkLayerCount(std::uint16_t neuralNetIndex) const override {
+            return static_cast<std::uint16_t>(unoptimized::mlbNetWeights[neuralNetIndex].size());
+        }
+
+        dna::ActivationFunction getNeuralNetworkLayerActivationFunction(std::uint16_t neuralNetIndex,
+                                                                        std::uint16_t layerIndex) const override {
+            return unoptimized::mlbNetActivationFunctions[neuralNetIndex][layerIndex];
+        }
+
+        ConstArrayView<float> getNeuralNetworkLayerActivationFunctionParameters(std::uint16_t neuralNetIndex,
+                                                                                std::uint16_t layerIndex) const override {
+            return ConstArrayView<float>{unoptimized::mlbNetActivationFunctionParameters[neuralNetIndex][layerIndex]};
+        }
+
+        ConstArrayView<float> getNeuralNetworkLayerBiases(std::uint16_t neuralNetIndex, std::uint16_t layerIndex) const override {
+            return ConstArrayView<float>{unoptimized::mlbNetBiases[neuralNetIndex][layerIndex]};
+        }
+
+        ConstArrayView<float> getNeuralNetworkLayerWeights(std::uint16_t neuralNetIndex,
+                                                           std::uint16_t layerIndex) const override {
+            return ConstArrayView<float>{unoptimized::mlbNetWeights[neuralNetIndex][layerIndex]};
+        }
+
+};
+
+}  // namespace block4
+
+}  // namespace ml
+
+}  // namespace rltests
diff --git a/Engine/Plugins/Animation/RigLogic/Source/RigLogicLibTest/Private/rltests/ml/cpu/TestInference.cpp b/Engine/Plugins/Animation/RigLogic/Source/RigLogicLibTest/Private/rltests/ml/cpu/TestInference.cpp
new file mode 100644
--- /dev/null
+++ b/Engine/Plugins/Animation/RigLogic/Source/RigLogicLibTest/Private/rltests/ml/cpu/TestInference.cpp
@@ -0,0 +1,107 @@
+// Copyright Epic Games, Inc. All Rights Reserved.
+
+#include "rltests/Defs.h"
+#include "rltests/StorageValueType.h"
+#include "rltests/controls/ControlFixtures.h"
+#include "rltests/ml/cpu/FixturesBlock4.h"
+
+#include "riglogic/TypeDefs.h"
+#include "riglogic/ml/cpu/Factory.h"
+#include "riglogic/system/simd/Detect.h"
+
+#include <tuple>
+
+#ifdef _MSC_VER
+    #pragma warning(push)
+    #pragma warning(disable : 4324)
+#endif
+
+namespace {
+
+template<typename TTestTypes>
+class MLBInferenceTest : public ::testing::Test {
+    protected:
+        void SetUp() {
+            using T = typename std::tuple_element<0, TTestTypes>::type;
+            using TF256 = typename std::tuple_element<1, TTestTypes>::type;
+            using TF128 = typename std::tuple_element<2, TTestTypes>::type;
+            evaluator = rl4::ml::cpu::Factory<T, TF256, TF128>::create(&reader, &memRes);
+        }
+
+    protected:
+        pma::AlignedMemoryResource memRes;
+        rltests::ml::block4::CanonicalReader reader;
+        rl4::MachineLearnedBehaviorEvaluator::Pointer evaluator;
+
+};
+
+}  // namespace
+
+#if defined(RL_BUILD_WITH_AVX) && defined(RL_BUILD_WITH_SSE)
+    using InferenceTypeList = ::testing::Types<
+    #ifdef RL_USE_HALF_FLOATS
+            std::tuple<StorageValueType, trimd::avx::F256, trimd::F128>,
+            std::tuple<StorageValueType, trimd::sse::F256, trimd::sse::F128>
+    #else
+            std::tuple<StorageValueType, trimd::avx::F256, trimd::F128>,
+            std::tuple<StorageValueType, trimd::sse::F256, trimd::sse::F128>,
+            std::tuple<StorageValueType, trimd::scalar::F256, trimd::scalar::F128>
+    #endif  // RL_USE_HALF_FLOATS
+        >;
+#elif defined(RL_BUILD_WITH_AVX)
+    using InferenceTypeList = ::testing::Types<
+    #ifdef RL_USE_HALF_FLOATS
+            std::tuple<StorageValueType, trimd::avx::F256, trimd::F128>
+    #else
+            std::tuple<StorageValueType, trimd::avx::F256, trimd::F128>,
+            std::tuple<StorageValueType, trimd::scalar::F256, trimd::scalar::F128>
+    #endif  // RL_USE_HALF_FLOATS
+        >;
+#elif defined(RL_BUILD_WITH_SSE)
+    using InferenceTypeList = ::testing::Types<
+    #ifdef RL_USE_HALF_FLOATS
+            std::tuple<StorageValueType, trimd::sse::F256, trimd::sse::F128>
+    #else
+            std::tuple<StorageValueType, trimd::sse::F256, trimd::sse::F128>,
+            std::tuple<StorageValueType, trimd::scalar::F256, trimd::scalar::F128>
+    #endif  // RL_USE_HALF_FLOATS
+        >;
+#else
+    #ifndef RL_USE_HALF_FLOATS
+        using InferenceTypeList = ::testing::Types<std::tuple<StorageValueType, trimd::scalar::F256, trimd::scalar::F128> >;
+    #else
+        using InferenceTypeList = ::testing::Types<>;
+    #endif  // RL_USE_HALF_FLOATS
+#endif
+
+TYPED_TEST_SUITE(MLBInferenceTest, InferenceTypeList, );
+
+TYPED_TEST(MLBInferenceTest, InferencePerLOD) {
+    auto inputInstanceFactory = ControlsFactory::getInstanceFactory(0,
+                                                                    rltests::ml::block4::unoptimized::rawControlCount,
+                                                                    0,
+                                                                    rltests::ml::block4::unoptimized::mlControlCount);
+    auto inputInstance = inputInstanceFactory(&this->memRes);
+    auto inputBuffer = inputInstance->getInputBuffer();
+    auto outputBuffer = inputBuffer.subview(rltests::ml::block4::unoptimized::rawControlCount,
+                                            rltests::ml::block4::unoptimized::mlControlCount);
+    std::copy(rltests::ml::block4::input::values.begin(), rltests::ml::block4::input::values.end(), inputBuffer.begin());
+
+    auto intermediateOutputs = this->evaluator->createInstance(&this->memRes);
+
+    for (std::uint16_t lod = 0u; lod < rltests::ml::block4::unoptimized::lodCount; ++lod) {
+        std::fill(outputBuffer.begin(), outputBuffer.end(), 0.0f);
+        this->evaluator->calculate(inputInstance.get(), intermediateOutputs.get(), lod);
+        const auto& expected = rltests::ml::block4::output::valuesPerLOD[lod];
+        #ifdef RL_USE_HALF_FLOATS
+            static constexpr float threshold = 0.05f;
+        #else
+            static constexpr float threshold = 0.0001f;
+        #endif  // RL_USE_HALF_FLOATS
+        ASSERT_ELEMENTS_NEAR(outputBuffer, expected, expected.size(), threshold);
+    }
+}
+
+#ifdef _MSC_VER
+    #pragma warning(pop)
+#endif
diff --git a/Engine/Plugins/Animation/RigLogic/Source/RigLogicLibTest/Private/rltests/ml/cpu/TestStorageBuilder.cpp b/Engine/Plugins/Animation/RigLogic/Source/RigLogicLibTest/Private/rltests/ml/cpu/TestStorageBuilder.cpp
new file mode 100644
--- /dev/null
+++ b/Engine/Plugins/Animation/RigLogic/Source/RigLogicLibTest/Private/rltests/ml/cpu/TestStorageBuilder.cpp
@@ -0,0 +1,130 @@
+// Copyright Epic Games, Inc. All Rights Reserved.
+
+#include "rltests/Defs.h"
+#include "rltests/StorageValueType.h"
+#include "rltests/ml/cpu/FixturesBlock4.h"
+
+#include "riglogic/TypeDefs.h"
+#include "riglogic/ml/cpu/Factory.h"
+#include "riglogic/system/simd/Detect.h"
+
+#include <tuple>
+
+#ifdef _MSC_VER
+    #pragma warning(push)
+    #pragma warning(disable : 4324)
+#endif
+
+namespace rl4 {
+
+namespace ml {
+
+namespace cpu {
+
+template<typename T, typename TF256, typename TF128>
+struct Evaluator<T, TF256, TF128>::Accessor {
+
+    static void assertRawDataEqual(const Evaluator<T, TF256, TF128>& result) {
+        ASSERT_EQ(result.lods.netIndicesPerLOD, rltests::ml::block4::optimized::lods.netIndicesPerLOD);
+        ASSERT_EQ(result.lods.netCount, rltests::ml::block4::optimized::lods.netCount);
+        ASSERT_EQ(result.neuralNets.size(), rltests::ml::block4::optimized::lods.netCount);
+        for (std::uint16_t netIndex = {}; netIndex < rltests::ml::block4::unoptimized::neuralNetworkCount; ++netIndex) {
+            const auto& inference = result.neuralNets[netIndex];
+            ASSERT_EQ(inference.layerEvaluators.size(),
+                      rltests::ml::block4::unoptimized::mlbNetActivationFunctions[netIndex].size());
+            ASSERT_EQ(inference.neuralNet.layers.size(),
+                      rltests::ml::block4::unoptimized::mlbNetActivationFunctions[netIndex].size());
+            for (std::uint16_t layerIndex = {}; layerIndex < inference.neuralNet.layers.size(); ++layerIndex) {
+                const auto& layer = inference.neuralNet.layers[layerIndex];
+                using OptimizedValues = typename rltests::ml::block4::optimized::Values<T>;
+                ASSERT_EQ(layer.weights.values, OptimizedValues::weights()[netIndex][layerIndex]);
+                ASSERT_EQ(layer.biases, OptimizedValues::biases()[netIndex][layerIndex]);
+                ASSERT_EQ(layer.activationFunction,
+                          rltests::ml::block4::unoptimized::mlbNetActivationFunctions[netIndex][layerIndex]);
+                ASSERT_EQ(layer.activationFunctionParameters,
+                          rltests::ml::block4::unoptimized::mlbNetActivationFunctionParameters[netIndex][
+                              layerIndex]);
+            }
+            ASSERT_EQ(inference.neuralNet.inputIndices,
+                      rltests::ml::block4::unoptimized::mlbNetInputIndices[netIndex]);
+            ASSERT_EQ(inference.neuralNet.outputIndices,
+                      rltests::ml::block4::unoptimized::mlbNetOutputIndices[netIndex]);
+        }
+    }
+
+};
+
+}  // namespace cpu
+
+}  // namespace ml
+
+}  // namespace rl4
+
+namespace {
+
+template<typename TTestTypes>
+class MLBSStorageBuilderTest : public ::testing::Test {
+    protected:
+        void buildStorage() {
+            using T = typename std::tuple_element<0, TTestTypes>::type;
+            using TF256 = typename std::tuple_element<1, TTestTypes>::type;
+            using TF128 = typename std::tuple_element<2, TTestTypes>::type;
+            auto evaluator = rl4::ml::cpu::Factory<T, TF256, TF128>::create(&reader, &memRes);
+            auto evaluatorImpl = static_cast<rl4::ml::cpu::Evaluator<T, TF256, TF128>*>(evaluator.get());
+            rl4::ml::cpu::Evaluator<T, TF256, TF128>::Accessor::assertRawDataEqual(*evaluatorImpl);
+        }
+
+    protected:
+        pma::AlignedMemoryResource memRes;
+        rltests::ml::block4::CanonicalReader reader;
+
+};
+
+}  // namespace
+
+#if defined(RL_BUILD_WITH_AVX) && defined(RL_BUILD_WITH_SSE)
+    using StorageValueTypeList = ::testing::Types<
+    #ifdef RL_USE_HALF_FLOATS
+            std::tuple<StorageValueType, trimd::avx::F256, trimd::F128>,
+            std::tuple<StorageValueType, trimd::sse::F256, trimd::sse::F128>
+    #else
+            std::tuple<StorageValueType, trimd::avx::F256, trimd::F128>,
+            std::tuple<StorageValueType, trimd::sse::F256, trimd::sse::F128>,
+            std::tuple<StorageValueType, trimd::scalar::F256, trimd::scalar::F128>
+    #endif  // RL_USE_HALF_FLOATS
+        >;
+#elif defined(RL_BUILD_WITH_AVX)
+    using StorageValueTypeList = ::testing::Types<
+    #ifdef RL_USE_HALF_FLOATS
+            std::tuple<StorageValueType, trimd::avx::F256, trimd::F128>
+    #else
+            std::tuple<StorageValueType, trimd::avx::F256, trimd::F128>,
+            std::tuple<StorageValueType, trimd::scalar::F256, trimd::scalar::F128>
+    #endif  // RL_USE_HALF_FLOATS
+        >;
+#elif defined(RL_BUILD_WITH_SSE)
+    using StorageValueTypeList = ::testing::Types<
+    #ifdef RL_USE_HALF_FLOATS
+            std::tuple<StorageValueType, trimd::sse::F256, trimd::sse::F128>
+    #else
+            std::tuple<StorageValueType, trimd::sse::F256, trimd::sse::F128>,
+            std::tuple<StorageValueType, trimd::scalar::F256, trimd::scalar::F128>
+    #endif  // RL_USE_HALF_FLOATS
+        >;
+#else
+    #ifndef RL_USE_HALF_FLOATS
+        using StorageValueTypeList = ::testing::Types<std::tuple<StorageValueType, trimd::scalar::F256, trimd::scalar::F128> >;
+    #else
+        using StorageValueTypeList = ::testing::Types<>;
+    #endif  // RL_USE_HALF_FLOATS
+#endif
+
+TYPED_TEST_SUITE(MLBSStorageBuilderTest, StorageValueTypeList, );
+
+TYPED_TEST(MLBSStorageBuilderTest, LayoutOptimization) {
+    this->buildStorage();
+}
+
+#ifdef _MSC_VER
+    #pragma warning(pop)
+#endif
diff --git a/Engine/Plugins/Animation/RigLogic/Source/RigLogicLibTest/Private/rltests/psdmatrix/TestPSDMatrix.cpp b/Engine/Plugins/Animation/RigLogic/Source/RigLogicLibTest/Private/rltests/psdmatrix/TestPSDMatrix.cpp
--- a/Engine/Plugins/Animation/RigLogic/Source/RigLogicLibTest/Private/rltests/psdmatrix/TestPSDMatrix.cpp
+++ b/Engine/Plugins/Animation/RigLogic/Source/RigLogicLibTest/Private/rltests/psdmatrix/TestPSDMatrix.cpp
@@ -2,11 +2,9 @@
 
 #include "rltests/Defs.h"
 
+#include "riglogic/TypeDefs.h"
 #include "riglogic/psdmatrix/PSDMatrix.h"
 
-#include <pma/resources/AlignedMemoryResource.h>
-#include <pma/utils/ManagedInstance.h>
-
 TEST(PSDMatrixTest, OutputsAreClamped) {
     pma::AlignedMemoryResource amr;
 
diff --git a/Engine/Plugins/Animation/RigLogic/Source/RigLogicLibTest/Private/rltests/riglogic/TestRigLogic.cpp b/Engine/Plugins/Animation/RigLogic/Source/RigLogicLibTest/Private/rltests/riglogic/TestRigLogic.cpp
--- a/Engine/Plugins/Animation/RigLogic/Source/RigLogicLibTest/Private/rltests/riglogic/TestRigLogic.cpp
+++ b/Engine/Plugins/Animation/RigLogic/Source/RigLogicLibTest/Private/rltests/riglogic/TestRigLogic.cpp
@@ -4,8 +4,7 @@
 #include "rltests/dna/DNAFixtures.h"
 
 #include "riglogic/RigLogic.h"
-
-#include <pma/resources/AlignedMemoryResource.h>
+#include "riglogic/TypeDefs.h"
 
 namespace {
 
@@ -14,28 +13,22 @@ class RigLogicTest : public ::testing::Test {
         void SetUp() override {
             const auto bytes = rltests::raw::getBytes();
             stream = pma::makeScoped<trio::MemoryStream>();
-            stream->write(reinterpret_cast<const char*>(bytes.data()), bytes.size());
+            stream->write(bytes.data(), bytes.size());
             stream->seek(0);
 
-            reader = dna::BinaryStreamReader::create(stream.get());
+            reader = pma::makeScoped<dna::BinaryStreamReader>(stream.get());
             reader->read();
 
-            rigLogic = rl4::RigLogic::create(reader);
-            rigInstance = rl4::RigInstance::create(rigLogic, &memRes);
-        }
-
-        void TearDown() override {
-            rl4::RigInstance::destroy(rigInstance);
-            rl4::RigLogic::destroy(rigLogic);
-            dna::BinaryStreamReader::destroy(reader);
+            rigLogic = pma::makeScoped<rl4::RigLogic>(reader.get());
+            rigInstance = pma::makeScoped<rl4::RigInstance>(rigLogic.get(), &memRes);
         }
 
     protected:
         pma::AlignedMemoryResource memRes;
-        pma::ScopedPtr<trio::MemoryStream, pma::FactoryDestroy<trio::MemoryStream> > stream;
-        dna::BinaryStreamReader* reader;
-        rl4::RigLogic* rigLogic;
-        rl4::RigInstance* rigInstance;
+        pma::ScopedPtr<trio::MemoryStream> stream;
+        pma::ScopedPtr<dna::BinaryStreamReader> reader;
+        pma::ScopedPtr<rl4::RigLogic> rigLogic;
+        pma::ScopedPtr<rl4::RigInstance> rigInstance;
 };
 
 }  // namespace
@@ -44,12 +37,12 @@ TEST_F(RigLogicTest, EvaluateRigInstance) {
     // Try both approaches
     float guiControls[] = {1.0f, 2.0f, 3.0f, 4.0f, 5.0f, 6.0f, 7.0f, 8.0f, 9.0f};
     rigInstance->setGUIControlValues(static_cast<const float*>(guiControls));
-    rigLogic->mapGUIToRawControls(rigInstance);
+    rigLogic->mapGUIToRawControls(rigInstance.get());
     // Regardless that this overwrites the computed gui to raw values
     float rawControls[] = {1.0f, 2.0f, 3.0f, 4.0f, 5.0f, 6.0f, 7.0f, 8.0f, 9.0f};
     rigInstance->setRawControlValues(static_cast<const float*>(rawControls));
 
-    rigLogic->calculate(rigInstance);
+    rigLogic->calculate(rigInstance.get());
 
     ASSERT_EQ(rigInstance->getRawJointOutputs().size(), reader->getJointRowCount());
     ASSERT_EQ(rigInstance->getBlendShapeOutputs().size(), reader->getBlendShapeChannelCount());
@@ -69,14 +62,14 @@ TEST_F(RigLogicTest, DumpStateThenRestore) {
     auto dumpedState = pma::makeScoped<trio::MemoryStream>();
     rigLogic->dump(dumpedState.get());
     dumpedState->seek(0);
-    auto cloneRigLogic = rl4::RigLogic::restore(dumpedState.get());
-    auto cloneRigInstance = rl4::RigInstance::create(cloneRigLogic);
+    auto cloneRigLogic = rl4::RigLogic::restore(dumpedState.get(), &memRes);
+    auto cloneRigInstance = rl4::RigInstance::create(cloneRigLogic, &memRes);
 
     for (std::uint16_t lod = 0u; lod < rigLogic->getLODCount(); ++lod) {
         rigInstance->setLOD(lod);
         rigInstance->setGUIControlValues(static_cast<const float*>(guiControls));
-        rigLogic->mapGUIToRawControls(rigInstance);
-        rigLogic->calculate(rigInstance);
+        rigLogic->mapGUIToRawControls(rigInstance.get());
+        rigLogic->calculate(rigInstance.get());
 
         cloneRigInstance->setLOD(lod);
         cloneRigInstance->setGUIControlValues(static_cast<const float*>(guiControls));
diff --git a/Engine/Plugins/Animation/RigLogic/Source/RigLogicLibTest/Private/sctests/Defs.h b/Engine/Plugins/Animation/RigLogic/Source/RigLogicLibTest/Private/sctests/Defs.h
new file mode 100644
--- /dev/null
+++ b/Engine/Plugins/Animation/RigLogic/Source/RigLogicLibTest/Private/sctests/Defs.h
@@ -0,0 +1,65 @@
+// Copyright Epic Games, Inc. All Rights Reserved.
+
+#pragma once
+
+#ifdef _MSC_VER
+    #pragma warning(push)
+    #pragma warning(disable : 4061 4121 4365 4987 4668 4996 6330)
+#endif
+#ifdef __clang__
+    #pragma clang diagnostic push
+    #pragma clang diagnostic ignored "-Wundef"
+#endif
+#include <gtest/gtest.h>
+#ifdef __clang__
+    #pragma clang diagnostic pop
+    // This cannot be applied to the scope of the header inclusion only, as the TEST macros will trigger it
+    #pragma clang diagnostic ignored "-Wglobal-constructors"
+#endif
+#ifdef _MSC_VER
+    #pragma warning(pop)
+#endif
+
+#ifndef INSTANTIATE_TEST_SUITE_P
+    #define INSTANTIATE_TEST_SUITE_P INSTANTIATE_TEST_CASE_P
+#endif
+
+#ifndef TYPED_TEST_SUITE_P
+    #define TYPED_TEST_SUITE_P TYPED_TEST_CASE_P
+#endif
+
+#ifndef REGISTER_TYPED_TEST_SUITE_P
+    #define REGISTER_TYPED_TEST_SUITE_P REGISTER_TYPED_TEST_CASE_P
+#endif
+
+#ifndef INSTANTIATE_TYPED_TEST_SUITE_P
+    #define INSTANTIATE_TYPED_TEST_SUITE_P INSTANTIATE_TYPED_TEST_CASE_P
+#endif
+
+#ifndef ASSERT_ELEMENTS_EQ
+    #define ASSERT_ELEMENTS_EQ(result, expected, count)     \
+        for (std::size_t i = 0ul; i < count; ++i) {         \
+            ASSERT_EQ(result[i], expected[i]);              \
+        }
+#endif
+
+#ifndef EXPECT_ELEMENTS_EQ
+    #define EXPECT_ELEMENTS_EQ(result, expected, count)     \
+        for (std::size_t i = 0ul; i < count; ++i) {         \
+            EXPECT_EQ(result[i], expected[i]);              \
+        }
+#endif
+
+#ifndef ASSERT_ELEMENTS_NEAR
+    #define ASSERT_ELEMENTS_NEAR(result, expected, count, threshold)    \
+        for (std::size_t i = 0ul; i < count; ++i) {                     \
+            ASSERT_NEAR(result[i], expected[i], threshold);             \
+        }
+#endif
+
+#ifndef EXPECT_ELEMENTS_NEAR
+    #define EXPECT_ELEMENTS_NEAR(result, expected, count, threshold)    \
+        for (std::size_t i = 0ul; i < count; ++i) {                     \
+            EXPECT_NEAR(result[i], expected[i], threshold);             \
+        }
+#endif
diff --git a/Engine/Plugins/Animation/RigLogic/Source/RigLogicLibTest/Private/tersetests/archives/binary/Common.h b/Engine/Plugins/Animation/RigLogic/Source/RigLogicLibTest/Private/tersetests/archives/binary/Common.h
--- a/Engine/Plugins/Animation/RigLogic/Source/RigLogicLibTest/Private/tersetests/archives/binary/Common.h
+++ b/Engine/Plugins/Animation/RigLogic/Source/RigLogicLibTest/Private/tersetests/archives/binary/Common.h
@@ -127,22 +127,22 @@ struct UpgradedSerializableType {
     std::uint16_t b;
 
     template<class Archive>
-    void load(Archive& archive, terse::Version<1>) {
+    void load(Archive& archive, terse::Version<1>  /*unused*/) {
         archive(a);
     }
 
     template<class Archive>
-    void load(Archive& archive, terse::Version<2>) {
+    void load(Archive& archive, terse::Version<2>  /*unused*/) {
         archive(a, b);
     }
 
     template<class Archive>
-    void save(Archive& archive, terse::Version<1>) {
+    void save(Archive& archive, terse::Version<1>  /*unused*/) {
         archive(a);
     }
 
     template<class Archive>
-    void save(Archive& archive, terse::Version<2>) {
+    void save(Archive& archive, terse::Version<2>  /*unused*/) {
         archive(a, b);
     }
 
diff --git a/Engine/Plugins/Animation/RigLogic/Source/RigLogicLibTest/Private/triotests/TestStreamsIntegration.cpp b/Engine/Plugins/Animation/RigLogic/Source/RigLogicLibTest/Private/triotests/TestStreamsIntegration.cpp
--- a/Engine/Plugins/Animation/RigLogic/Source/RigLogicLibTest/Private/triotests/TestStreamsIntegration.cpp
+++ b/Engine/Plugins/Animation/RigLogic/Source/RigLogicLibTest/Private/triotests/TestStreamsIntegration.cpp
@@ -141,7 +141,14 @@ TYPED_TEST(StreamTest, LFSIntegrationTest) {
     ASSERT_STATUS_OK();
     ASSERT_EQ(stream->tell(), (expectedPosition += smallBuffer.size()));
 
+    #if !defined(__clang__) && defined(__GNUC__)
+        #pragma GCC diagnostic push
+        #pragma GCC diagnostic ignored "-Wuseless-cast"
+    #endif
     const std::uint64_t expectedSize = (5ull * static_cast<std::uint64_t>(largeBuffer.size())) + (6ull * static_cast<std::uint64_t>(smallBuffer.size()));
+    #if !defined(__clang__) && defined(__GNUC__)
+        #pragma GCC diagnostic pop
+    #endif
     ASSERT_EQ(stream->size(), expectedSize);
 
     // Verify written data
diff --git a/Engine/Plugins/Animation/RigLogic/Source/RigLogicLibTest/RigLogicLibTest.Build.cs b/Engine/Plugins/Animation/RigLogic/Source/RigLogicLibTest/RigLogicLibTest.Build.cs
--- a/Engine/Plugins/Animation/RigLogic/Source/RigLogicLibTest/RigLogicLibTest.Build.cs
+++ b/Engine/Plugins/Animation/RigLogic/Source/RigLogicLibTest/RigLogicLibTest.Build.cs
@@ -1,41 +1,42 @@
 // Copyright Epic Games, Inc. All Rights Reserved.
 
-using System;
 using System.IO;
-using UnrealBuildTool;
 using EpicGames.Core;
 
-public class RigLogicLibTest : ModuleRules
+namespace UnrealBuildTool.Rules
 {
-    public RigLogicLibTest(ReadOnlyTargetRules Target) : base(Target)
-    {
-		bUseUnity = false; // A windows include is preprocessing some method names causing compile failures.
+	public class RigLogicLibTest : ModuleRules
+	{
+		public RigLogicLibTest(ReadOnlyTargetRules Target) : base(Target)
+		{
+			bUseUnity = false; // A windows include is preprocessing some method names causing compile failures.
 
-        if (Target.Platform == UnrealTargetPlatform.Win64)
-        {
-            PrivateDefinitions.Add("RL_BUILD_WITH_SSE=1");
-            PublicDefinitions.Add("GTEST_OS_WINDOWS=1");
-        }
+			if (Target.Platform == UnrealTargetPlatform.Win64)
+			{
+				PrivateDefinitions.Add("RL_BUILD_WITH_SSE=1");
+				PublicDefinitions.Add("GTEST_OS_WINDOWS=1");
+			}
 
-        string RigLogicLibPath = Path.GetFullPath(Path.Combine(ModuleDirectory, "../RigLogicLib"));
+			string RigLogicLibPath = Path.GetFullPath(Path.Combine(ModuleDirectory, "../RigLogicLib"));
 
-        if (Target.LinkType == TargetLinkType.Monolithic)
-        {
-            PublicDependencyModuleNames.Add("RigLogicLib");
-            PrivateIncludePaths.Add(Path.Combine(RigLogicLibPath, "Private"));
-        }
-        else
-        {
-            PrivateDefinitions.Add("RIGLOGIC_MODULE_DISCARD");
-            ConditionalAddModuleDirectory(new DirectoryReference(RigLogicLibPath));
-        }
+			if (Target.LinkType == TargetLinkType.Monolithic)
+			{
+				PublicDependencyModuleNames.Add("RigLogicLib");
+				PrivateIncludePaths.Add(Path.Combine(RigLogicLibPath, "Private"));
+			}
+			else
+			{
+				PrivateDefinitions.Add("RIGLOGIC_MODULE_DISCARD");
+				ConditionalAddModuleDirectory(new DirectoryReference(RigLogicLibPath));
+			}
 
-        PublicDependencyModuleNames.AddRange(
-            new string[]
-            {
-                "Core",
-                "GoogleTest"
-            }
-        );
-    }
+			PublicDependencyModuleNames.AddRange(
+				new string[]
+				{
+					"Core",
+					"GoogleTest"
+				}
+			);
+		}
+	}
 }
diff --git a/Engine/Plugins/Animation/RigLogic/Source/RigLogicModule/Private/DNAAsset.cpp b/Engine/Plugins/Animation/RigLogic/Source/RigLogicModule/Private/DNAAsset.cpp
--- a/Engine/Plugins/Animation/RigLogic/Source/RigLogicModule/Private/DNAAsset.cpp
+++ b/Engine/Plugins/Animation/RigLogic/Source/RigLogicModule/Private/DNAAsset.cpp
@@ -33,11 +33,12 @@ DEFINE_LOG_CATEGORY(LogDNAAsset);
 
 static constexpr uint32 AVG_EMPTY_SIZE = 4 * 1024;
 static constexpr uint32 AVG_BEHAVIOR_SIZE = 5 * 1024 * 1024;
+static constexpr uint32 AVG_MACHINE_LEARNED_BEHAVIOR_SIZE = 15 * 1024 * 1024;
 static constexpr uint32 AVG_GEOMETRY_SIZE = 50 * 1024 * 1024;
 
 static TSharedPtr<IDNAReader> ReadDNAFromStream(rl4::BoundedIOStream* Stream, EDNADataLayer Layer, uint16 MaxLOD)
 {
-	auto DNAStreamReader = rl4::makeScoped<dna::BinaryStreamReader>(Stream, static_cast<dna::DataLayer>(Layer), MaxLOD, FMemoryResource::Instance());
+	auto DNAStreamReader = rl4::makeScoped<dna::BinaryStreamReader>(Stream, static_cast<dna::DataLayer>(Layer), dna::UnknownLayerPolicy::Preserve, MaxLOD, FMemoryResource::Instance());
 	DNAStreamReader->read();
 	if (!rl4::Status::isOk())
 	{
@@ -52,7 +53,7 @@ static void WriteDNAToStream(const IDNAReader* Source, EDNADataLayer Layer, rl4:
 	auto DNAWriter = rl4::makeScoped<dna::BinaryStreamWriter>(Destination, FMemoryResource::Instance());
 	if (Source != nullptr)
 	{
-		DNAWriter->setFrom(Source->Unwrap(), static_cast<dna::DataLayer>(Layer), FMemoryResource::Instance());
+		DNAWriter->setFrom(Source->Unwrap(), static_cast<dna::DataLayer>(Layer), dna::UnknownLayerPolicy::Preserve, FMemoryResource::Instance());
 	}
 	DNAWriter->write();
 }
@@ -93,14 +94,14 @@ UDNAAsset::UDNAAsset() : RigRuntimeContext{nullptr}
 
 UDNAAsset::~UDNAAsset() = default;
 
-TSharedPtr<IBehaviorReader> UDNAAsset::GetBehaviorReader()
+TSharedPtr<IDNAReader> UDNAAsset::GetBehaviorReader()
 {
 	FScopeLock ScopeLock{&DNAUpdateSection};
 	return BehaviorReader;
 }
 
 #if WITH_EDITORONLY_DATA
-TSharedPtr<IGeometryReader> UDNAAsset::GetGeometryReader()
+TSharedPtr<IDNAReader> UDNAAsset::GetGeometryReader()
 {
 	FScopeLock ScopeLock{&DNAUpdateSection};
 	return GeometryReader;
@@ -110,7 +111,8 @@ TSharedPtr<IGeometryReader> UDNAAsset::GetGeometryReader()
 void UDNAAsset::SetBehaviorReader(TSharedPtr<IDNAReader> SourceDNAReader)
 {
 	FScopeLock ScopeLock{&DNAUpdateSection};
-	BehaviorReader = CopyDNALayer(SourceDNAReader.Get(), EDNADataLayer::Behavior, AVG_BEHAVIOR_SIZE);
+	const size_t PredictedSize = (SourceDNAReader->GetNeuralNetworkCount() != 0) ? AVG_BEHAVIOR_SIZE + AVG_MACHINE_LEARNED_BEHAVIOR_SIZE : AVG_BEHAVIOR_SIZE;
+	BehaviorReader = CopyDNALayer(SourceDNAReader.Get(), EDNADataLayer::Behavior | EDNADataLayer::MachineLearnedBehavior, PredictedSize);
 	InvalidateRigRuntimeContext();
 }
 
@@ -153,6 +155,7 @@ TSharedPtr<FSharedRigRuntimeContext> UDNAAsset::GetRigRuntimeContext(EDNARetenti
 				{
 					BehaviorReader->Unload(EDNADataLayer::Behavior);
 					BehaviorReader->Unload(EDNADataLayer::Geometry);
+					BehaviorReader->Unload(EDNADataLayer::MachineLearnedBehavior);
 				}
 #endif  // !WITH_EDITOR
 			}
diff --git a/Engine/Plugins/Animation/RigLogic/Source/RigLogicModule/Private/DNAUtils.cpp b/Engine/Plugins/Animation/RigLogic/Source/RigLogicModule/Private/DNAUtils.cpp
--- a/Engine/Plugins/Animation/RigLogic/Source/RigLogicModule/Private/DNAUtils.cpp
+++ b/Engine/Plugins/Animation/RigLogic/Source/RigLogicModule/Private/DNAUtils.cpp
@@ -26,7 +26,7 @@ TSharedPtr<IDNAReader> ReadDNAFromFile(const FString& Path, EDNADataLayer Layer,
 {
 	LLM_SCOPE_BYNAME(TEXT("Animation/RigLogic"));
 	auto DNAFileStream = rl4::makeScoped<rl4::MemoryMappedFileStream>(TCHAR_TO_UTF8(*Path), rl4::MemoryMappedFileStream::AccessMode::Read, FMemoryResource::Instance());
-	auto DNAStreamReader = rl4::makeScoped<dna::BinaryStreamReader>(DNAFileStream.get(), static_cast<dna::DataLayer>(Layer), MaxLOD, FMemoryResource::Instance());
+	auto DNAStreamReader = rl4::makeScoped<dna::BinaryStreamReader>(DNAFileStream.get(), static_cast<dna::DataLayer>(Layer), dna::UnknownLayerPolicy::Preserve, MaxLOD, FMemoryResource::Instance());
 	return ReadDNAStream(MoveTemp(DNAStreamReader));
 }
 
@@ -34,7 +34,7 @@ TSharedPtr<IDNAReader> ReadDNAFromFile(const FString& Path, EDNADataLayer Layer,
 {
 	LLM_SCOPE_BYNAME(TEXT("Animation/RigLogic"));
 	auto DNAFileStream = rl4::makeScoped<rl4::MemoryMappedFileStream>(TCHAR_TO_UTF8(*Path), rl4::MemoryMappedFileStream::AccessMode::Read, FMemoryResource::Instance());
-	auto DNAStreamReader = rl4::makeScoped<dna::BinaryStreamReader>(DNAFileStream.get(), static_cast<dna::DataLayer>(Layer), LODs.GetData(), static_cast<uint16>(LODs.Num()), FMemoryResource::Instance());
+	auto DNAStreamReader = rl4::makeScoped<dna::BinaryStreamReader>(DNAFileStream.get(), static_cast<dna::DataLayer>(Layer), dna::UnknownLayerPolicy::Preserve, LODs.GetData(), static_cast<uint16>(LODs.Num()), FMemoryResource::Instance());
 	return ReadDNAStream(MoveTemp(DNAStreamReader));
 }
 
@@ -42,7 +42,7 @@ TSharedPtr<IDNAReader> ReadDNAFromBuffer(TArray<uint8>* DNABuffer, EDNADataLayer
 {
 	LLM_SCOPE_BYNAME(TEXT("Animation/RigLogic"));
 	FRigLogicMemoryStream DNAMemoryStream(DNABuffer);
-	auto DNAStreamReader = rl4::makeScoped<dna::BinaryStreamReader>(&DNAMemoryStream, static_cast<dna::DataLayer>(Layer), MaxLOD, FMemoryResource::Instance());
+	auto DNAStreamReader = rl4::makeScoped<dna::BinaryStreamReader>(&DNAMemoryStream, static_cast<dna::DataLayer>(Layer), dna::UnknownLayerPolicy::Preserve, MaxLOD, FMemoryResource::Instance());
 	return ReadDNAStream(MoveTemp(DNAStreamReader));
 }
 
@@ -50,7 +50,7 @@ TSharedPtr<IDNAReader> ReadDNAFromBuffer(TArray<uint8>* DNABuffer, EDNADataLayer
 {
 	LLM_SCOPE_BYNAME(TEXT("Animation/RigLogic"));
 	FRigLogicMemoryStream DNAMemoryStream(DNABuffer);
-	auto DNAStreamReader = rl4::makeScoped<dna::BinaryStreamReader>(&DNAMemoryStream, static_cast<dna::DataLayer>(Layer), LODs.GetData(), static_cast<uint16>(LODs.Num()), FMemoryResource::Instance());
+	auto DNAStreamReader = rl4::makeScoped<dna::BinaryStreamReader>(&DNAMemoryStream, static_cast<dna::DataLayer>(Layer), dna::UnknownLayerPolicy::Preserve, LODs.GetData(), static_cast<uint16>(LODs.Num()), FMemoryResource::Instance());
 	return ReadDNAStream(MoveTemp(DNAStreamReader));
 }
 
@@ -60,7 +60,7 @@ TArray<uint8> ReadStreamFromDNA(const IDNAReader* Reader, EDNADataLayer Layer)
 	TArray<char> DNABuffer;
 	auto DeltaDnaStream = rl4::makeScoped<trio::MemoryStream>();
 	auto DNAStreamWriter = rl4::makeScoped<dna::BinaryStreamWriter>(DeltaDnaStream.get(), FMemoryResource::Instance());
-	DNAStreamWriter->setFrom(Reader->Unwrap(), static_cast<dna::DataLayer>(Layer), FMemoryResource::Instance());
+	DNAStreamWriter->setFrom(Reader->Unwrap(), static_cast<dna::DataLayer>(Layer), dna::UnknownLayerPolicy::Preserve, FMemoryResource::Instance());
 	DNAStreamWriter->write();
 	DNABuffer.AddZeroed(DeltaDnaStream->size());
 	DeltaDnaStream->read(DNABuffer.GetData(), DeltaDnaStream->size());
@@ -72,6 +72,6 @@ void WriteDNAToFile(const IDNAReader* Reader, EDNADataLayer Layer, const FString
 	LLM_SCOPE_BYNAME(TEXT("Animation/RigLogic"));
 	auto DNAFileStream = rl4::makeScoped<rl4::FileStream>(TCHAR_TO_UTF8(*Path), rl4::FileStream::AccessMode::Write, rl4::FileStream::OpenMode::Binary, FMemoryResource::Instance());
 	auto DNAStreamWriter = rl4::makeScoped<dna::BinaryStreamWriter>(DNAFileStream.get(), FMemoryResource::Instance());
-	DNAStreamWriter->setFrom(Reader->Unwrap(), static_cast<dna::DataLayer>(Layer), FMemoryResource::Instance());
+	DNAStreamWriter->setFrom(Reader->Unwrap(), static_cast<dna::DataLayer>(Layer), dna::UnknownLayerPolicy::Preserve, FMemoryResource::Instance());
 	DNAStreamWriter->write();
 }
diff --git a/Engine/Plugins/Animation/RigLogic/Source/RigLogicModule/Private/RigInstance.cpp b/Engine/Plugins/Animation/RigLogic/Source/RigLogicModule/Private/RigInstance.cpp
--- a/Engine/Plugins/Animation/RigLogic/Source/RigLogicModule/Private/RigInstance.cpp
+++ b/Engine/Plugins/Animation/RigLogic/Source/RigLogicModule/Private/RigInstance.cpp
@@ -25,11 +25,22 @@ uint16 FRigInstance::GetGUIControlCount() const
 	return RigInstance->getGUIControlCount();
 }
 
+float FRigInstance::GetGUIControl(uint16 Index) const
+{
+	return RigInstance->getGUIControl(Index);
+}
+
 void FRigInstance::SetGUIControl(uint16 Index, float Value)
 {
 	RigInstance->setGUIControl(Index, Value);
 }
 
+TArrayView<const float> FRigInstance::GetGUIControlValues() const
+{
+	rl4::ConstArrayView<float> Values = RigInstance->getGUIControlValues();
+	return TArrayView<const float>{Values.data(), static_cast<int32>(Values.size())};
+}
+
 void FRigInstance::SetGUIControlValues(const float* Values)
 {
 	RigInstance->setGUIControlValues(Values);
@@ -40,16 +51,74 @@ uint16 FRigInstance::GetRawControlCount() const
 	return RigInstance->getRawControlCount();
 }
 
+float FRigInstance::GetRawControl(uint16 Index) const
+{
+	return RigInstance->getRawControl(Index);
+}
+
 void FRigInstance::SetRawControl(uint16 Index, float Value)
 {
 	RigInstance->setRawControl(Index, Value);
 }
 
+TArrayView<const float> FRigInstance::GetRawControlValues() const
+{
+	rl4::ConstArrayView<float> Values = RigInstance->getRawControlValues();
+	return TArrayView<const float>{Values.data(), static_cast<int32>(Values.size())};
+}
+
 void FRigInstance::SetRawControlValues(const float* Values)
 {
 	RigInstance->setRawControlValues(Values);
 }
 
+uint16 FRigInstance::GetPSDControlCount() const
+{
+	return RigInstance->getPSDControlCount();
+}
+
+float FRigInstance::GetPSDControl(uint16 Index) const
+{
+	return RigInstance->getPSDControl(Index);
+}
+
+TArrayView<const float> FRigInstance::GetPSDControlValues() const
+{
+	rl4::ConstArrayView<float> Values = RigInstance->getPSDControlValues();
+	return TArrayView<const float>{Values.data(), static_cast<int32>(Values.size())};
+}
+
+uint16 FRigInstance::GetMLControlCount() const
+{
+	return RigInstance->getMLControlCount();
+}
+
+float FRigInstance::GetMLControl(uint16 Index) const
+{
+	return RigInstance->getMLControl(Index);
+}
+
+TArrayView<const float> FRigInstance::GetMLControlValues() const
+{
+	rl4::ConstArrayView<float> Values = RigInstance->getMLControlValues();
+	return TArrayView<const float>{Values.data(), static_cast<int32>(Values.size())};
+}
+
+uint16 FRigInstance::GetNeuralNetworkCount() const
+{
+	return RigInstance->getNeuralNetworkCount();
+}
+
+float FRigInstance::GetNeuralNetworkMask(uint16 NeuralNetIndex) const
+{
+	return RigInstance->getNeuralNetworkMask(NeuralNetIndex);
+}
+
+void FRigInstance::SetNeuralNetworkMask(uint16 NeuralNetIndex, float Value)
+{
+	RigInstance->setNeuralNetworkMask(NeuralNetIndex, Value);
+}
+
 TArrayView<const float> FRigInstance::GetRawJointOutputs() const
 {
 	rl4::ConstArrayView<float> Outputs = RigInstance->getRawJointOutputs();
diff --git a/Engine/Plugins/Animation/RigLogic/Source/RigLogicModule/Private/RigLogic.cpp b/Engine/Plugins/Animation/RigLogic/Source/RigLogicModule/Private/RigLogic.cpp
--- a/Engine/Plugins/Animation/RigLogic/Source/RigLogicModule/Private/RigLogic.cpp
+++ b/Engine/Plugins/Animation/RigLogic/Source/RigLogicModule/Private/RigLogic.cpp
@@ -13,11 +13,10 @@ void FRigLogic::FRigLogicDeleter::operator()(rl4::RigLogic* Pointer)
 	rl4::RigLogic::destroy(Pointer);
 }
 
-FRigLogic::FRigLogic(const IBehaviorReader* Reader, ERigLogicCalculationType CalculationType) :
+FRigLogic::FRigLogic(const IDNAReader* Reader, ERigLogicCalculationType CalculationType) :
 	MemoryResource{FMemoryResource::SharedInstance()},
-	RigLogic{rl4::RigLogic::create(Reader->Unwrap(), rl4::CalculationType::SSE, FMemoryResource::Instance())}
+	RigLogic{rl4::RigLogic::create(Reader->Unwrap(), rl4::Configuration{static_cast<rl4::CalculationType>(CalculationType)}, FMemoryResource::Instance())}
 {
-
 }
 
 FRigLogic::~FRigLogic() = default;
@@ -39,14 +38,35 @@ FTransformArrayView FRigLogic::GetNeutralJointValues() const
 	return FTransformArrayView{Values.data(), Values.size()};
 }
 
+TArrayView<const uint16> FRigLogic::GetJointVariableAttributeIndices(uint16 LOD) const
+{
+	rl4::ConstArrayView<uint16> Indices = RigLogic->getJointVariableAttributeIndices(LOD);
+	return TArrayView<const uint16>{Indices.data(), static_cast<int32>(Indices.size())};
+}
+
 uint16 FRigLogic::GetJointGroupCount() const
 {
 	return RigLogic->getJointGroupCount();
 }
 
-TArrayView<const uint16> FRigLogic::GetJointVariableAttributeIndices(uint16 LOD) const
+uint16 FRigLogic::GetNeuralNetworkCount() const
 {
-	rl4::ConstArrayView<uint16> Indices = RigLogic->getJointVariableAttributeIndices(LOD);
+	return RigLogic->getNeuralNetworkCount();
+}
+
+uint16 FRigLogic::GetMeshCount() const
+{
+	return RigLogic->getMeshCount();
+}
+
+uint16 FRigLogic::GetMeshRegionCount(uint16 MeshIndex) const
+{
+	return RigLogic->getMeshRegionCount(MeshIndex);
+}
+
+TArrayView<const uint16> FRigLogic::GetNeuralNetworkIndices(uint16 MeshIndex, uint16 RegionIndex) const
+{
+	rl4::ConstArrayView<uint16> Indices = RigLogic->getNeuralNetworkIndices(MeshIndex, RegionIndex);
 	return TArrayView<const uint16>{Indices.data(), static_cast<int32>(Indices.size())};
 }
 
@@ -55,11 +75,26 @@ void FRigLogic::MapGUIToRawControls(FRigInstance* Instance) const
 	RigLogic->mapGUIToRawControls(Instance->Unwrap());
 }
 
+void FRigLogic::MapRawToGUIControls(FRigInstance* Instance) const
+{
+	RigLogic->mapRawToGUIControls(Instance->Unwrap());
+}
+
 void FRigLogic::CalculateControls(FRigInstance* Instance) const
 {
 	RigLogic->calculateControls(Instance->Unwrap());
 }
 
+void FRigLogic::CalculateMachineLearnedBehaviorControls(FRigInstance* Instance) const
+{
+	RigLogic->calculateMachineLearnedBehaviorControls(Instance->Unwrap());
+}
+
+void FRigLogic::CalculateMachineLearnedBehaviorControls(FRigInstance* Instance, uint16 NeuralNetIndex) const
+{
+	RigLogic->calculateMachineLearnedBehaviorControls(Instance->Unwrap(), NeuralNetIndex);
+}
+
 void FRigLogic::CalculateJoints(FRigInstance* Instance) const
 {
 	RigLogic->calculateJoints(Instance->Unwrap());
