Update RigLogic to v9.1.1 and GeneSplicer to v9.4.0

diff --git a/Engine/Plugins/Animation/RigLogic/Source/RigLogicModule/Private/SkelMeshDNAReader.cpp b/Engine/Plugins/Animation/RigLogic/Source/RigLogicModule/Private/SkelMeshDNAReader.cpp
--- a/Engine/Plugins/Animation/RigLogic/Source/RigLogicModule/Private/SkelMeshDNAReader.cpp
+++ b/Engine/Plugins/Animation/RigLogic/Source/RigLogicModule/Private/SkelMeshDNAReader.cpp
@@ -17,11 +17,11 @@
 
 DEFINE_LOG_CATEGORY(LogSkelMeshDNAReader);
 
-FSkelMeshDNAReader::FSkelMeshDNAReader(UDNAAsset* DNAAsset) : GeometryStreamReader{ nullptr }
+FSkelMeshDNAReader::FSkelMeshDNAReader(UDNAAsset* DNAAsset) : GeometryReader{nullptr}
 {
-	BehaviorStreamReader = DNAAsset->GetBehaviorReader().Get();
+	BehaviorReader = DNAAsset->GetBehaviorReader();
 #if WITH_EDITORONLY_DATA
-	GeometryStreamReader = DNAAsset->GetGeometryReader().Get();
+	GeometryReader = DNAAsset->GetGeometryReader();
 #endif
 }
 
@@ -30,435 +30,446 @@ dna::Reader* FSkelMeshDNAReader::Unwrap() const
 	return nullptr;  // Unused in SkelMeshDNAReader
 }
 
+/** HEADER READER **/
+uint16 FSkelMeshDNAReader::GetFileFormatGeneration() const
+{
+	return BehaviorReader->GetFileFormatGeneration();
+}
+
+uint16 FSkelMeshDNAReader::GetFileFormatVersion() const
+{
+	return BehaviorReader->GetFileFormatVersion();
+}
+
 /** DESCRIPTOR READER **/
 
 FString FSkelMeshDNAReader::GetName() const
 {
-	return BehaviorStreamReader->GetName();
+	return BehaviorReader->GetName();
 }
 
 EArchetype FSkelMeshDNAReader::GetArchetype() const
 {
-	return BehaviorStreamReader->GetArchetype();
+	return BehaviorReader->GetArchetype();
 }
 
 EGender FSkelMeshDNAReader::GetGender() const
 {
-	return BehaviorStreamReader->GetGender();
+	return BehaviorReader->GetGender();
 }
 
 uint16 FSkelMeshDNAReader::GetAge() const
 {
-	return BehaviorStreamReader->GetAge();
+	return BehaviorReader->GetAge();
 }
 
 uint32 FSkelMeshDNAReader::GetMetaDataCount() const
 {
-	return BehaviorStreamReader->GetMetaDataCount();
+	return BehaviorReader->GetMetaDataCount();
 }
 
 FString FSkelMeshDNAReader::GetMetaDataKey(uint32 Index) const
 {
-	return BehaviorStreamReader->GetMetaDataKey(Index);
+	return BehaviorReader->GetMetaDataKey(Index);
 }
 
 FString FSkelMeshDNAReader::GetMetaDataValue(const FString& Key) const
 {
-	return BehaviorStreamReader->GetMetaDataValue(Key);
+	return BehaviorReader->GetMetaDataValue(Key);
 }
 
 ETranslationUnit FSkelMeshDNAReader::GetTranslationUnit() const
 {
-	return BehaviorStreamReader->GetTranslationUnit();
+	return BehaviorReader->GetTranslationUnit();
 }
 
 ERotationUnit FSkelMeshDNAReader::GetRotationUnit() const
 {
-	return BehaviorStreamReader->GetRotationUnit();
+	return BehaviorReader->GetRotationUnit();
 }
 
 FCoordinateSystem FSkelMeshDNAReader::GetCoordinateSystem() const
 {
-	return BehaviorStreamReader->GetCoordinateSystem();
+	return BehaviorReader->GetCoordinateSystem();
 }
 
 uint16 FSkelMeshDNAReader::GetLODCount() const
 {
-	return BehaviorStreamReader->GetLODCount();
+	return BehaviorReader->GetLODCount();
 }
 
 uint16 FSkelMeshDNAReader::GetDBMaxLOD() const
 {
-	return BehaviorStreamReader->GetDBMaxLOD();
+	return BehaviorReader->GetDBMaxLOD();
 }
 
 FString FSkelMeshDNAReader::GetDBComplexity() const
 {
-	return BehaviorStreamReader->GetDBComplexity();
+	return BehaviorReader->GetDBComplexity();
 }
 
 FString FSkelMeshDNAReader::GetDBName() const
 {
-	return BehaviorStreamReader->GetDBName();
+	return BehaviorReader->GetDBName();
 }
 
 /** DEFINITION READER **/
 
 uint16 FSkelMeshDNAReader::GetGUIControlCount() const
 {
-	return BehaviorStreamReader->GetGUIControlCount();
+	return BehaviorReader->GetGUIControlCount();
 }
 
 FString FSkelMeshDNAReader::GetGUIControlName(uint16 Index) const
 {
-	return BehaviorStreamReader->GetGUIControlName(Index);
+	return BehaviorReader->GetGUIControlName(Index);
 }
 
 uint16 FSkelMeshDNAReader::GetRawControlCount() const
 {
-	return BehaviorStreamReader->GetRawControlCount();
+	return BehaviorReader->GetRawControlCount();
 }
 
 FString FSkelMeshDNAReader::GetRawControlName(uint16 Index) const
 {
-	return BehaviorStreamReader->GetRawControlName(Index);
+	return BehaviorReader->GetRawControlName(Index);
 }
 
 uint16 FSkelMeshDNAReader::GetJointCount() const
 {
-	return BehaviorStreamReader->GetJointCount();
+	return BehaviorReader->GetJointCount();
 }
 
 FString FSkelMeshDNAReader::GetJointName(uint16 Index) const
 {
-	return BehaviorStreamReader->GetJointName(Index);
+	return BehaviorReader->GetJointName(Index);
 }
 
 uint16 FSkelMeshDNAReader::GetJointIndexListCount() const
 {
-	return BehaviorStreamReader->GetJointIndexListCount();
+	return BehaviorReader->GetJointIndexListCount();
 }
 
 TArrayView<const uint16> FSkelMeshDNAReader::GetJointIndicesForLOD(uint16 LOD) const
 {
-	return BehaviorStreamReader->GetJointIndicesForLOD(LOD);
+	return BehaviorReader->GetJointIndicesForLOD(LOD);
 }
 
 uint16 FSkelMeshDNAReader::GetBlendShapeChannelCount() const
 {
-	return BehaviorStreamReader->GetBlendShapeChannelCount();
+	return BehaviorReader->GetBlendShapeChannelCount();
 }
 
 FString FSkelMeshDNAReader::GetBlendShapeChannelName(uint16 Index) const
 {
-	return BehaviorStreamReader->GetBlendShapeChannelName(Index);
+	return BehaviorReader->GetBlendShapeChannelName(Index);
 }
 
 uint16 FSkelMeshDNAReader::GetBlendShapeChannelIndexListCount() const
 {
-	return BehaviorStreamReader->GetBlendShapeChannelIndexListCount();
+	return BehaviorReader->GetBlendShapeChannelIndexListCount();
 }
 
 TArrayView<const uint16> FSkelMeshDNAReader::GetBlendShapeChannelIndicesForLOD(uint16 LOD) const
 {
-	return BehaviorStreamReader->GetBlendShapeChannelIndicesForLOD(LOD);
+	return BehaviorReader->GetBlendShapeChannelIndicesForLOD(LOD);
 }
 
 uint16 FSkelMeshDNAReader::GetAnimatedMapCount() const
 {
-	return BehaviorStreamReader->GetAnimatedMapCount();
+	return BehaviorReader->GetAnimatedMapCount();
 }
 
 FString FSkelMeshDNAReader::GetAnimatedMapName(uint16 Index) const
 {
-	return BehaviorStreamReader->GetAnimatedMapName(Index);
+	return BehaviorReader->GetAnimatedMapName(Index);
 }
 
 uint16 FSkelMeshDNAReader::GetAnimatedMapIndexListCount() const
 {
-	return BehaviorStreamReader->GetAnimatedMapIndexListCount();
+	return BehaviorReader->GetAnimatedMapIndexListCount();
 }
 
 TArrayView<const uint16> FSkelMeshDNAReader::GetAnimatedMapIndicesForLOD(uint16 LOD) const
 {
-	return BehaviorStreamReader->GetAnimatedMapIndicesForLOD(LOD);
+	return BehaviorReader->GetAnimatedMapIndicesForLOD(LOD);
 }
 
 uint16 FSkelMeshDNAReader::GetMeshCount() const
 {
-	return BehaviorStreamReader->GetMeshCount();
+	return BehaviorReader->GetMeshCount();
 }
 
 FString FSkelMeshDNAReader::GetMeshName(uint16 Index) const
 {
-	return BehaviorStreamReader->GetMeshName(Index);
+	return BehaviorReader->GetMeshName(Index);
 }
 
 uint16 FSkelMeshDNAReader::GetMeshIndexListCount() const
 {
-	return BehaviorStreamReader->GetMeshIndexListCount();
+	return BehaviorReader->GetMeshIndexListCount();
 }
 
 TArrayView<const uint16> FSkelMeshDNAReader::GetMeshIndicesForLOD(uint16 LOD) const
 {
-	return BehaviorStreamReader->GetMeshIndicesForLOD(LOD);
+	return BehaviorReader->GetMeshIndicesForLOD(LOD);
 }
 
 uint16 FSkelMeshDNAReader::GetMeshBlendShapeChannelMappingCount() const
 {
-	return BehaviorStreamReader->GetMeshBlendShapeChannelMappingCount();
+	return BehaviorReader->GetMeshBlendShapeChannelMappingCount();
 }
 
 FMeshBlendShapeChannelMapping FSkelMeshDNAReader::GetMeshBlendShapeChannelMapping(uint16 Index) const
 {
-	return BehaviorStreamReader->GetMeshBlendShapeChannelMapping(Index);
+	return BehaviorReader->GetMeshBlendShapeChannelMapping(Index);
 }
 
 TArrayView<const uint16> FSkelMeshDNAReader::GetMeshBlendShapeChannelMappingIndicesForLOD(uint16 LOD) const
 {
-	return BehaviorStreamReader->GetMeshBlendShapeChannelMappingIndicesForLOD(LOD);
+	return BehaviorReader->GetMeshBlendShapeChannelMappingIndicesForLOD(LOD);
 }
 
 FVector FSkelMeshDNAReader::GetNeutralJointTranslation(uint16 Index) const
 {
-	return BehaviorStreamReader->GetNeutralJointTranslation(Index);
+	return BehaviorReader->GetNeutralJointTranslation(Index);
 }
 
 FVector FSkelMeshDNAReader::GetNeutralJointRotation(uint16 Index) const
 {
-	return BehaviorStreamReader->GetNeutralJointRotation(Index);
+	return BehaviorReader->GetNeutralJointRotation(Index);
 }
 
 uint16 FSkelMeshDNAReader::GetJointParentIndex(uint16 Index) const
 {
-	return BehaviorStreamReader->GetJointParentIndex(Index);
+	return BehaviorReader->GetJointParentIndex(Index);
 }
 
 /** BEHAVIOR READER **/
 
 TArrayView<const uint16> FSkelMeshDNAReader::GetGUIToRawInputIndices() const
 {
-	return BehaviorStreamReader->GetGUIToRawInputIndices();
+	return BehaviorReader->GetGUIToRawInputIndices();
 }
 
 TArrayView<const uint16> FSkelMeshDNAReader::GetGUIToRawOutputIndices() const
 {
-	return BehaviorStreamReader->GetGUIToRawOutputIndices();
+	return BehaviorReader->GetGUIToRawOutputIndices();
 }
 
 TArrayView<const float> FSkelMeshDNAReader::GetGUIToRawFromValues() const
 {
-	return BehaviorStreamReader->GetGUIToRawFromValues();
+	return BehaviorReader->GetGUIToRawFromValues();
 }
 
 TArrayView<const float> FSkelMeshDNAReader::GetGUIToRawToValues() const
 {
-	return  BehaviorStreamReader->GetGUIToRawToValues();
+	return  BehaviorReader->GetGUIToRawToValues();
 }
 
 TArrayView<const float> FSkelMeshDNAReader::GetGUIToRawSlopeValues() const
 {
-	return BehaviorStreamReader->GetGUIToRawSlopeValues();
+	return BehaviorReader->GetGUIToRawSlopeValues();
 }
 
 TArrayView<const float> FSkelMeshDNAReader::GetGUIToRawCutValues() const
 {
-	return BehaviorStreamReader->GetGUIToRawCutValues();
+	return BehaviorReader->GetGUIToRawCutValues();
 }
 
 uint16 FSkelMeshDNAReader::GetPSDCount() const
 {
-	return BehaviorStreamReader->GetPSDCount();
+	return BehaviorReader->GetPSDCount();
 }
 
 TArrayView<const uint16> FSkelMeshDNAReader::GetPSDRowIndices() const
 {
-	return BehaviorStreamReader->GetPSDRowIndices();
+	return BehaviorReader->GetPSDRowIndices();
 }
 
 TArrayView<const uint16> FSkelMeshDNAReader::GetPSDColumnIndices() const
 {
-	return BehaviorStreamReader->GetPSDColumnIndices();
+	return BehaviorReader->GetPSDColumnIndices();
 }
 
 TArrayView<const float> FSkelMeshDNAReader::GetPSDValues() const
 {
-	return BehaviorStreamReader->GetPSDValues();
+	return BehaviorReader->GetPSDValues();
 }
 
 uint16 FSkelMeshDNAReader::GetJointRowCount() const
 {
-	return BehaviorStreamReader->GetJointRowCount();
+	return BehaviorReader->GetJointRowCount();
 }
 
 uint16 FSkelMeshDNAReader::GetJointColumnCount() const
 {
-	return BehaviorStreamReader->GetJointColumnCount();
+	return BehaviorReader->GetJointColumnCount();
 }
 
 TArrayView<const uint16> FSkelMeshDNAReader::GetJointGroupJointIndices(uint16 JointGroupIndex) const
 {
-	return BehaviorStreamReader->GetJointGroupJointIndices(JointGroupIndex);
+	return BehaviorReader->GetJointGroupJointIndices(JointGroupIndex);
 }
 
 TArrayView<const uint16> FSkelMeshDNAReader::GetJointVariableAttributeIndices(uint16 LOD) const
 {
-	return BehaviorStreamReader->GetJointVariableAttributeIndices(LOD);
+	return BehaviorReader->GetJointVariableAttributeIndices(LOD);
 }
 
 uint16 FSkelMeshDNAReader::GetJointGroupCount() const
 {
-	return BehaviorStreamReader->GetJointGroupCount();
+	return BehaviorReader->GetJointGroupCount();
 }
 
 TArrayView<const uint16> FSkelMeshDNAReader::GetJointGroupLODs(uint16 JointGroupIndex) const
 {
-	return BehaviorStreamReader->GetJointGroupLODs(JointGroupIndex);
+	return BehaviorReader->GetJointGroupLODs(JointGroupIndex);
 }
 
 TArrayView<const uint16> FSkelMeshDNAReader::GetJointGroupInputIndices(uint16 JointGroupIndex) const
 {
-	return BehaviorStreamReader->GetJointGroupInputIndices(JointGroupIndex);
+	return BehaviorReader->GetJointGroupInputIndices(JointGroupIndex);
 }
 
 TArrayView<const uint16> FSkelMeshDNAReader::GetJointGroupOutputIndices(uint16 JointGroupIndex) const
 {
-	return BehaviorStreamReader->GetJointGroupOutputIndices(JointGroupIndex);
+	return BehaviorReader->GetJointGroupOutputIndices(JointGroupIndex);
 }
 
 TArrayView<const float> FSkelMeshDNAReader::GetJointGroupValues(uint16 JointGroupIndex) const
 {
-	return BehaviorStreamReader->GetJointGroupValues(JointGroupIndex);
+	return BehaviorReader->GetJointGroupValues(JointGroupIndex);
 }
 
 TArrayView<const uint16> FSkelMeshDNAReader::GetBlendShapeChannelLODs() const
 {
-	return BehaviorStreamReader->GetBlendShapeChannelLODs();
+	return BehaviorReader->GetBlendShapeChannelLODs();
 }
 
 TArrayView<const uint16> FSkelMeshDNAReader::GetBlendShapeChannelInputIndices() const
 {
-	return BehaviorStreamReader->GetBlendShapeChannelInputIndices();
+	return BehaviorReader->GetBlendShapeChannelInputIndices();
 }
 
 TArrayView<const uint16> FSkelMeshDNAReader::GetBlendShapeChannelOutputIndices() const
 {
-	return BehaviorStreamReader->GetBlendShapeChannelOutputIndices();
+	return BehaviorReader->GetBlendShapeChannelOutputIndices();
 }
 
 TArrayView<const uint16> FSkelMeshDNAReader::GetAnimatedMapLODs() const
 {
-	return BehaviorStreamReader->GetAnimatedMapLODs();
+	return BehaviorReader->GetAnimatedMapLODs();
 }
 
 TArrayView<const uint16> FSkelMeshDNAReader::GetAnimatedMapInputIndices() const
 {
-	return BehaviorStreamReader->GetAnimatedMapInputIndices();
+	return BehaviorReader->GetAnimatedMapInputIndices();
 }
 
 TArrayView<const uint16> FSkelMeshDNAReader::GetAnimatedMapOutputIndices() const
 {
-	return BehaviorStreamReader->GetAnimatedMapOutputIndices();
+	return BehaviorReader->GetAnimatedMapOutputIndices();
 }
 
 TArrayView<const float> FSkelMeshDNAReader::GetAnimatedMapFromValues() const
 {
-	return BehaviorStreamReader->GetAnimatedMapFromValues();
+	return BehaviorReader->GetAnimatedMapFromValues();
 }
 
 TArrayView<const float> FSkelMeshDNAReader::GetAnimatedMapToValues() const
 {
-	return BehaviorStreamReader->GetAnimatedMapToValues();
+	return BehaviorReader->GetAnimatedMapToValues();
 }
 
 TArrayView<const float> FSkelMeshDNAReader::GetAnimatedMapSlopeValues() const
 {
-	return BehaviorStreamReader->GetAnimatedMapSlopeValues();
+	return BehaviorReader->GetAnimatedMapSlopeValues();
 }
 
 TArrayView<const float> FSkelMeshDNAReader::GetAnimatedMapCutValues() const
 {
-	return BehaviorStreamReader->GetAnimatedMapCutValues();
+	return BehaviorReader->GetAnimatedMapCutValues();
 }
 
 /** GEOMETRY READER **/
 uint32 FSkelMeshDNAReader::GetVertexPositionCount(uint16 MeshIndex) const
 {
-	if (GeometryStreamReader)
+	if (GeometryReader)
 	{
-		return GeometryStreamReader->GetVertexPositionCount(MeshIndex);
+		return GeometryReader->GetVertexPositionCount(MeshIndex);
 	}
 	return {};
 }
 
 FVector FSkelMeshDNAReader::GetVertexPosition(uint16 MeshIndex, uint32 PositionIndex) const
 {
-	if (GeometryStreamReader)
+	if (GeometryReader)
 	{
-		return GeometryStreamReader->GetVertexPosition(MeshIndex, PositionIndex);
+		return GeometryReader->GetVertexPosition(MeshIndex, PositionIndex);
 	}
 	return {};
 }
 
 TArrayView<const float> FSkelMeshDNAReader::GetVertexPositionXs(uint16 MeshIndex) const
 {
-	if (GeometryStreamReader)
+	if (GeometryReader)
 	{
-		return GeometryStreamReader->GetVertexPositionXs(MeshIndex);
+		return GeometryReader->GetVertexPositionXs(MeshIndex);
 	}
 	return {};
 }
 
 TArrayView<const float> FSkelMeshDNAReader::GetVertexPositionYs(uint16 MeshIndex) const
 {
-	if (GeometryStreamReader)
+	if (GeometryReader)
 	{
-		return GeometryStreamReader->GetVertexPositionYs(MeshIndex);
+		return GeometryReader->GetVertexPositionYs(MeshIndex);
 	}
 	return {};
 }
 
 TArrayView<const float> FSkelMeshDNAReader::GetVertexPositionZs(uint16 MeshIndex) const
 {
-	if (GeometryStreamReader)
+	if (GeometryReader)
 	{
-		return GeometryStreamReader->GetVertexPositionZs(MeshIndex);
+		return GeometryReader->GetVertexPositionZs(MeshIndex);
 	}
 	return {};
 }
 
 uint32 FSkelMeshDNAReader::GetVertexTextureCoordinateCount(uint16 MeshIndex) const
 {
-	if (GeometryStreamReader)
+	if (GeometryReader)
 	{
-		return GeometryStreamReader->GetVertexTextureCoordinateCount(MeshIndex);
+		return GeometryReader->GetVertexTextureCoordinateCount(MeshIndex);
 	}
 	return {};
 }
 
 FTextureCoordinate FSkelMeshDNAReader::GetVertexTextureCoordinate(uint16 MeshIndex, uint32 TextureCoordinateIndex) const
 {
-	if (GeometryStreamReader)
+	if (GeometryReader)
 	{
-		return GeometryStreamReader->GetVertexTextureCoordinate(MeshIndex, TextureCoordinateIndex);
+		return GeometryReader->GetVertexTextureCoordinate(MeshIndex, TextureCoordinateIndex);
 	}
 	return {};
 }
 
 uint32 FSkelMeshDNAReader::GetVertexNormalCount(uint16 MeshIndex) const
 {
-	if (GeometryStreamReader)
+	if (GeometryReader)
 	{
-		return GeometryStreamReader->GetVertexNormalCount(MeshIndex);
+		return GeometryReader->GetVertexNormalCount(MeshIndex);
 	}
 	return {};
 }
 
 FVector FSkelMeshDNAReader::GetVertexNormal(uint16 MeshIndex, uint32 NormalIndex) const
 {
-	if (GeometryStreamReader)
+	if (GeometryReader)
 	{
-		return GeometryStreamReader->GetVertexNormal(MeshIndex, NormalIndex);
+		return GeometryReader->GetVertexNormal(MeshIndex, NormalIndex);
 	}
 	return {};
 }
@@ -466,9 +477,9 @@ FVector FSkelMeshDNAReader::GetVertexNormal(uint16 MeshIndex, uint32 NormalIndex
 /* not needed for gene splicer */
 uint32 FSkelMeshDNAReader::GetVertexLayoutCount(uint16 MeshIndex) const
 {
-	if (GeometryStreamReader)
+	if (GeometryReader)
 	{
-		return GeometryStreamReader->GetVertexLayoutCount(MeshIndex);
+		return GeometryReader->GetVertexLayoutCount(MeshIndex);
 	}
 	return {};
 }
@@ -476,9 +487,9 @@ uint32 FSkelMeshDNAReader::GetVertexLayoutCount(uint16 MeshIndex) const
 /* not needed for gene splicer */
 FVertexLayout FSkelMeshDNAReader::GetVertexLayout(uint16 MeshIndex, uint32 VertexIndex) const
 {
-	if (GeometryStreamReader)
+	if (GeometryReader)
 	{
-		return GeometryStreamReader->GetVertexLayout(MeshIndex, VertexIndex);
+		return GeometryReader->GetVertexLayout(MeshIndex, VertexIndex);
 	}
 	return {};
 }
@@ -486,9 +497,9 @@ FVertexLayout FSkelMeshDNAReader::GetVertexLayout(uint16 MeshIndex, uint32 Verte
 /* not needed for gene splicer */
 uint32 FSkelMeshDNAReader::GetFaceCount(uint16 MeshIndex) const
 {
-	if (GeometryStreamReader)
+	if (GeometryReader)
 	{
-		return GeometryStreamReader->GetFaceCount(MeshIndex);
+		return GeometryReader->GetFaceCount(MeshIndex);
 	}
 	return {};
 }
@@ -496,121 +507,198 @@ uint32 FSkelMeshDNAReader::GetFaceCount(uint16 MeshIndex) const
 /* not needed for gene splicer */
 TArrayView<const uint32> FSkelMeshDNAReader::GetFaceVertexLayoutIndices(uint16 MeshIndex, uint32 FaceIndex) const
 {
-	if (GeometryStreamReader)
+	if (GeometryReader)
 	{
-		return GeometryStreamReader->GetFaceVertexLayoutIndices(MeshIndex, FaceIndex);
+		return GeometryReader->GetFaceVertexLayoutIndices(MeshIndex, FaceIndex);
 	}
 	return {};
 }
 
 uint16 FSkelMeshDNAReader::GetMaximumInfluencePerVertex(uint16 MeshIndex) const
 {
-	if (GeometryStreamReader)
+	if (GeometryReader)
 	{
-		return GeometryStreamReader->GetMaximumInfluencePerVertex(MeshIndex);
+		return GeometryReader->GetMaximumInfluencePerVertex(MeshIndex);
 	}
 	return {};
 }
 
 uint32 FSkelMeshDNAReader::GetSkinWeightsCount(uint16 MeshIndex) const
 {
-	if (GeometryStreamReader)
+	if (GeometryReader)
 	{
-		return GeometryStreamReader->GetSkinWeightsCount(MeshIndex);
+		return GeometryReader->GetSkinWeightsCount(MeshIndex);
 	}
 	return {};
 }
 
 TArrayView<const float> FSkelMeshDNAReader::GetSkinWeightsValues(uint16 MeshIndex, uint32 VertexIndex) const
 {
-	if (GeometryStreamReader)
+	if (GeometryReader)
 	{
-		return GeometryStreamReader->GetSkinWeightsValues(MeshIndex, VertexIndex);
+		return GeometryReader->GetSkinWeightsValues(MeshIndex, VertexIndex);
 	}
 	return {};
 }
 
 TArrayView<const uint16> FSkelMeshDNAReader::GetSkinWeightsJointIndices(uint16 MeshIndex, uint32 VertexIndex) const
 {
-	if (GeometryStreamReader)
+	if (GeometryReader)
 	{
-		return GeometryStreamReader->GetSkinWeightsJointIndices(MeshIndex, VertexIndex);
+		return GeometryReader->GetSkinWeightsJointIndices(MeshIndex, VertexIndex);
 	}
 	return {};
 }
 
 uint16 FSkelMeshDNAReader::GetBlendShapeTargetCount(uint16 MeshIndex) const
 {
-	if (GeometryStreamReader)
+	if (GeometryReader)
 	{
-		return GeometryStreamReader->GetBlendShapeTargetCount(MeshIndex);
+		return GeometryReader->GetBlendShapeTargetCount(MeshIndex);
 	}
 	return {};
 }
 
 uint16 FSkelMeshDNAReader::GetBlendShapeChannelIndex(uint16 MeshIndex, uint16 BlendShapeTargetIndex) const
 {
-	if (GeometryStreamReader)
+	if (GeometryReader)
 	{
-		return GeometryStreamReader->GetBlendShapeChannelIndex(MeshIndex, BlendShapeTargetIndex);
+		return GeometryReader->GetBlendShapeChannelIndex(MeshIndex, BlendShapeTargetIndex);
 	}
 	return {};
 }
 
 uint32 FSkelMeshDNAReader::GetBlendShapeTargetDeltaCount(uint16 MeshIndex, uint16 BlendShapeIndex) const
 {
-	if (GeometryStreamReader)
+	if (GeometryReader)
 	{
-		return GeometryStreamReader->GetBlendShapeTargetDeltaCount(MeshIndex, BlendShapeIndex);
+		return GeometryReader->GetBlendShapeTargetDeltaCount(MeshIndex, BlendShapeIndex);
 	}
 	return {};
 }
 
 FVector FSkelMeshDNAReader::GetBlendShapeTargetDelta(uint16 MeshIndex, uint16 BlendShapeIndex, uint32 DeltaIndex) const
 {
-	if (GeometryStreamReader)
+	if (GeometryReader)
 	{
-		return GeometryStreamReader->GetBlendShapeTargetDelta(MeshIndex, BlendShapeIndex, DeltaIndex);
+		return GeometryReader->GetBlendShapeTargetDelta(MeshIndex, BlendShapeIndex, DeltaIndex);
 	}
 	return {};
 }
 
 TArrayView<const float> FSkelMeshDNAReader::GetBlendShapeTargetDeltaXs(uint16 MeshIndex, uint16 BlendShapeTargetIndex) const
 {
-	if (GeometryStreamReader)
+	if (GeometryReader)
 	{
-		return GeometryStreamReader->GetBlendShapeTargetDeltaXs(MeshIndex, BlendShapeTargetIndex);
+		return GeometryReader->GetBlendShapeTargetDeltaXs(MeshIndex, BlendShapeTargetIndex);
 	}
 	return {};
 }
 
 TArrayView<const float> FSkelMeshDNAReader::GetBlendShapeTargetDeltaYs(uint16 MeshIndex, uint16 BlendShapeTargetIndex) const
 {
-	if (GeometryStreamReader)
+	if (GeometryReader)
 	{
-		return GeometryStreamReader->GetBlendShapeTargetDeltaYs(MeshIndex, BlendShapeTargetIndex);
+		return GeometryReader->GetBlendShapeTargetDeltaYs(MeshIndex, BlendShapeTargetIndex);
 	}
 	return {};
 }
 
 TArrayView<const float> FSkelMeshDNAReader::GetBlendShapeTargetDeltaZs(uint16 MeshIndex, uint16 BlendShapeTargetIndex) const
 {
-	if (GeometryStreamReader)
+	if (GeometryReader)
 	{
-		return GeometryStreamReader->GetBlendShapeTargetDeltaZs(MeshIndex, BlendShapeTargetIndex);
+		return GeometryReader->GetBlendShapeTargetDeltaZs(MeshIndex, BlendShapeTargetIndex);
 	}
 	return {};
 }
 
 TArrayView<const uint32> FSkelMeshDNAReader::GetBlendShapeTargetVertexIndices(uint16 MeshIndex, uint16 BlendShapeIndex) const
 {
-	if (GeometryStreamReader)
+	if (GeometryReader)
 	{
-		return GeometryStreamReader->GetBlendShapeTargetVertexIndices(MeshIndex, BlendShapeIndex);
+		return GeometryReader->GetBlendShapeTargetVertexIndices(MeshIndex, BlendShapeIndex);
 	}
 	return {};
 }
 
-void FSkelMeshDNAReader::Unload(EDNADataLayer Layer) {
+/** MACHINE LEARNED BEHAVIOR READER **/
+uint16 FSkelMeshDNAReader::GetMLControlCount() const
+{
+	return BehaviorReader->GetMLControlCount();
+}
+
+FString FSkelMeshDNAReader::GetMLControlName(uint16 Index) const
+{
+	return BehaviorReader->GetMLControlName(Index);
+}
+
+uint16 FSkelMeshDNAReader::GetNeuralNetworkCount() const
+{
+	return BehaviorReader->GetNeuralNetworkCount();
+}
+
+uint16 FSkelMeshDNAReader::GetNeuralNetworkIndexListCount() const
+{
+	return BehaviorReader->GetNeuralNetworkIndexListCount();
+}
+
+TArrayView<const uint16> FSkelMeshDNAReader::GetNeuralNetworkIndicesForLOD(uint16 LOD) const
+{
+	return BehaviorReader->GetNeuralNetworkIndicesForLOD(LOD);
+}
+
+uint16 FSkelMeshDNAReader::GetMeshRegionCount(uint16 MeshIndex) const
+{
+	return BehaviorReader->GetMeshRegionCount(MeshIndex);
+}
+
+FString FSkelMeshDNAReader::GetMeshRegionName(uint16 MeshIndex, uint16 RegionIndex) const
+{
+	return BehaviorReader->GetMeshRegionName(MeshIndex, RegionIndex);
+}
+
+TArrayView<const uint16> FSkelMeshDNAReader::GetNeuralNetworkIndicesForMeshRegion(uint16 MeshIndex, uint16 RegionIndex) const
+{
+	return BehaviorReader->GetNeuralNetworkIndicesForMeshRegion(MeshIndex, RegionIndex);
+}
+
+TArrayView<const uint16> FSkelMeshDNAReader::GetNeuralNetworkInputIndices(uint16 NetIndex) const
+{
+	return BehaviorReader->GetNeuralNetworkInputIndices(NetIndex);
+}
+
+TArrayView<const uint16> FSkelMeshDNAReader::GetNeuralNetworkOutputIndices(uint16 NetIndex) const
+{
+	return BehaviorReader->GetNeuralNetworkOutputIndices(NetIndex);
+}
+
+uint16 FSkelMeshDNAReader::GetNeuralNetworkLayerCount(uint16 NetIndex) const
+{
+	return BehaviorReader->GetNeuralNetworkLayerCount(NetIndex);
+}
+
+EActivationFunction FSkelMeshDNAReader::GetNeuralNetworkLayerActivationFunction(uint16 NetIndex, uint16 LayerIndex) const
+{
+	return BehaviorReader->GetNeuralNetworkLayerActivationFunction(NetIndex, LayerIndex);
+}
+
+TArrayView<const float> FSkelMeshDNAReader::GetNeuralNetworkLayerActivationFunctionParameters(uint16 NetIndex, uint16 LayerIndex) const
+{
+	return BehaviorReader->GetNeuralNetworkLayerActivationFunctionParameters(NetIndex, LayerIndex);
+}
+
+TArrayView<const float> FSkelMeshDNAReader::GetNeuralNetworkLayerBiases(uint16 NetIndex, uint16 LayerIndex) const
+{
+	return BehaviorReader->GetNeuralNetworkLayerBiases(NetIndex, LayerIndex);
+}
+
+TArrayView<const float> FSkelMeshDNAReader::GetNeuralNetworkLayerWeights(uint16 NetIndex, uint16 LayerIndex) const
+{
+	return BehaviorReader->GetNeuralNetworkLayerWeights(NetIndex, LayerIndex);
+}
+
+void FSkelMeshDNAReader::Unload(EDNADataLayer Layer)
+{
 	ensureMsgf(false, TEXT("Assest are not unloadable"));
 }
\ No newline at end of file
diff --git a/Engine/Plugins/Animation/RigLogic/Source/RigLogicModule/Public/DNACommon.h b/Engine/Plugins/Animation/RigLogic/Source/RigLogicModule/Public/DNACommon.h
--- a/Engine/Plugins/Animation/RigLogic/Source/RigLogicModule/Public/DNACommon.h
+++ b/Engine/Plugins/Animation/RigLogic/Source/RigLogicModule/Public/DNACommon.h
@@ -55,13 +55,29 @@ enum class EDirection: uint8
 UENUM(BlueprintType)
 enum class EDNADataLayer : uint8
 {
-	Descriptor,
-	Definition,  // Includes Descriptor
-	Behavior,  // Includes Descriptor and Definition
-	Geometry,  // Includes Descriptor and Definition
-	GeometryWithoutBlendShapes,  // Includes Descriptor and Definition
-	AllWithoutBlendShapes,  // Includes everything except blend shapes from Geometry
-	All
+	None,
+	Descriptor = 1,
+	Definition = 2 | Descriptor,  // Implicitly loads Descriptor
+	Behavior = 4 | Definition,  // Implicitly loads Descriptor and Definition
+	Geometry = 8 | Definition,  // Implicitly loads Descriptor and Definition
+	GeometryWithoutBlendShapes = 16 | Definition,  // Implicitly loads Descriptor and Definition
+	MachineLearnedBehavior = 32 | Definition,  // Implicitly loads Definition
+	All = Behavior | Geometry | MachineLearnedBehavior
+};
+
+inline EDNADataLayer operator|(EDNADataLayer LHS, EDNADataLayer RHS)
+{
+	return static_cast<EDNADataLayer>(static_cast<uint8>(LHS) | static_cast<uint8>(RHS));
+}
+
+UENUM(BlueprintType)
+enum class EActivationFunction : uint8
+{
+	Linear,
+	ReLU,
+	LeakyReLU,
+	Tanh,
+	Sigmoid
 };
 
 USTRUCT(BlueprintType)
diff --git a/Engine/Plugins/Animation/RigLogic/Source/RigLogicModule/Public/DNAReader.h b/Engine/Plugins/Animation/RigLogic/Source/RigLogicModule/Public/DNAReader.h
--- a/Engine/Plugins/Animation/RigLogic/Source/RigLogicModule/Public/DNAReader.h
+++ b/Engine/Plugins/Animation/RigLogic/Source/RigLogicModule/Public/DNAReader.h
@@ -19,7 +19,29 @@ class Reader;
 
 }  // namespace dna
 
-class IDescriptorReader
+class IDNAReaderBase
+{
+protected:
+	virtual ~IDNAReaderBase() = default;
+
+public:
+	virtual dna::Reader* Unwrap() const = 0;
+	virtual void Unload(EDNADataLayer Layer) = 0;
+
+};
+
+class IHeaderReader : public IDNAReaderBase
+{
+protected:
+	virtual ~IHeaderReader() = default;
+
+public:
+	virtual uint16 GetFileFormatGeneration() const = 0;
+	virtual uint16 GetFileFormatVersion() const = 0;
+
+};
+
+class IDescriptorReader : public IHeaderReader
 {
 protected:
 	virtual ~IDescriptorReader() = default;
@@ -40,9 +62,6 @@ public:
 	virtual FString GetDBComplexity() const = 0;
 	virtual FString GetDBName() const = 0;
 
-	virtual dna::Reader* Unwrap() const = 0;
-	virtual void Unload(EDNADataLayer Layer) = 0;
-
 };
 
 class IDefinitionReader : public IDescriptorReader
@@ -149,10 +168,34 @@ public:
 	virtual TArrayView<const uint32> GetBlendShapeTargetVertexIndices(uint16 MeshIndex, uint16 BlendShapeTargetIndex) const = 0;
 };
 
+class IMachineLearnedBehaviorReader : public virtual IDefinitionReader
+{
+protected:
+	virtual ~IMachineLearnedBehaviorReader() = default;
+
+public:
+	virtual uint16 GetMLControlCount() const = 0;
+	virtual FString GetMLControlName(uint16 Index) const = 0;
+	virtual uint16 GetNeuralNetworkCount() const = 0;
+	virtual uint16 GetNeuralNetworkIndexListCount() const = 0;
+	virtual TArrayView<const uint16> GetNeuralNetworkIndicesForLOD(uint16 LOD) const = 0;
+	virtual uint16 GetMeshRegionCount(uint16 MeshIndex) const = 0;
+	virtual FString GetMeshRegionName(uint16 MeshIndex, uint16 RegionIndex) const = 0;
+	virtual TArrayView<const uint16> GetNeuralNetworkIndicesForMeshRegion(uint16 MeshIndex, uint16 RegionIndex) const = 0;
+	virtual TArrayView<const uint16> GetNeuralNetworkInputIndices(uint16 NetIndex) const = 0;
+	virtual TArrayView<const uint16> GetNeuralNetworkOutputIndices(uint16 NetIndex) const = 0;
+	virtual uint16 GetNeuralNetworkLayerCount(uint16 NetIndex) const = 0;
+	virtual EActivationFunction GetNeuralNetworkLayerActivationFunction(uint16 NetIndex, uint16 LayerIndex) const = 0;
+	virtual TArrayView<const float> GetNeuralNetworkLayerActivationFunctionParameters(uint16 NetIndex, uint16 LayerIndex) const = 0;
+	virtual TArrayView<const float> GetNeuralNetworkLayerBiases(uint16 NetIndex, uint16 LayerIndex) const = 0;
+	virtual TArrayView<const float> GetNeuralNetworkLayerWeights(uint16 NetIndex, uint16 LayerIndex) const = 0;
+
+};
+
 /**
 	@brief UE interface for DNA Reader wrappers.
 */
-class IDNAReader : public IBehaviorReader, public IGeometryReader
+class IDNAReader : public IBehaviorReader, public IGeometryReader, public IMachineLearnedBehaviorReader
 {
 public:
 	virtual ~IDNAReader() = default;
diff --git a/Engine/Plugins/Animation/RigLogic/Source/RigLogicModule/Public/DNAReaderAdapter.h b/Engine/Plugins/Animation/RigLogic/Source/RigLogicModule/Public/DNAReaderAdapter.h
--- a/Engine/Plugins/Animation/RigLogic/Source/RigLogicModule/Public/DNAReaderAdapter.h
+++ b/Engine/Plugins/Animation/RigLogic/Source/RigLogicModule/Public/DNAReaderAdapter.h
@@ -18,6 +18,9 @@ public:
 	FDNAReader(FDNAReader&&) = default;
 	FDNAReader& operator=(FDNAReader&&) = default;
 
+	// Header
+	uint16 GetFileFormatGeneration() const override;
+	uint16 GetFileFormatVersion() const override;
 	// Descriptor
 	FString GetName() const override;
 	EArchetype GetArchetype() const override;
@@ -116,6 +119,22 @@ public:
 	TArrayView<const float> GetBlendShapeTargetDeltaYs(uint16 MeshIndex, uint16 BlendShapeTargetIndex) const override;
 	TArrayView<const float> GetBlendShapeTargetDeltaZs(uint16 MeshIndex, uint16 BlendShapeTargetIndex) const override;
 	TArrayView<const uint32> GetBlendShapeTargetVertexIndices(uint16 MeshIndex, uint16 BlendShapeTargetIndex) const override;
+	// Machine Learned Behavior
+	uint16 GetMLControlCount() const override;
+	FString GetMLControlName(uint16 Index) const override;
+	uint16 GetNeuralNetworkCount() const override;
+	uint16 GetNeuralNetworkIndexListCount() const override;
+	TArrayView<const uint16> GetNeuralNetworkIndicesForLOD(uint16 LOD) const override;
+	uint16 GetMeshRegionCount(uint16 MeshIndex) const override;
+	FString GetMeshRegionName(uint16 MeshIndex, uint16 RegionIndex) const override;
+	TArrayView<const uint16> GetNeuralNetworkIndicesForMeshRegion(uint16 MeshIndex, uint16 RegionIndex) const override;
+	TArrayView<const uint16> GetNeuralNetworkInputIndices(uint16 NetIndex) const override;
+	TArrayView<const uint16> GetNeuralNetworkOutputIndices(uint16 NetIndex) const override;
+	uint16 GetNeuralNetworkLayerCount(uint16 NetIndex) const override;
+	EActivationFunction GetNeuralNetworkLayerActivationFunction(uint16 NetIndex, uint16 LayerIndex) const override;
+	TArrayView<const float> GetNeuralNetworkLayerActivationFunctionParameters(uint16 NetIndex, uint16 LayerIndex) const override;
+	TArrayView<const float> GetNeuralNetworkLayerBiases(uint16 NetIndex, uint16 LayerIndex) const override;
+	TArrayView<const float> GetNeuralNetworkLayerWeights(uint16 NetIndex, uint16 LayerIndex) const override;
 
 	void Unload(EDNADataLayer Layer) override;
 
@@ -160,6 +179,18 @@ void FDNAReader<TWrappedReader>::FWrappedReaderDeleter::operator()(TWrappedReade
 	}
 }
 
+template <class TWrappedReader>
+uint16 FDNAReader<TWrappedReader>::GetFileFormatGeneration() const
+{
+	return ReaderPtr->getFileFormatGeneration();
+}
+
+template <class TWrappedReader>
+uint16 FDNAReader<TWrappedReader>::GetFileFormatVersion() const
+{
+	return ReaderPtr->getFileFormatVersion();
+}
+
 template <class TWrappedReader>
 FString FDNAReader<TWrappedReader>::GetName() const
 {
@@ -788,3 +819,100 @@ TArrayView<const uint32> FDNAReader<TWrappedReader>::GetBlendShapeTargetVertexIn
 	const auto Indices = ReaderPtr->getBlendShapeTargetVertexIndices(MeshIndex, BlendShapeTargetIndex);
 	return TArrayView<const uint32>(Indices.data(), static_cast<int32>(Indices.size()));
 }
+
+template <class TWrappedReader>
+uint16 FDNAReader<TWrappedReader>::GetMLControlCount() const
+{
+	return ReaderPtr->getMLControlCount();
+}
+
+template <class TWrappedReader>
+FString FDNAReader<TWrappedReader>::GetMLControlName(uint16 Index) const
+{
+	return FString(ANSI_TO_TCHAR(ReaderPtr->getMLControlName(Index).data()));
+}
+
+template <class TWrappedReader>
+uint16 FDNAReader<TWrappedReader>::GetNeuralNetworkCount() const
+{
+	return ReaderPtr->getNeuralNetworkCount();
+}
+
+template <class TWrappedReader>
+uint16 FDNAReader<TWrappedReader>::GetNeuralNetworkIndexListCount() const
+{
+	return ReaderPtr->getNeuralNetworkIndexListCount();
+}
+
+template <class TWrappedReader>
+TArrayView<const uint16> FDNAReader<TWrappedReader>::GetNeuralNetworkIndicesForLOD(uint16 LOD) const
+{
+	const auto Values = ReaderPtr->getNeuralNetworkIndicesForLOD(LOD);
+	return TArrayView<const uint16>(Values.data(), static_cast<int32>(Values.size()));
+}
+
+template <class TWrappedReader>
+uint16 FDNAReader<TWrappedReader>::GetMeshRegionCount(uint16 MeshIndex) const
+{
+	return ReaderPtr->getMeshRegionCount(MeshIndex);
+}
+
+template <class TWrappedReader>
+FString FDNAReader<TWrappedReader>::GetMeshRegionName(uint16 MeshIndex, uint16 RegionIndex) const
+{
+	return FString(ANSI_TO_TCHAR(ReaderPtr->getMeshRegionName(MeshIndex, RegionIndex).data()));
+}
+
+template <class TWrappedReader>
+TArrayView<const uint16> FDNAReader<TWrappedReader>::GetNeuralNetworkIndicesForMeshRegion(uint16 MeshIndex, uint16 RegionIndex) const
+{
+	const auto Values = ReaderPtr->getNeuralNetworkIndicesForMeshRegion(MeshIndex, RegionIndex);
+	return TArrayView<const uint16>(Values.data(), static_cast<int32>(Values.size()));
+}
+
+template <class TWrappedReader>
+TArrayView<const uint16> FDNAReader<TWrappedReader>::GetNeuralNetworkInputIndices(uint16 NetIndex) const
+{
+	const auto Values = ReaderPtr->getNeuralNetworkInputIndices(NetIndex);
+	return TArrayView<const uint16>(Values.data(), static_cast<int32>(Values.size()));
+}
+
+template <class TWrappedReader>
+TArrayView<const uint16> FDNAReader<TWrappedReader>::GetNeuralNetworkOutputIndices(uint16 NetIndex) const
+{
+	const auto Values = ReaderPtr->getNeuralNetworkOutputIndices(NetIndex);
+	return TArrayView<const uint16>(Values.data(), static_cast<int32>(Values.size()));
+}
+
+template <class TWrappedReader>
+uint16 FDNAReader<TWrappedReader>::GetNeuralNetworkLayerCount(uint16 NetIndex) const
+{
+	return ReaderPtr->getNeuralNetworkLayerCount(NetIndex);
+}
+
+template <class TWrappedReader>
+EActivationFunction FDNAReader<TWrappedReader>::GetNeuralNetworkLayerActivationFunction(uint16 NetIndex, uint16 LayerIndex) const
+{
+	return static_cast<EActivationFunction>(ReaderPtr->getNeuralNetworkLayerActivationFunction(NetIndex, LayerIndex));
+}
+
+template <class TWrappedReader>
+TArrayView<const float> FDNAReader<TWrappedReader>::GetNeuralNetworkLayerActivationFunctionParameters(uint16 NetIndex, uint16 LayerIndex) const
+{
+	const auto Values = ReaderPtr->getNeuralNetworkLayerActivationFunctionParameters(NetIndex, LayerIndex);
+	return TArrayView<const float>(Values.data(), static_cast<int32>(Values.size()));
+}
+
+template <class TWrappedReader>
+TArrayView<const float> FDNAReader<TWrappedReader>::GetNeuralNetworkLayerBiases(uint16 NetIndex, uint16 LayerIndex) const
+{
+	const auto Values = ReaderPtr->getNeuralNetworkLayerBiases(NetIndex, LayerIndex);
+	return TArrayView<const float>(Values.data(), static_cast<int32>(Values.size()));
+}
+
+template <class TWrappedReader>
+TArrayView<const float> FDNAReader<TWrappedReader>::GetNeuralNetworkLayerWeights(uint16 NetIndex, uint16 LayerIndex) const
+{
+	const auto Values = ReaderPtr->getNeuralNetworkLayerWeights(NetIndex, LayerIndex);
+	return TArrayView<const float>(Values.data(), static_cast<int32>(Values.size()));
+}
diff --git a/Engine/Plugins/Animation/RigLogic/Source/RigLogicModule/Public/RigInstance.h b/Engine/Plugins/Animation/RigLogic/Source/RigLogicModule/Public/RigInstance.h
--- a/Engine/Plugins/Animation/RigLogic/Source/RigLogicModule/Public/RigInstance.h
+++ b/Engine/Plugins/Animation/RigLogic/Source/RigLogicModule/Public/RigInstance.h
@@ -33,12 +33,29 @@ public:
 	FRigInstance& operator=(FRigInstance&&) = default;
 
 	uint16 GetGUIControlCount() const;
+	float GetGUIControl(uint16 Index) const;
 	void SetGUIControl(uint16 Index, float Value);
+	TArrayView<const float> GetGUIControlValues() const;
 	void SetGUIControlValues(const float* Values);
 
 	uint16 GetRawControlCount() const;
+	float GetRawControl(uint16 Index) const;
 	void SetRawControl(uint16 Index, float Value);
+	TArrayView<const float> GetRawControlValues() const;
 	void SetRawControlValues(const float* Values);
+
+	uint16 GetPSDControlCount() const;
+	float GetPSDControl(uint16 Index) const;
+	TArrayView<const float> GetPSDControlValues() const;
+
+	uint16 GetMLControlCount() const;
+	float GetMLControl(uint16 Index) const;
+	TArrayView<const float> GetMLControlValues() const;
+
+	uint16 GetNeuralNetworkCount() const;
+	float GetNeuralNetworkMask(uint16 NeuralNetIndex) const;
+	void SetNeuralNetworkMask(uint16 NeuralNetIndex, float Value);
+
 	TArrayView<const float> GetRawJointOutputs() const;
 	FTransformArrayView GetJointOutputs() const;
 	TArrayView<const float> GetBlendShapeOutputs() const;
diff --git a/Engine/Plugins/Animation/RigLogic/Source/RigLogicModule/Public/RigLogic.h b/Engine/Plugins/Animation/RigLogic/Source/RigLogicModule/Public/RigLogic.h
--- a/Engine/Plugins/Animation/RigLogic/Source/RigLogicModule/Public/RigLogic.h
+++ b/Engine/Plugins/Animation/RigLogic/Source/RigLogicModule/Public/RigLogic.h
@@ -12,7 +12,7 @@
 #include "TransformArrayView.h"
 
 class FRigInstance;
-class IBehaviorReader;
+class IDNAReader;
 
 namespace rl4
 {
@@ -32,7 +32,7 @@ enum class ERigLogicCalculationType: uint8
 class RIGLOGICMODULE_API FRigLogic
 {
 public:
-	FRigLogic(const IBehaviorReader* Reader, ERigLogicCalculationType CalculationType = ERigLogicCalculationType::SSE);
+	FRigLogic(const IDNAReader* Reader, ERigLogicCalculationType CalculationType = ERigLogicCalculationType::SSE);
 	~FRigLogic();
 
 	FRigLogic(const FRigLogic&) = delete;
@@ -44,10 +44,18 @@ public:
 	uint16 GetLODCount() const;
 	TArrayView<const float> GetRawNeutralJointValues() const;
 	FTransformArrayView GetNeutralJointValues() const;
-	uint16 GetJointGroupCount() const;
 	TArrayView<const uint16> GetJointVariableAttributeIndices(uint16 LOD) const;
+	uint16 GetJointGroupCount() const;
+	uint16 GetNeuralNetworkCount() const;
+	uint16 GetMeshCount() const;
+	uint16 GetMeshRegionCount(uint16 MeshIndex) const;
+	TArrayView<const uint16> GetNeuralNetworkIndices(uint16 MeshIndex, uint16 RegionIndex) const;
+
 	void MapGUIToRawControls(FRigInstance* Instance) const;
+	void MapRawToGUIControls(FRigInstance* Instance) const;
 	void CalculateControls(FRigInstance* Instance) const;
+	void CalculateMachineLearnedBehaviorControls(FRigInstance* Instance) const;
+	void CalculateMachineLearnedBehaviorControls(FRigInstance* Instance, uint16 NeuralNetIndex) const;
 	void CalculateJoints(FRigInstance* Instance) const;
 	void CalculateJoints(FRigInstance* Instance, uint16 JointGroupIndex) const;
 	void CalculateBlendShapes(FRigInstance* Instance) const;
diff --git a/Engine/Plugins/Animation/RigLogic/Source/RigLogicModule/Public/RigUnit_RigLogic_Test.h b/Engine/Plugins/Animation/RigLogic/Source/RigLogicModule/Public/RigUnit_RigLogic_Test.h
--- a/Engine/Plugins/Animation/RigLogic/Source/RigLogicModule/Public/RigUnit_RigLogic_Test.h
+++ b/Engine/Plugins/Animation/RigLogic/Source/RigLogicModule/Public/RigUnit_RigLogic_Test.h
@@ -17,15 +17,25 @@ struct FRigUnit_RigLogic_JointUpdateParams;
 
 //** A unit test-only helper class to prevent FString from being deleted as it is converted from FString*//
 
-class TestBehaviorReader: public IBehaviorReader
+class TestBehaviorReader: public IDNAReader
 {
 public:
 	void Unload(EDNADataLayer Layer) override
 	{
+	}
 
+	// IHeaderReader methods
+	uint16 GetFileFormatGeneration() const
+	{
+		return {};
 	}
 
-	// DescriptorReader methods start
+	uint16 GetFileFormatVersion() const
+	{
+		return {};
+	}
+
+	// IDescriptorReader methods
 	FString GetName() const override
 	{
 		return {};
@@ -96,7 +106,7 @@ public:
 		return {};
 	}
 
-	// DefinitionReader methods start
+	// IDefinitionReader methods
 	uint16 GetGUIControlCount() const override
 	{
 		return {};
@@ -262,7 +272,7 @@ public:
 		return {};
 	}
 
-	// BehaviorReader methods start
+	// IBehaviorReader methods
 	TArrayView<const uint16> GetGUIToRawInputIndices() const override
 	{
 		return {};
@@ -408,6 +418,208 @@ public:
 		return {};
 	}
 
+	// IGeometryReader methods
+	uint32 GetVertexPositionCount(uint16 MeshIndex) const
+	{
+		return {};
+	}
+
+	FVector GetVertexPosition(uint16 MeshIndex, uint32 PositionIndex) const
+	{
+		return {};
+	}
+
+	TArrayView<const float> GetVertexPositionXs(uint16 MeshIndex) const
+	{
+		return {};
+	}
+
+	TArrayView<const float> GetVertexPositionYs(uint16 MeshIndex) const
+	{
+		return {};
+	}
+
+	TArrayView<const float> GetVertexPositionZs(uint16 MeshIndex) const
+	{
+		return {};
+	}
+
+	uint32 GetVertexTextureCoordinateCount(uint16 MeshIndex) const
+	{
+		return {};
+	}
+
+	FTextureCoordinate GetVertexTextureCoordinate(uint16 MeshIndex, uint32 TextureCoordinateIndex) const
+	{
+		return {};
+	}
+
+	uint32 GetVertexNormalCount(uint16 MeshIndex) const
+	{
+		return {};
+	}
+
+	FVector GetVertexNormal(uint16 MeshIndex, uint32 NormalIndex) const
+	{
+		return {};
+	}
+
+	uint32 GetVertexLayoutCount(uint16 MeshIndex) const
+	{
+		return {};
+	}
+
+	FVertexLayout GetVertexLayout(uint16 MeshIndex, uint32 VertexIndex) const
+	{
+		return {};
+	}
+
+	uint32 GetFaceCount(uint16 MeshIndex) const
+	{
+		return {};
+	}
+
+	TArrayView<const uint32> GetFaceVertexLayoutIndices(uint16 MeshIndex, uint32 FaceIndex) const
+	{
+		return {};
+	}
+
+	uint16 GetMaximumInfluencePerVertex(uint16 MeshIndex) const
+	{
+		return {};
+	}
+
+	uint32 GetSkinWeightsCount(uint16 MeshIndex) const
+	{
+		return {};
+	}
+
+	TArrayView<const float> GetSkinWeightsValues(uint16 MeshIndex, uint32 VertexIndex) const
+	{
+		return {};
+	}
+
+	TArrayView<const uint16> GetSkinWeightsJointIndices(uint16 MeshIndex, uint32 VertexIndex) const
+	{
+		return {};
+	}
+
+	uint16 GetBlendShapeTargetCount(uint16 MeshIndex) const
+	{
+		return {};
+	}
+
+	uint16 GetBlendShapeChannelIndex(uint16 MeshIndex, uint16 BlendShapeTargetIndex) const
+	{
+		return {};
+	}
+
+	uint32 GetBlendShapeTargetDeltaCount(uint16 MeshIndex, uint16 BlendShapeIndex) const
+	{
+		return {};
+	}
+
+	FVector GetBlendShapeTargetDelta(uint16 MeshIndex, uint16 BlendShapeIndex, uint32 DeltaIndex) const
+	{
+		return {};
+	}
+
+	TArrayView<const float> GetBlendShapeTargetDeltaXs(uint16 MeshIndex, uint16 BlendShapeTargetIndex) const
+	{
+		return {};
+	}
+
+	TArrayView<const float> GetBlendShapeTargetDeltaYs(uint16 MeshIndex, uint16 BlendShapeTargetIndex) const
+	{
+		return {};
+	}
+
+	TArrayView<const float> GetBlendShapeTargetDeltaZs(uint16 MeshIndex, uint16 BlendShapeTargetIndex) const
+	{
+		return {};
+	}
+
+	TArrayView<const uint32> GetBlendShapeTargetVertexIndices(uint16 MeshIndex, uint16 BlendShapeIndex) const
+	{
+		return {};
+	}
+
+	// IMachineLearnedBehaviorReader methods
+	uint16 GetMLControlCount() const
+	{
+		return {};
+	}
+
+	FString GetMLControlName(uint16 Index) const
+	{
+		return {};
+	}
+
+	uint16 GetNeuralNetworkCount() const
+	{
+		return {};
+	}
+
+	uint16 GetNeuralNetworkIndexListCount() const
+	{
+		return {};
+	}
+
+	TArrayView<const uint16> GetNeuralNetworkIndicesForLOD(uint16 LOD) const
+	{
+		return {};
+	}
+
+	uint16 GetMeshRegionCount(uint16 MeshIndex) const
+	{
+		return {};
+	}
+
+	FString GetMeshRegionName(uint16 MeshIndex, uint16 RegionIndex) const
+	{
+		return {};
+	}
+
+	TArrayView<const uint16> GetNeuralNetworkIndicesForMeshRegion(uint16 MeshIndex, uint16 RegionIndex) const
+	{
+		return {};
+	}
+
+	TArrayView<const uint16> GetNeuralNetworkInputIndices(uint16 NetIndex) const
+	{
+		return {};
+	}
+
+	TArrayView<const uint16> GetNeuralNetworkOutputIndices(uint16 NetIndex) const
+	{
+		return {};
+	}
+
+	uint16 GetNeuralNetworkLayerCount(uint16 NetIndex) const
+	{
+		return {};
+	}
+
+	EActivationFunction GetNeuralNetworkLayerActivationFunction(uint16 NetIndex, uint16 LayerIndex) const
+	{
+		return {};
+	}
+
+	TArrayView<const float> GetNeuralNetworkLayerActivationFunctionParameters(uint16 NetIndex, uint16 LayerIndex) const
+	{
+		return {};
+	}
+
+	TArrayView<const float> GetNeuralNetworkLayerBiases(uint16 NetIndex, uint16 LayerIndex) const
+	{
+		return {};
+	}
+
+	TArrayView<const float> GetNeuralNetworkLayerWeights(uint16 NetIndex, uint16 LayerIndex) const
+	{
+		return {};
+	}
+
 	uint16 LODCount;
 
 	TArray<FString> rawControls;
diff --git a/Engine/Plugins/Animation/RigLogic/Source/RigLogicModule/Public/SkelMeshDNAReader.h b/Engine/Plugins/Animation/RigLogic/Source/RigLogicModule/Public/SkelMeshDNAReader.h
--- a/Engine/Plugins/Animation/RigLogic/Source/RigLogicModule/Public/SkelMeshDNAReader.h
+++ b/Engine/Plugins/Animation/RigLogic/Source/RigLogicModule/Public/SkelMeshDNAReader.h
@@ -20,6 +20,9 @@ class RIGLOGICMODULE_API FSkelMeshDNAReader: public IDNAReader
 public:
 	explicit FSkelMeshDNAReader(UDNAAsset* DNAAsset);
 
+	// Header
+	uint16 GetFileFormatGeneration() const override;
+	uint16 GetFileFormatVersion() const override;
 	// Descriptor
 	FString GetName() const override;
 	EArchetype GetArchetype() const override;
@@ -118,18 +121,35 @@ public:
 	TArrayView<const float> GetBlendShapeTargetDeltaYs(uint16 MeshIndex, uint16 BlendShapeTargetIndex) const override;
 	TArrayView<const float> GetBlendShapeTargetDeltaZs(uint16 MeshIndex, uint16 BlendShapeTargetIndex) const override;
 	TArrayView<const uint32> GetBlendShapeTargetVertexIndices(uint16 MeshIndex, uint16 BlendShapeTargetIndex) const override;
+	// Machine Learned Behavior
+	uint16 GetMLControlCount() const override;
+	FString GetMLControlName(uint16 Index) const override;
+	uint16 GetNeuralNetworkCount() const override;
+	uint16 GetNeuralNetworkIndexListCount() const override;
+	TArrayView<const uint16> GetNeuralNetworkIndicesForLOD(uint16 LOD) const override;
+	uint16 GetMeshRegionCount(uint16 MeshIndex) const override;
+	FString GetMeshRegionName(uint16 MeshIndex, uint16 RegionIndex) const override;
+	TArrayView<const uint16> GetNeuralNetworkIndicesForMeshRegion(uint16 MeshIndex, uint16 RegionIndex) const override;
+	TArrayView<const uint16> GetNeuralNetworkInputIndices(uint16 NetIndex) const override;
+	TArrayView<const uint16> GetNeuralNetworkOutputIndices(uint16 NetIndex) const override;
+	uint16 GetNeuralNetworkLayerCount(uint16 NetIndex) const override;
+	EActivationFunction GetNeuralNetworkLayerActivationFunction(uint16 NetIndex, uint16 LayerIndex) const override;
+	TArrayView<const float> GetNeuralNetworkLayerActivationFunctionParameters(uint16 NetIndex, uint16 LayerIndex) const override;
+	TArrayView<const float> GetNeuralNetworkLayerBiases(uint16 NetIndex, uint16 LayerIndex) const override;
+	TArrayView<const float> GetNeuralNetworkLayerWeights(uint16 NetIndex, uint16 LayerIndex) const override;
+
 	void Unload(EDNADataLayer /*unused*/) override;
 
 private:
 	dna::Reader* Unwrap() const override;
 
 private:
-	/** Both BehaviorReader and DesignDataStreamReader are StreamReaders from DNAAsset 	   
+	/** Both BehaviorReader and GeometryReader are StreamReaders from DNAAsset
 	  * split out into run-time and in-editor parts from a full DNA that is either:
 	  * 1) imported manually into SkeletalMesh asset through ContentBrowser
 	  *	2) overwritten by GeneSplicer (GeneSplicerDNAReader) in a transient SkeletalMesh copy
 	  * They both just borrow DNAAsset's readers and are not owned by SkelMeshDNAReader
 	 **/
-	IBehaviorReader* BehaviorStreamReader = nullptr;
-	IGeometryReader* GeometryStreamReader = nullptr;
+	TSharedPtr<IDNAReader> BehaviorReader = nullptr;
+	TSharedPtr<IDNAReader> GeometryReader = nullptr;
 };
diff --git a/Engine/Plugins/Animation/RigLogic/Source/RigLogicModule/RigLogicModule.Build.cs b/Engine/Plugins/Animation/RigLogic/Source/RigLogicModule/RigLogicModule.Build.cs
--- a/Engine/Plugins/Animation/RigLogic/Source/RigLogicModule/RigLogicModule.Build.cs
+++ b/Engine/Plugins/Animation/RigLogic/Source/RigLogicModule/RigLogicModule.Build.cs
@@ -1,56 +1,41 @@
 // Copyright Epic Games, Inc. All Rights Reserved.
 
-using System;
-using System.IO;
-using UnrealBuildTool;
-
 namespace UnrealBuildTool.Rules
 {
-    public class RigLogicModule : ModuleRules
-    {
-        private string ModulePath
-        {
-            get { return ModuleDirectory; }
-        }
-
-        public RigLogicModule(ReadOnlyTargetRules Target) : base(Target)
-        {
-            if (Target.LinkType != TargetLinkType.Monolithic)
-            {
-                PublicDefinitions.Add("RL_SHARED=1");
-            }
-
-            PublicDependencyModuleNames.AddRange(
-                new string[]
-                {
-                    "Core",
-                    "CoreUObject",
-                    "Engine",
-                    "ControlRig",
-                    "RigLogicLib",
-                    "RigVM",
-                    "Projects"
+	public class RigLogicModule : ModuleRules
+	{
+		public RigLogicModule(ReadOnlyTargetRules Target) : base(Target)
+		{
+			PublicDependencyModuleNames.AddRange(
+				new string[]
+				{
+					"Core",
+					"CoreUObject",
+					"Engine",
+					"ControlRig",
+					"RigLogicLib",
+					"RigVM",
+					"Projects"
 				}
-            );
-
+			);
 
-            if (Target.Type == TargetType.Editor)
-            {
-                PublicDependencyModuleNames.Add("UnrealEd");
-                PublicDependencyModuleNames.Add("EditorFramework");
-                PublicDependencyModuleNames.Add("MessageLog");
+			if (Target.Type == TargetType.Editor)
+			{
+				PublicDependencyModuleNames.Add("UnrealEd");
+				PublicDependencyModuleNames.Add("EditorFramework");
+				PublicDependencyModuleNames.Add("MessageLog");
 
 				PrivateDependencyModuleNames.Add("SkeletalMeshUtilitiesCommon");
 				PrivateDependencyModuleNames.Add("RHI");
 				PrivateDependencyModuleNames.Add("RenderCore");
 			}
 
-            PrivateDependencyModuleNames.AddRange(
-                new string[]
+			PrivateDependencyModuleNames.AddRange(
+				new string[]
 				{
 					"AnimationCore"
 				}
-            );
-        }
-    }
+			);
+		}
+	}
 }
diff --git a/Engine/Source/Programs/IncludeTool/IncludeTool/Rules.cs b/Engine/Source/Programs/IncludeTool/IncludeTool/Rules.cs
--- a/Engine/Source/Programs/IncludeTool/IncludeTool/Rules.cs
+++ b/Engine/Source/Programs/IncludeTool/IncludeTool/Rules.cs
@@ -591,14 +591,24 @@ namespace IncludeTool
 			{
 				return true;
 			}
-			if (Markup.Type == PreprocessorMarkupType.Define && (Markup.Tokens[0].Text == "RLAPI" || Markup.Tokens[0].Text == "DNAAPI" || Markup.Tokens[0].Text == "TRIOAPI" || Markup.Tokens[0].Text == "SCAPI" || Markup.Tokens[0].Text == "PMAAPI"))
+			// Start of RigLogic exclusions
+			if (Markup.Type == PreprocessorMarkupType.Define && (Markup.Tokens[0].Text == "RLAPI" || Markup.Tokens[0].Text == "DNAAPI" || Markup.Tokens[0].Text == "TRIOAPI" || Markup.Tokens[0].Text == "SCAPI" || Markup.Tokens[0].Text == "PMAAPI" || Markup.Tokens[0].Text == "RAFAPI" || Markup.Tokens[0].Text == "GSAPI"))
 			{
 				return true;
 			}
-			if (Markup.Type == PreprocessorMarkupType.Elif && Markup.Tokens.Count == 4 && Markup.Tokens[2].Text == "RL_SHARED")
+			if (Markup.Type == PreprocessorMarkupType.Elif && Markup.Tokens.Count == 4 && (Markup.Tokens[2].Text == "RL_SHARED" || Markup.Tokens[2].Text == "GS_SHARED"))
 			{
 				return true;
 			}
+			if (Markup.Type == PreprocessorMarkupType.If && Markup.Tokens.Count == 4 && Markup.Tokens[2].Text == "_MSC_VER")
+			{
+				return true;
+			}
+			if (Markup.Type == PreprocessorMarkupType.Define && Markup.Tokens.Count == 8 && Markup.Tokens[0].Text == "FORCE_INLINE")
+			{
+				return true;
+			}
+			// End of RigLogic exclusions
 			if((File.Flags & SourceFileFlags.External) != 0)
 			{
 				return true;
