Fix bad merge of previous changelist

diff --git a/Engine/Plugins/Animation/RigLogic/Source/RigLogicEditor/RigLogicEditor.Build.cs b/Engine/Plugins/Animation/RigLogic/Source/RigLogicEditor/RigLogicEditor.Build.cs
--- a/Engine/Plugins/Animation/RigLogic/Source/RigLogicEditor/RigLogicEditor.Build.cs
+++ b/Engine/Plugins/Animation/RigLogic/Source/RigLogicEditor/RigLogicEditor.Build.cs
@@ -22,7 +22,9 @@ namespace UnrealBuildTool.Rules
 					"SlateCore",
 					"ApplicationCore",
 					"Slate",
-					"InputCore"
+					"InputCore",
+					"EditorWidgets",
+					"DesktopPlatform"
 				}
 			);
 
@@ -30,17 +32,10 @@ namespace UnrealBuildTool.Rules
 				new string[]
 				{
 					"PropertyEditor",
-					"AssetTools"
-				}
-			);
-
-                    "InputCore",
-					"EditorWidgets",
-					"DesktopPlatform"
-                    "AssetTools",
+					"AssetTools",
 					"Settings"
 				}
-
+			);
 		}
 	}
 }
