Renamed LOCTEXT Key duplicates

diff --git a/Engine/Plugins/Animation/RigLogic/Source/RigLogicEditor/Private/RigLogicEditor.cpp b/Engine/Plugins/Animation/RigLogic/Source/RigLogicEditor/Private/RigLogicEditor.cpp
--- a/Engine/Plugins/Animation/RigLogic/Source/RigLogicEditor/Private/RigLogicEditor.cpp
+++ b/Engine/Plugins/Animation/RigLogic/Source/RigLogicEditor/Private/RigLogicEditor.cpp
@@ -70,14 +70,14 @@ void FRigLogicEditor::GetDNAMenu(FMenuBuilder& MenuBuilder, const TArray<FAssetD
 	
 	MenuBuilder.AddMenuEntry(
 		LOCTEXT("Add DNA", "Import new DNA File"),
-		LOCTEXT("ImportDNA", "Import DNA"),
+		LOCTEXT("ImportDNA_Tooltip", "Import DNA"),
 		FSlateIcon(FAppStyle::GetAppStyleSetName(), "Icons.Import"),
 		FUIAction(FExecuteAction::CreateStatic(&FRigLogicEditor::ExecuteDNAImport, static_cast<UObject*>(Mesh)))
 	);
 
 	MenuBuilder.AddMenuEntry(
 		LOCTEXT("Reimport DNA", "Reimport existing DNA File"),
-		LOCTEXT("Reimport DNA", "Reimport DNA"),
+		LOCTEXT("ReimportDNA_Tooltip ", "Reimport DNA"),
 		FSlateIcon(FAppStyle::GetAppStyleSetName(), "Icons.Import"),
 		FUIAction(FExecuteAction::CreateStatic(&FRigLogicEditor::ExecuteDNAReimport, static_cast<UObject*>(Mesh)))
 	);
