Fix static analysis warning about possible overflow

diff --git a/Engine/Plugins/Animation/RigLogic/Source/RigLogicLib/Private/riglogic/riglogic/RigInstanceImpl.cpp b/Engine/Plugins/Animation/RigLogic/Source/RigLogicLib/Private/riglogic/riglogic/RigInstanceImpl.cpp
--- a/Engine/Plugins/Animation/RigLogic/Source/RigLogicLib/Private/riglogic/riglogic/RigInstanceImpl.cpp
+++ b/Engine/Plugins/Animation/RigLogic/Source/RigLogicLib/Private/riglogic/riglogic/RigInstanceImpl.cpp
@@ -151,7 +151,7 @@ float RigInstanceImpl::getMLControl(std::uint16_t index) const {
 }
 
 ConstArrayView<float> RigInstanceImpl::getMLControlValues() const {
-    const auto mlControlsOffset = static_cast<std::size_t>(rawControlCount + psdControlCount);
+    const auto mlControlsOffset = static_cast<std::size_t>(rawControlCount) + static_cast<std::size_t>(psdControlCount);
     return controlsInstance->getInputBuffer().subview(mlControlsOffset, mlControlCount);
 }
 
