Fix empty DNA file loading

diff --git a/Engine/Plugins/Animation/RigLogic/RigLogic.uplugin b/Engine/Plugins/Animation/RigLogic/RigLogic.uplugin
--- a/Engine/Plugins/Animation/RigLogic/RigLogic.uplugin
+++ b/Engine/Plugins/Animation/RigLogic/RigLogic.uplugin
@@ -1,9 +1,9 @@
 {
 	"FileVersion": 3,
 	"Version": 9.1,
-	"VersionName": "9.1.1",
-	"FriendlyName": "RigLogic Plugin v9.1.1",
-	"Description": "3Lateral RigLogic Plugin for Facial Animation v9.1.1",
+	"VersionName": "9.1.2",
+	"FriendlyName": "RigLogic Plugin v9.1.2",
+	"Description": "3Lateral RigLogic Plugin for Facial Animation v9.1.2",
 	"Category": "Animation",
 	"CreatedBy": "3Lateral",
 	"CreatedByURL": "https://www.3lateral.com/",
diff --git a/Engine/Plugins/Animation/RigLogic/Source/RigLogicLib/Private/dna/stream/FilteredBinaryInputArchive.cpp b/Engine/Plugins/Animation/RigLogic/Source/RigLogicLib/Private/dna/stream/FilteredBinaryInputArchive.cpp
--- a/Engine/Plugins/Animation/RigLogic/Source/RigLogicLib/Private/dna/stream/FilteredBinaryInputArchive.cpp
+++ b/Engine/Plugins/Animation/RigLogic/Source/RigLogicLib/Private/dna/stream/FilteredBinaryInputArchive.cpp
@@ -92,7 +92,7 @@ void FilteredBinaryInputArchive::processSubset(TContainer& dest, std::size_t off
 
 void FilteredBinaryInputArchive::process(RawDescriptor& dest) {
     BaseArchive::process(dest);
-    malformed = (dest.lodCount == static_cast<std::uint16_t>(0)) || (dest.lodCount > LODLimits::count());
+    malformed = (dest.lodCount > LODLimits::count());
     if (malformed) {
         return;
     }
diff --git a/Engine/Plugins/Animation/RigLogic/Source/RigLogicLib/Public/dna/version/Version.h b/Engine/Plugins/Animation/RigLogic/Source/RigLogicLib/Public/dna/version/Version.h
--- a/Engine/Plugins/Animation/RigLogic/Source/RigLogicLib/Public/dna/version/Version.h
+++ b/Engine/Plugins/Animation/RigLogic/Source/RigLogicLib/Public/dna/version/Version.h
@@ -4,5 +4,5 @@
 
 #define DNA_MAJOR_VERSION 9
 #define DNA_MINOR_VERSION 1
-#define DNA_PATCH_VERSION 1
-#define DNA_VERSION_STRING "9.1.1"
+#define DNA_PATCH_VERSION 2
+#define DNA_VERSION_STRING "9.1.2"
diff --git a/Engine/Plugins/Animation/RigLogic/Source/RigLogicLib/Public/riglogic/version/Version.h b/Engine/Plugins/Animation/RigLogic/Source/RigLogicLib/Public/riglogic/version/Version.h
--- a/Engine/Plugins/Animation/RigLogic/Source/RigLogicLib/Public/riglogic/version/Version.h
+++ b/Engine/Plugins/Animation/RigLogic/Source/RigLogicLib/Public/riglogic/version/Version.h
@@ -4,5 +4,5 @@
 
 #define RL_MAJOR_VERSION 9
 #define RL_MINOR_VERSION 1
-#define RL_PATCH_VERSION 1
-#define RL_VERSION_STRING "9.1.1"
+#define RL_PATCH_VERSION 2
+#define RL_VERSION_STRING "9.1.2"
diff --git a/Engine/Plugins/Animation/RigLogic/Source/RigLogicLibTest/Private/dnatests/TestStreamReadWriteIntegration.cpp b/Engine/Plugins/Animation/RigLogic/Source/RigLogicLibTest/Private/dnatests/TestStreamReadWriteIntegration.cpp
--- a/Engine/Plugins/Animation/RigLogic/Source/RigLogicLibTest/Private/dnatests/TestStreamReadWriteIntegration.cpp
+++ b/Engine/Plugins/Animation/RigLogic/Source/RigLogicLibTest/Private/dnatests/TestStreamReadWriteIntegration.cpp
@@ -534,6 +534,8 @@ using TRawCopyTestParameters = ::testing::Types<
     RawCopyParameters<RawV22, RawV21, UnknownLayerPolicy::Ignore, 2, 1>,
     RawCopyParameters<RawV2xNewer, RawV22WithUnknownDataFromNewer2x, UnknownLayerPolicy::Preserve, 2, 2>,
     RawCopyParameters<RawV2xNewer, RawV22Empty, UnknownLayerPolicy::Ignore, 2, 2>,
+    RawCopyParameters<RawV22Empty, RawV22Empty, UnknownLayerPolicy::Preserve, 2, 2>,
+    RawCopyParameters<RawV22Empty, RawV22Empty, UnknownLayerPolicy::Ignore, 2, 2>,
     RawCopyParameters<RawV23, RawV22DowngradedFromV23, UnknownLayerPolicy::Preserve, 2, 2>,
     RawCopyParameters<RawV23, RawV22WithUnknownDataIgnoredAndDNARewritten, UnknownLayerPolicy::Ignore, 2, 2>
     >;
@@ -608,4 +610,103 @@ TEST(StreamReadWriteIntegrationTest, ReadWriteJSON) {
     ASSERT_TRUE(dna::Status::isOk());
 }
 
+using TReadWriteMultipleParameters = ::testing::Types<
+    ReadWriteMultipleParameters<RawV21>,
+    ReadWriteMultipleParameters<RawV22>,
+    ReadWriteMultipleParameters<RawV23>,
+    ReadWriteMultipleParameters<RawV22Empty>,
+    ReadWriteMultipleParameters<RawV22WithUnknownDataIgnoredAndDNARewritten>,
+    ReadWriteMultipleParameters<RawV2xNewerWithUnknownDataIgnoredAndDNARewritten>,
+    ReadWriteMultipleParameters<RawV2xNewerWithUnknownDataPreservedAndDNARewritten>,
+    ReadWriteMultipleParameters<RawV22WithUnknownDataFromNewer2x>,
+    ReadWriteMultipleParameters<RawV2xNewer>,
+    ReadWriteMultipleParameters<RawV22DowngradedFromV23>
+    >;
+TYPED_TEST_SUITE(StreamReadWriteMultipleIntegrationTest, TReadWriteMultipleParameters, );
+
+TYPED_TEST(StreamReadWriteMultipleIntegrationTest, ReadWriteTwoDNAsToSameStream) {
+    using CurrentParameters = typename TestFixture::Parameters;
+
+    const auto bytes = CurrentParameters::RawBytes::getBytes();
+    auto source = pma::makeScoped<trio::MemoryStream>();
+    source->write(reinterpret_cast<const char*>(bytes.data()), bytes.size());
+
+    source->seek(0);
+    auto sourceReader = pma::makeScoped<BinaryStreamReader>(source.get(),
+                                                            DataLayer::All,
+                                                            UnknownLayerPolicy::Preserve,
+                                                            static_cast<std::uint16_t>(0));
+    sourceReader->read();
+    ASSERT_TRUE(dna::Status::isOk());
+
+    auto clone = pma::makeScoped<trio::MemoryStream>();
+    auto cloneWriter1 = pma::makeScoped<BinaryStreamWriter>(clone.get());
+    cloneWriter1->setFrom(sourceReader.get(), DataLayer::All, UnknownLayerPolicy::Preserve);
+    cloneWriter1->write();
+    ASSERT_TRUE(dna::Status::isOk());
+
+    // Stream position is reset on open / close of stream (by implementation of trio::MemoryStream)
+    const std::uint64_t firstDNASize = clone->size();
+    clone->seek(firstDNASize);
+
+    auto cloneWriter2 = pma::makeScoped<BinaryStreamWriter>(clone.get());
+    cloneWriter2->setFrom(sourceReader.get(), DataLayer::All, UnknownLayerPolicy::Preserve);
+    cloneWriter2->write();
+    ASSERT_TRUE(dna::Status::isOk());
+
+    clone->seek(0ul);
+
+    auto cloneReader1 = pma::makeScoped<BinaryStreamReader>(clone.get(),
+                                                            DataLayer::All,
+                                                            UnknownLayerPolicy::Preserve,
+                                                            static_cast<std::uint16_t>(0));
+    cloneReader1->read();
+    ASSERT_TRUE(dna::Status::isOk());
+
+    // Stream position is reset on open / close of stream (by implementation of trio::MemoryStream)
+    clone->seek(firstDNASize);
+
+    auto cloneReader2 = pma::makeScoped<BinaryStreamReader>(clone.get(),
+                                                            DataLayer::All,
+                                                            UnknownLayerPolicy::Preserve,
+                                                            static_cast<std::uint16_t>(0));
+    cloneReader2->read();
+    ASSERT_TRUE(dna::Status::isOk());
+
+    auto cloneRewritten = pma::makeScoped<trio::MemoryStream>();
+    auto cloneRewriter1 = pma::makeScoped<BinaryStreamWriter>(cloneRewritten.get());
+    cloneRewriter1->setFrom(cloneReader1.get(), DataLayer::All, UnknownLayerPolicy::Preserve);
+    cloneRewriter1->write();
+    ASSERT_TRUE(dna::Status::isOk());
+
+    // Stream position is reset on open / close of stream (by implementation of trio::MemoryStream)
+    cloneRewritten->seek(cloneRewritten->size());
+
+    auto cloneRewriter2 = pma::makeScoped<BinaryStreamWriter>(cloneRewritten.get());
+    cloneRewriter2->setFrom(cloneReader2.get(), DataLayer::All, UnknownLayerPolicy::Preserve);
+    cloneRewriter2->write();
+    ASSERT_TRUE(dna::Status::isOk());
+
+    clone->seek(0ul);
+    cloneRewritten->seek(0ul);
+
+    #if !defined(__clang__) && defined(__GNUC__)
+        #pragma GCC diagnostic push
+        #pragma GCC diagnostic ignored "-Wuseless-cast"
+    #endif
+    const auto cloneSize = static_cast<std::size_t>(clone->size());
+    const auto cloneRewrittenSize = static_cast<std::size_t>(cloneRewritten->size());
+    #if !defined(__clang__) && defined(__GNUC__)
+        #pragma GCC diagnostic pop
+    #endif
+    std::vector<char> copiedCloneBytes(cloneSize);
+    clone->read(copiedCloneBytes.data(), cloneSize);
+
+    std::vector<char> copiedCloneRewrittenBytes(cloneRewrittenSize);
+    cloneRewritten->read(copiedCloneRewrittenBytes.data(), cloneRewrittenSize);
+
+    ASSERT_EQ(cloneSize, cloneRewrittenSize);
+    ASSERT_EQ(copiedCloneBytes, copiedCloneRewrittenBytes);
+}
+
 }  // namespace dna
diff --git a/Engine/Plugins/Animation/RigLogic/Source/RigLogicLibTest/Private/dnatests/TestStreamReadWriteIntegration.h b/Engine/Plugins/Animation/RigLogic/Source/RigLogicLibTest/Private/dnatests/TestStreamReadWriteIntegration.h
--- a/Engine/Plugins/Animation/RigLogic/Source/RigLogicLibTest/Private/dnatests/TestStreamReadWriteIntegration.h
+++ b/Engine/Plugins/Animation/RigLogic/Source/RigLogicLibTest/Private/dnatests/TestStreamReadWriteIntegration.h
@@ -66,3 +66,16 @@ class StreamReadWriteRawCopyIntegrationTest : public ::testing::Test {
         using Parameters = TRawCopyParameters;
 
 };
+
+template<typename TRawBytes>
+struct ReadWriteMultipleParameters {
+    using RawBytes = TRawBytes;
+
+};
+
+template<typename TReadWriteMultipleParameters>
+class StreamReadWriteMultipleIntegrationTest : public ::testing::Test {
+    protected:
+        using Parameters = TReadWriteMultipleParameters;
+
+};
diff --git a/Engine/Plugins/Animation/RigLogic/Source/RigLogicModule/Private/DNAAsset.cpp b/Engine/Plugins/Animation/RigLogic/Source/RigLogicModule/Private/DNAAsset.cpp
--- a/Engine/Plugins/Animation/RigLogic/Source/RigLogicModule/Private/DNAAsset.cpp
+++ b/Engine/Plugins/Animation/RigLogic/Source/RigLogicModule/Private/DNAAsset.cpp
@@ -33,7 +33,7 @@ DEFINE_LOG_CATEGORY(LogDNAAsset);
 
 static constexpr uint32 AVG_EMPTY_SIZE = 4 * 1024;
 static constexpr uint32 AVG_BEHAVIOR_SIZE = 5 * 1024 * 1024;
-static constexpr uint32 AVG_MACHINE_LEARNED_BEHAVIOR_SIZE = 15 * 1024 * 1024;
+static constexpr uint32 AVG_MACHINE_LEARNED_BEHAVIOR_SIZE = 5 * 1024 * 1024;
 static constexpr uint32 AVG_GEOMETRY_SIZE = 50 * 1024 * 1024;
 
 static TSharedPtr<IDNAReader> ReadDNAFromStream(rl4::BoundedIOStream* Stream, EDNADataLayer Layer, uint16 MaxLOD)
