Added DNA path to SkeletalMesh tooltip enabling advanced search for SkeletalMesh with DNA file in content browser

diff --git a/Engine/Plugins/Animation/RigLogic/Source/RigLogicEditor/Private/RigLogicEditor.cpp b/Engine/Plugins/Animation/RigLogic/Source/RigLogicEditor/Private/RigLogicEditor.cpp
--- a/Engine/Plugins/Animation/RigLogic/Source/RigLogicEditor/Private/RigLogicEditor.cpp
+++ b/Engine/Plugins/Animation/RigLogic/Source/RigLogicEditor/Private/RigLogicEditor.cpp
@@ -16,7 +16,8 @@
 #include "Framework/Notifications/NotificationManager.h"
 #include "Widgets/Notifications/SNotificationList.h"
 #include "Misc/Paths.h"
-
+#include "AssetRegistry/IAssetRegistry.h"
+#include "AssetRegistry/AssetData.h"
 
 
 IMPLEMENT_MODULE(FRigLogicEditor, RigLogicEditor)
@@ -31,6 +32,8 @@ void FRigLogicEditor::StartupModule()
 	TArray<FContentBrowserMenuExtender_SelectedAssets>& MenuExtenderDelegates = ContentBrowserModule.GetAllAssetViewContextMenuExtenders(); 
 	
 	MenuExtenderDelegates.Add(FContentBrowserMenuExtender_SelectedAssets::CreateStatic(&FRigLogicEditor::OnExtendSkelMeshWithDNASelectionMenu));
+
+	UObject::FAssetRegistryTag::OnGetExtraObjectTags.AddStatic(&GetAssetRegistryTagsForDNA);
 }
 
 
@@ -126,6 +129,30 @@ void FRigLogicEditor::ExecuteDNAReimport(class UObject* Mesh)
 	}
 }
 
+void FRigLogicEditor::GetAssetRegistryTagsForDNA(const class UObject* Object, TArray<UObject::FAssetRegistryTag>& OutTags)
+{
+ 	if(Object->GetClass()->IsChildOf(USkeletalMesh::StaticClass()))
+ 	{
+		FString DNAname = (LOCTEXT("DnaNotOnSkeletalMesh","No DNA Attached")).ToString();
+		USkeletalMesh* SkelMesh = const_cast<USkeletalMesh*>(Cast<USkeletalMesh>(Object));
+ 		const TArray<UAssetUserData*>* AssetDataArr = SkelMesh->GetAssetUserDataArray();
+ 		if(!AssetDataArr->IsEmpty())
+		{
+			UAssetUserData* UserData = SkelMesh->GetAssetUserDataOfClass(UDNAAsset::StaticClass());
+ 			if(UserData)
+ 			{	
+ 				const UDNAAsset* DNAAsset = Cast<UDNAAsset>(UserData);
+ 				if(DNAAsset)
+ 				{	
+					DNAname = DNAAsset->DnaFileName;
+					FPaths::NormalizeFilename(DNAname);
+ 				}
+ 			}
+ 		}
+		OutTags.Add(UObject::FAssetRegistryTag("DNA", DNAname, UObject::FAssetRegistryTag::TT_Alphabetical));
+ 	}
+}
+
 void FRigLogicEditor::ShutdownModule()
 {
 
diff --git a/Engine/Plugins/Animation/RigLogic/Source/RigLogicEditor/Public/RigLogicEditor.h b/Engine/Plugins/Animation/RigLogic/Source/RigLogicEditor/Public/RigLogicEditor.h
--- a/Engine/Plugins/Animation/RigLogic/Source/RigLogicEditor/Public/RigLogicEditor.h
+++ b/Engine/Plugins/Animation/RigLogic/Source/RigLogicEditor/Public/RigLogicEditor.h
@@ -3,6 +3,7 @@
 #pragma once
 
 #include "CoreMinimal.h"
+#include "UObject/Object.h"
 #include "Modules/ModuleInterface.h"
 
 class FExtender;
@@ -22,6 +23,8 @@ private:
 	static void CreateDnaActionsSubMenu(FMenuBuilder& MenuBuilder, const TArray<FAssetData> SelectedAssets);
 	static void GetDNAMenu(FMenuBuilder& MenuBuilder, const TArray<FAssetData> SelectedAssets);
 
-	static void ExecuteDNAImport(class UObject* Mesh);
-	static void ExecuteDNAReimport(class UObject* Mesh);
+	static void ExecuteDNAImport(UObject* Mesh);
+	static void ExecuteDNAReimport(UObject* Mesh);
+
+	static void GetAssetRegistryTagsForDNA(const UObject* Object, TArray<UObject::FAssetRegistryTag>& OutTags);
 };
\ No newline at end of file
diff --git a/Engine/Plugins/Animation/RigLogic/Source/RigLogicModule/Public/DNAAsset.h b/Engine/Plugins/Animation/RigLogic/Source/RigLogicModule/Public/DNAAsset.h
--- a/Engine/Plugins/Animation/RigLogic/Source/RigLogicModule/Public/DNAAsset.h
+++ b/Engine/Plugins/Animation/RigLogic/Source/RigLogicModule/Public/DNAAsset.h
@@ -51,7 +51,7 @@ public:
 	TSharedPtr<IDNAReader> GetGeometryReader();
 #endif
 
-	UPROPERTY(VisibleAnywhere, Category = ImportSettings)
+	UPROPERTY(VisibleAnywhere, AssetRegistrySearchable, Category = ImportSettings)
 	FString DnaFileName;
 
 	bool Init(const FString& Filename);
