Fixed import notification message for importing DNA file with same name as SkeletalMesh

diff --git a/Engine/Plugins/Animation/RigLogic/Source/RigLogicEditor/Private/DNAAssetImportFactory.cpp b/Engine/Plugins/Animation/RigLogic/Source/RigLogicEditor/Private/DNAAssetImportFactory.cpp
--- a/Engine/Plugins/Animation/RigLogic/Source/RigLogicEditor/Private/DNAAssetImportFactory.cpp
+++ b/Engine/Plugins/Animation/RigLogic/Source/RigLogicEditor/Private/DNAAssetImportFactory.cpp
@@ -20,6 +20,9 @@
 #include "Misc/MessageDialog.h"
 #include "Misc/Paths.h"
 
+#include "Framework/Notifications/NotificationManager.h"
+#include "Widgets/Notifications/SNotificationList.h"
+
 #include "AssetImportTask.h"
 #include "Misc/FileHelper.h"
 
@@ -76,7 +79,18 @@ EReimportResult::Type UDNAAssetImportFactory::Reimport(UObject* Obj, int32 Sourc
 
 	if (Result != nullptr)
 	{
-		return EReimportResult::Succeeded;
+		// TODO: Find a better solution to showing the right message, withouth this we get reimport message success for Skeletal Mesh
+		FNotificationInfo Info(LOCTEXT("DNA_ReimportSuccessMessage", "DNA file successfully imported"));
+		Info.ExpireDuration = 3.0f;
+		Info.bUseLargeFont = false;
+		TSharedPtr<SNotificationItem> Notification = FSlateNotificationManager::Get().AddNotification(Info);
+		if (Notification.IsValid())
+		{
+			Notification->SetCompletionState(SNotificationItem::CS_Success);
+		}
+
+		//Cancel is sent here because the notification created for successful import ( reimport is only registered because DNA has the same name as SkelMesh ) is already showed with above code
+		return EReimportResult::Cancelled;
 	}
 
 	return EReimportResult::Failed;
diff --git a/Engine/Plugins/Animation/RigLogic/Source/RigLogicEditor/Private/RigLogicEditor.cpp b/Engine/Plugins/Animation/RigLogic/Source/RigLogicEditor/Private/RigLogicEditor.cpp
--- a/Engine/Plugins/Animation/RigLogic/Source/RigLogicEditor/Private/RigLogicEditor.cpp
+++ b/Engine/Plugins/Animation/RigLogic/Source/RigLogicEditor/Private/RigLogicEditor.cpp
@@ -95,7 +95,14 @@ void FRigLogicEditor::ExecuteDNAImport(UObject* Mesh)
 	UFactory* Factory = Cast<UFactory>(UDNAAssetImportFactory::StaticClass()->GetDefaultObject());
 	FString Path = FPaths::GetPath(Mesh->GetPathName());
 	
-	AssetToolsModule.Get().ImportAssets(Filenames,Path, Factory,true, nullptr,true);
+	if((Mesh->GetFName()).ToString().Compare(Filenames[0]))
+	{
+		FReimportManager::Instance()->Reimport(Mesh, false, true, Filenames[0]);
+	}
+	else
+	{
+		AssetToolsModule.Get().ImportAssets(Filenames, Path, Factory, true, nullptr, true);
+	}
 }
 
 void FRigLogicEditor::ExecuteDNAReimport(class UObject* Mesh)
@@ -113,7 +120,14 @@ void FRigLogicEditor::ExecuteDNAReimport(class UObject* Mesh)
 			UFactory* Factory = Cast<UFactory>(UDNAAssetImportFactory::StaticClass()->GetDefaultObject());
 			FString Path = FPaths::GetPath(Mesh->GetPathName());
 
-			AssetToolsModule.Get().ImportAssets(Filenames, Path, Factory, true, nullptr, true);
+			if ((Mesh->GetFName()).ToString().Compare(Filenames[0]))
+			{
+				FReimportManager::Instance()->Reimport(Mesh, false, true, Filenames[0]);
+			}
+			else
+			{
+				AssetToolsModule.Get().ImportAssets(Filenames, Path, Factory, true, nullptr, true);
+			}
 		}
 	}
 	else
