Skeletal Mesh DNA improvements

diff --git a/Engine/Plugins/Animation/RigLogic/Source/RigLogicEditor/Private/RigLogicEditor.cpp b/Engine/Plugins/Animation/RigLogic/Source/RigLogicEditor/Private/RigLogicEditor.cpp
--- a/Engine/Plugins/Animation/RigLogic/Source/RigLogicEditor/Private/RigLogicEditor.cpp
+++ b/Engine/Plugins/Animation/RigLogic/Source/RigLogicEditor/Private/RigLogicEditor.cpp
@@ -24,13 +24,13 @@ IMPLEMENT_MODULE(FRigLogicEditor, RigLogicEditor)
 
 DEFINE_LOG_CATEGORY_STATIC(LogRigLogicEditor, Log, All);
 
-#define LOCTEXT_NAMESPACE "DNAAssetImport"
+#define LOCTEXT_NAMESPACE "RigLogicEditor"
 
 void FRigLogicEditor::StartupModule()
 {
-	FContentBrowserModule& ContentBrowserModule = FModuleManager::LoadModuleChecked<FContentBrowserModule>(TEXT("ContentBrowser"));	
-	TArray<FContentBrowserMenuExtender_SelectedAssets>& MenuExtenderDelegates = ContentBrowserModule.GetAllAssetViewContextMenuExtenders(); 
-	
+	FContentBrowserModule& ContentBrowserModule = FModuleManager::LoadModuleChecked<FContentBrowserModule>(TEXT("ContentBrowser"));
+	TArray<FContentBrowserMenuExtender_SelectedAssets>& MenuExtenderDelegates = ContentBrowserModule.GetAllAssetViewContextMenuExtenders();
+
 	MenuExtenderDelegates.Add(FContentBrowserMenuExtender_SelectedAssets::CreateStatic(&FRigLogicEditor::OnExtendSkelMeshWithDNASelectionMenu));
 
 	UObject::FAssetRegistryTag::OnGetExtraObjectTags.AddStatic(&GetAssetRegistryTagsForDNA);
@@ -62,17 +62,17 @@ void FRigLogicEditor::CreateDnaActionsSubMenu(FMenuBuilder& MenuBuilder, const T
 			FNewMenuDelegate::CreateStatic(&FRigLogicEditor::GetDNAMenu, SelectedAssets),
 			false,
 			FSlateIcon(FAppStyle::GetAppStyleSetName(), "Icons.Import")
-		);	
+		);
 	}
-	
+
 }
 
 void FRigLogicEditor::GetDNAMenu(FMenuBuilder& MenuBuilder, const TArray<FAssetData> SelectedAssets)
 {
 	auto Mesh = SelectedAssets[0].GetAsset();
-	
+
 	MenuBuilder.AddMenuEntry(
-		LOCTEXT("Add DNA", "Import new DNA File"),
+		LOCTEXT("Import DNA", "Import new DNA File"),
 		LOCTEXT("ImportDNA_Tooltip", "Import DNA"),
 		FSlateIcon(FAppStyle::GetAppStyleSetName(), "Icons.Import"),
 		FUIAction(FExecuteAction::CreateStatic(&FRigLogicEditor::ExecuteDNAImport, static_cast<UObject*>(Mesh)))
@@ -89,13 +89,13 @@ void FRigLogicEditor::GetDNAMenu(FMenuBuilder& MenuBuilder, const TArray<FAssetD
 void FRigLogicEditor::ExecuteDNAImport(UObject* Mesh)
 {
 	FAssetToolsModule& AssetToolsModule = FModuleManager::Get().LoadModuleChecked<FAssetToolsModule>("AssetTools");
-	
+
 	FDNAImporter* DNAImporter = FDNAImporter::GetInstance();
 	const TArray<FString> Filenames = { DNAImporter->PromptForDNAImportFile() };
 	UFactory* Factory = Cast<UFactory>(UDNAAssetImportFactory::StaticClass()->GetDefaultObject());
 	FString Path = FPaths::GetPath(Mesh->GetPathName());
-	
-	if((Mesh->GetFName()).ToString().Compare(Filenames[0]))
+
+	if ((Mesh->GetFName()).ToString().Compare(Filenames[0]))
 	{
 		FReimportManager::Instance()->Reimport(Mesh, false, true, Filenames[0]);
 	}
@@ -111,10 +111,10 @@ void FRigLogicEditor::ExecuteDNAReimport(class UObject* Mesh)
 	USkeletalMesh* SkelMesh = Cast<USkeletalMesh>(Mesh);
 	const TArray<UAssetUserData*>* AssetData = SkelMesh->GetAssetUserDataArray();
 
-	if(!AssetData->IsEmpty())
+	if (!AssetData->IsEmpty())
 	{
 		const UDNAAsset* DNAAsset = Cast<UDNAAsset>(&AssetData[0]);
-		if(DNAAsset)
+		if (DNAAsset)
 		{
 			const TArray<FString> Filenames = { DNAAsset->DnaFileName };
 			UFactory* Factory = Cast<UFactory>(UDNAAssetImportFactory::StaticClass()->GetDefaultObject());
@@ -145,31 +145,31 @@ void FRigLogicEditor::ExecuteDNAReimport(class UObject* Mesh)
 
 void FRigLogicEditor::GetAssetRegistryTagsForDNA(const class UObject* Object, TArray<UObject::FAssetRegistryTag>& OutTags)
 {
- 	if(Object->GetClass()->IsChildOf(USkeletalMesh::StaticClass()))
- 	{
-		FString DNAname = (LOCTEXT("DnaNotOnSkeletalMesh","No DNA Attached")).ToString();
+	if (Object->GetClass()->IsChildOf(USkeletalMesh::StaticClass()))
+	{
+		FString DNAname = (LOCTEXT("DnaNotOnSkeletalMesh", "No DNA Attached")).ToString();
 		USkeletalMesh* SkelMesh = const_cast<USkeletalMesh*>(Cast<USkeletalMesh>(Object));
- 		const TArray<UAssetUserData*>* AssetDataArr = SkelMesh->GetAssetUserDataArray();
- 		if(!AssetDataArr->IsEmpty())
+		const TArray<UAssetUserData*>* AssetDataArr = SkelMesh->GetAssetUserDataArray();
+		if (!AssetDataArr->IsEmpty())
 		{
 			UAssetUserData* UserData = SkelMesh->GetAssetUserDataOfClass(UDNAAsset::StaticClass());
- 			if(UserData)
- 			{	
- 				const UDNAAsset* DNAAsset = Cast<UDNAAsset>(UserData);
- 				if(DNAAsset)
- 				{	
+			if (UserData)
+			{
+				const UDNAAsset* DNAAsset = Cast<UDNAAsset>(UserData);
+				if (DNAAsset)
+				{
 					DNAname = DNAAsset->DnaFileName;
 					FPaths::NormalizeFilename(DNAname);
- 				}
- 			}
- 		}
+				}
+			}
+		}
 		OutTags.Add(UObject::FAssetRegistryTag("DNA", DNAname, UObject::FAssetRegistryTag::TT_Alphabetical));
- 	}
+	}
 }
 
 void FRigLogicEditor::ShutdownModule()
 {
-
 }
 
-#undef LOCTEXT_NAMESPACE
\ No newline at end of file
+
+#undef LOCTEXT_NAMESPACE
diff --git a/Engine/Plugins/Animation/RigLogic/Source/RigLogicModule/Public/DNAAsset.h b/Engine/Plugins/Animation/RigLogic/Source/RigLogicModule/Public/DNAAsset.h
--- a/Engine/Plugins/Animation/RigLogic/Source/RigLogicModule/Public/DNAAsset.h
+++ b/Engine/Plugins/Animation/RigLogic/Source/RigLogicModule/Public/DNAAsset.h
@@ -46,9 +46,9 @@ public:
 	TObjectPtr<class UAssetImportData> AssetImportData;
 #endif
 
-	TSharedPtr<IDNAReader> GetBehaviorReader();
+	TSharedPtr<IBehaviorReader> GetBehaviorReader();
 #if WITH_EDITORONLY_DATA
-	TSharedPtr<IDNAReader> GetGeometryReader();
+	TSharedPtr<IGeometryReader> GetGeometryReader();
 #endif
 
 	UPROPERTY(VisibleAnywhere, AssetRegistrySearchable, Category = ImportSettings)
@@ -91,11 +91,11 @@ private:
 
 	/** Part of the .dna file needed for run-time execution of RigLogic;
 	 **/
-	TSharedPtr<IDNAReader> BehaviorReader;
+	TSharedPtr<IBehaviorReader> BehaviorReader;
 
 	/** Part of the .dna file used design-time for updating SkeletalMesh geometry
 	 **/
-	TSharedPtr<IDNAReader> GeometryReader;
+	TSharedPtr<IGeometryReader> GeometryReader;
 
 	/** Runtime data necessary for rig computations that is shared between
 	  * multiple rig instances based on the same DNA.
