Minor fix for previous changelist

diff --git a/Engine/Plugins/Animation/RigLogic/Source/RigLogicModule/Private/DNAAsset.cpp b/Engine/Plugins/Animation/RigLogic/Source/RigLogicModule/Private/DNAAsset.cpp
--- a/Engine/Plugins/Animation/RigLogic/Source/RigLogicModule/Private/DNAAsset.cpp
+++ b/Engine/Plugins/Animation/RigLogic/Source/RigLogicModule/Private/DNAAsset.cpp
@@ -94,14 +94,14 @@ UDNAAsset::UDNAAsset() : RigRuntimeContext{nullptr}
 
 UDNAAsset::~UDNAAsset() = default;
 
-TSharedPtr<IDNAReader> UDNAAsset::GetBehaviorReader()
+TSharedPtr<IBehaviorReader> UDNAAsset::GetBehaviorReader()
 {
 	FScopeLock ScopeLock{&DNAUpdateSection};
 	return BehaviorReader;
 }
 
 #if WITH_EDITORONLY_DATA
-TSharedPtr<IDNAReader> UDNAAsset::GetGeometryReader()
+TSharedPtr<IGeometryReader> UDNAAsset::GetGeometryReader()
 {
 	FScopeLock ScopeLock{&DNAUpdateSection};
 	return GeometryReader;
