Fix for compilation errors in RigLogic module (issue 258024)

diff --git a/Engine/Plugins/Animation/RigLogic/Source/RigLogicModule/Private/DNAAsset.cpp b/Engine/Plugins/Animation/RigLogic/Source/RigLogicModule/Private/DNAAsset.cpp
--- a/Engine/Plugins/Animation/RigLogic/Source/RigLogicModule/Private/DNAAsset.cpp
+++ b/Engine/Plugins/Animation/RigLogic/Source/RigLogicModule/Private/DNAAsset.cpp
@@ -94,14 +94,14 @@ UDNAAsset::UDNAAsset() : RigRuntimeContext{nullptr}
 
 UDNAAsset::~UDNAAsset() = default;
 
-TSharedPtr<IBehaviorReader> UDNAAsset::GetBehaviorReader()
+TSharedPtr<IDNAReader> UDNAAsset::GetBehaviorReader()
 {
 	FScopeLock ScopeLock{&DNAUpdateSection};
 	return BehaviorReader;
 }
 
 #if WITH_EDITORONLY_DATA
-TSharedPtr<IGeometryReader> UDNAAsset::GetGeometryReader()
+TSharedPtr<IDNAReader> UDNAAsset::GetGeometryReader()
 {
 	FScopeLock ScopeLock{&DNAUpdateSection};
 	return GeometryReader;
diff --git a/Engine/Plugins/Animation/RigLogic/Source/RigLogicModule/Public/DNAAsset.h b/Engine/Plugins/Animation/RigLogic/Source/RigLogicModule/Public/DNAAsset.h
--- a/Engine/Plugins/Animation/RigLogic/Source/RigLogicModule/Public/DNAAsset.h
+++ b/Engine/Plugins/Animation/RigLogic/Source/RigLogicModule/Public/DNAAsset.h
@@ -46,9 +46,9 @@ public:
 	TObjectPtr<class UAssetImportData> AssetImportData;
 #endif
 
-	TSharedPtr<IBehaviorReader> GetBehaviorReader();
+	TSharedPtr<IDNAReader> GetBehaviorReader();
 #if WITH_EDITORONLY_DATA
-	TSharedPtr<IGeometryReader> GetGeometryReader();
+	TSharedPtr<IDNAReader> GetGeometryReader();
 #endif
 
 	UPROPERTY(VisibleAnywhere, AssetRegistrySearchable, Category = ImportSettings)
@@ -91,11 +91,11 @@ private:
 
 	/** Part of the .dna file needed for run-time execution of RigLogic;
 	 **/
-	TSharedPtr<IBehaviorReader> BehaviorReader;
+	TSharedPtr<IDNAReader> BehaviorReader;
 
 	/** Part of the .dna file used design-time for updating SkeletalMesh geometry
 	 **/
-	TSharedPtr<IGeometryReader> GeometryReader;
+	TSharedPtr<IDNAReader> GeometryReader;
 
 	/** Runtime data necessary for rig computations that is shared between
 	  * multiple rig instances based on the same DNA.
