Fix loading of MachineLearnedBehavior layer when importing DNA asset

diff --git a/Engine/Plugins/Animation/RigLogic/Source/RigLogicModule/Private/DNAAsset.cpp b/Engine/Plugins/Animation/RigLogic/Source/RigLogicModule/Private/DNAAsset.cpp
--- a/Engine/Plugins/Animation/RigLogic/Source/RigLogicModule/Private/DNAAsset.cpp
+++ b/Engine/Plugins/Animation/RigLogic/Source/RigLogicModule/Private/DNAAsset.cpp
@@ -235,7 +235,7 @@ bool UDNAAsset::Init(const FString& DNAFilename)
 	FScopeLock ScopeLock{&DNAUpdateSection};
 
 	// Load run-time data (behavior) from whole-DNA buffer into BehaviorReader
-	BehaviorReader = ReadDNAFromBuffer(&TempFileBuffer, EDNADataLayer::Behavior, 0u); //0u = MaxLOD
+	BehaviorReader = ReadDNAFromBuffer(&TempFileBuffer, EDNADataLayer::Behavior | EDNADataLayer::MachineLearnedBehavior, 0u); //0u = MaxLOD
 	if (!BehaviorReader.IsValid())
 	{
 		return false;
@@ -275,7 +275,7 @@ void UDNAAsset::Serialize(FArchive& Ar)
 		if (Ar.IsLoading())
 		{
 			FArchiveMemoryStream BehaviorStream{&Ar};
-			BehaviorReader = ReadDNAFromStream(&BehaviorStream, EDNADataLayer::Behavior, 0u); //0u = max LOD
+			BehaviorReader = ReadDNAFromStream(&BehaviorStream, EDNADataLayer::Behavior | EDNADataLayer::MachineLearnedBehavior, 0u); //0u = max LOD
 			// Geometry data is always present (even if only as an empty placeholder), just so the uasset
 			// format remains consistent between editor and non-editor builds
 			FArchiveMemoryStream GeometryStream{&Ar};
@@ -293,7 +293,7 @@ void UDNAAsset::Serialize(FArchive& Ar)
 			TSharedPtr<IDNAReader> EmptyDNA = CreateEmptyDNA(AVG_EMPTY_SIZE);
 			IDNAReader* BehaviorReaderPtr = (BehaviorReader.IsValid() ? static_cast<IDNAReader*>(BehaviorReader.Get()) : EmptyDNA.Get());
 			FArchiveMemoryStream BehaviorStream{&Ar};
-			WriteDNAToStream(BehaviorReaderPtr, EDNADataLayer::Behavior, &BehaviorStream);
+			WriteDNAToStream(BehaviorReaderPtr, EDNADataLayer::Behavior | EDNADataLayer::MachineLearnedBehavior, &BehaviorStream);
 
 			// When cooking (or when there was no Geometry data available), an empty DNA structure is written
 			// into the stream, serving as a placeholder just so uasset files can be conveniently loaded
