Fix for reimport DNA via context menu

diff --git a/Engine/Plugins/Animation/RigLogic/Source/RigLogicEditor/Private/RigLogicEditor.cpp b/Engine/Plugins/Animation/RigLogic/Source/RigLogicEditor/Private/RigLogicEditor.cpp
--- a/Engine/Plugins/Animation/RigLogic/Source/RigLogicEditor/Private/RigLogicEditor.cpp
+++ b/Engine/Plugins/Animation/RigLogic/Source/RigLogicEditor/Private/RigLogicEditor.cpp
@@ -113,7 +113,7 @@ void FRigLogicEditor::ExecuteDNAReimport(class UObject* Mesh)
 
 	if (!AssetData->IsEmpty())
 	{
-		const UDNAAsset* DNAAsset = Cast<UDNAAsset>(&AssetData[0]);
+		const UDNAAsset* DNAAsset = Cast<UDNAAsset>((*AssetData)[0]);
 		if (DNAAsset)
 		{
 			const TArray<FString> Filenames = { DNAAsset->DnaFileName };
