Allow reduction of compiled binary size of RigLogic by introducing additional feature flags to disable parts of the implementation which are not necessary for runtime

diff --git a/Engine/Plugins/Animation/RigLogic/RigLogic.uplugin b/Engine/Plugins/Animation/RigLogic/RigLogic.uplugin
--- a/Engine/Plugins/Animation/RigLogic/RigLogic.uplugin
+++ b/Engine/Plugins/Animation/RigLogic/RigLogic.uplugin
@@ -1,9 +1,9 @@
 {
 	"FileVersion": 3,
 	"Version": 9.1,
-	"VersionName": "9.1.2",
-	"FriendlyName": "RigLogic Plugin v9.1.2",
-	"Description": "3Lateral RigLogic Plugin for Facial Animation v9.1.2",
+	"VersionName": "9.1.3",
+	"FriendlyName": "RigLogic Plugin v9.1.3",
+	"Description": "3Lateral RigLogic Plugin for Facial Animation v9.1.3",
 	"Category": "Animation",
 	"CreatedBy": "3Lateral",
 	"CreatedByURL": "https://www.3lateral.com/",
diff --git a/Engine/Plugins/Animation/RigLogic/Source/RigLogicLib/Private/dna/stream/BinaryStreamWriterImpl.cpp b/Engine/Plugins/Animation/RigLogic/Source/RigLogicLib/Private/dna/stream/BinaryStreamWriterImpl.cpp
--- a/Engine/Plugins/Animation/RigLogic/Source/RigLogicLib/Private/dna/stream/BinaryStreamWriterImpl.cpp
+++ b/Engine/Plugins/Animation/RigLogic/Source/RigLogicLib/Private/dna/stream/BinaryStreamWriterImpl.cpp
@@ -57,12 +57,15 @@ void BinaryStreamWriterImpl::setFrom(const BinaryStreamReader* source,
     reader->rawCopyInto(dna, layer, policy, memRes_);
 }
 
-void BinaryStreamWriterImpl::setFrom(const JSONStreamReader* source,
-                                     DataLayer layer,
-                                     UnknownLayerPolicy policy,
-                                     MemoryResource* memRes_) {
-    auto reader = const_cast<JSONStreamReaderImpl*>(static_cast<const JSONStreamReaderImpl*>(source));
-    reader->rawCopyInto(dna, layer, policy, memRes_);
-}
+#ifdef DNA_BUILD_WITH_JSON_SUPPORT
+    void BinaryStreamWriterImpl::setFrom(const JSONStreamReader* source,
+                                         DataLayer layer,
+                                         UnknownLayerPolicy policy,
+                                         MemoryResource* memRes_) {
+        auto reader = const_cast<JSONStreamReaderImpl*>(static_cast<const JSONStreamReaderImpl*>(source));
+        reader->rawCopyInto(dna, layer, policy, memRes_);
+    }
+
+#endif  // DNA_BUILD_WITH_JSON_SUPPORT
 
 }  // namespace dna
diff --git a/Engine/Plugins/Animation/RigLogic/Source/RigLogicLib/Private/dna/stream/BinaryStreamWriterImpl.h b/Engine/Plugins/Animation/RigLogic/Source/RigLogicLib/Private/dna/stream/BinaryStreamWriterImpl.h
--- a/Engine/Plugins/Animation/RigLogic/Source/RigLogicLib/Private/dna/stream/BinaryStreamWriterImpl.h
+++ b/Engine/Plugins/Animation/RigLogic/Source/RigLogicLib/Private/dna/stream/BinaryStreamWriterImpl.h
@@ -22,8 +22,10 @@ class BinaryStreamWriterImpl : public WriterImpl<BinaryStreamWriter> {
         using BinaryStreamWriter::setFrom;
         void setFrom(const BinaryStreamReader* source, DataLayer layer, UnknownLayerPolicy policy,
                      MemoryResource* memRes_) override;
-        void setFrom(const JSONStreamReader* source, DataLayer layer, UnknownLayerPolicy policy,
-                     MemoryResource* memRes_) override;
+        #ifdef DNA_BUILD_WITH_JSON_SUPPORT
+            void setFrom(const JSONStreamReader* source, DataLayer layer, UnknownLayerPolicy policy,
+                         MemoryResource* memRes_) override;
+        #endif  // DNA_BUILD_WITH_JSON_SUPPORT
 
     private:
         BoundedIOStream* stream;
diff --git a/Engine/Plugins/Animation/RigLogic/Source/RigLogicLib/Private/dna/stream/JSONStreamReaderImpl.cpp b/Engine/Plugins/Animation/RigLogic/Source/RigLogicLib/Private/dna/stream/JSONStreamReaderImpl.cpp
--- a/Engine/Plugins/Animation/RigLogic/Source/RigLogicLib/Private/dna/stream/JSONStreamReaderImpl.cpp
+++ b/Engine/Plugins/Animation/RigLogic/Source/RigLogicLib/Private/dna/stream/JSONStreamReaderImpl.cpp
@@ -1,5 +1,7 @@
 // Copyright Epic Games, Inc. All Rights Reserved.
 
+#ifdef DNA_BUILD_WITH_JSON_SUPPORT
+
 #include "dna/stream/JSONStreamReaderImpl.h"
 
 #include "dna/TypeDefs.h"
@@ -95,3 +97,5 @@ void JSONStreamReaderImpl::read() {
 }
 
 }  // namespace dna
+
+#endif  // DNA_BUILD_WITH_JSON_SUPPORT
diff --git a/Engine/Plugins/Animation/RigLogic/Source/RigLogicLib/Private/dna/stream/JSONStreamReaderImpl.h b/Engine/Plugins/Animation/RigLogic/Source/RigLogicLib/Private/dna/stream/JSONStreamReaderImpl.h
--- a/Engine/Plugins/Animation/RigLogic/Source/RigLogicLib/Private/dna/stream/JSONStreamReaderImpl.h
+++ b/Engine/Plugins/Animation/RigLogic/Source/RigLogicLib/Private/dna/stream/JSONStreamReaderImpl.h
@@ -2,6 +2,9 @@
 
 #pragma once
 
+// *INDENT-OFF*
+#ifdef DNA_BUILD_WITH_JSON_SUPPORT
+
 #include "dna/JSONStreamReader.h"
 #include "dna/ReaderImpl.h"
 #include "dna/TypeDefs.h"
@@ -26,3 +29,6 @@ class JSONStreamReaderImpl : public ReaderImpl<JSONStreamReader> {
 };
 
 }  // namespace dna
+
+#endif  // DNA_BUILD_WITH_JSON_SUPPORT
+// *INDENT-ON*
diff --git a/Engine/Plugins/Animation/RigLogic/Source/RigLogicLib/Private/dna/stream/JSONStreamWriterImpl.cpp b/Engine/Plugins/Animation/RigLogic/Source/RigLogicLib/Private/dna/stream/JSONStreamWriterImpl.cpp
--- a/Engine/Plugins/Animation/RigLogic/Source/RigLogicLib/Private/dna/stream/JSONStreamWriterImpl.cpp
+++ b/Engine/Plugins/Animation/RigLogic/Source/RigLogicLib/Private/dna/stream/JSONStreamWriterImpl.cpp
@@ -1,5 +1,7 @@
 // Copyright Epic Games, Inc. All Rights Reserved.
 
+#ifdef DNA_BUILD_WITH_JSON_SUPPORT
+
 #include "dna/stream/JSONStreamWriterImpl.h"
 
 #include "dna/TypeDefs.h"
@@ -66,3 +68,5 @@ void JSONStreamWriterImpl::setFrom(const JSONStreamReader* source,
 }
 
 }  // namespace dna
+
+#endif  // DNA_BUILD_WITH_JSON_SUPPORT
diff --git a/Engine/Plugins/Animation/RigLogic/Source/RigLogicLib/Private/dna/stream/JSONStreamWriterImpl.h b/Engine/Plugins/Animation/RigLogic/Source/RigLogicLib/Private/dna/stream/JSONStreamWriterImpl.h
--- a/Engine/Plugins/Animation/RigLogic/Source/RigLogicLib/Private/dna/stream/JSONStreamWriterImpl.h
+++ b/Engine/Plugins/Animation/RigLogic/Source/RigLogicLib/Private/dna/stream/JSONStreamWriterImpl.h
@@ -2,6 +2,9 @@
 
 #pragma once
 
+// *INDENT-OFF*
+#ifdef DNA_BUILD_WITH_JSON_SUPPORT
+
 #include "dna/DNA.h"
 #include "dna/JSONStreamWriter.h"
 #include "dna/WriterImpl.h"
@@ -20,9 +23,13 @@ class JSONStreamWriterImpl : public WriterImpl<JSONStreamWriter> {
         void write() override;
 
         using JSONStreamWriter::setFrom;
-        void setFrom(const BinaryStreamReader* source, DataLayer layer, UnknownLayerPolicy policy,
+        void setFrom(const BinaryStreamReader* source,
+                     DataLayer layer,
+                     UnknownLayerPolicy policy,
                      MemoryResource* memRes_) override;
-        void setFrom(const JSONStreamReader* source, DataLayer layer, UnknownLayerPolicy policy,
+        void setFrom(const JSONStreamReader* source,
+                     DataLayer layer,
+                     UnknownLayerPolicy policy,
                      MemoryResource* memRes_) override;
 
     private:
@@ -32,3 +39,6 @@ class JSONStreamWriterImpl : public WriterImpl<JSONStreamWriter> {
 };
 
 }  // namespace dna
+
+#endif  // DNA_BUILD_WITH_JSON_SUPPORT
+// *INDENT-ON*
diff --git a/Engine/Plugins/Animation/RigLogic/Source/RigLogicLib/Private/pma/resources/AlignedMemoryResource.cpp b/Engine/Plugins/Animation/RigLogic/Source/RigLogicLib/Private/pma/resources/AlignedMemoryResource.cpp
--- a/Engine/Plugins/Animation/RigLogic/Source/RigLogicLib/Private/pma/resources/AlignedMemoryResource.cpp
+++ b/Engine/Plugins/Animation/RigLogic/Source/RigLogicLib/Private/pma/resources/AlignedMemoryResource.cpp
@@ -55,7 +55,7 @@
         // offset to meet the alignment requirement)
         std::ptrdiff_t shift = pAlignedChunk - pUnalignedChunk;
         // The guaranteed storage size for the shift value is only one byte
-        assert(shift > 0ul && shift <= 256ul);
+        assert(shift > 0 && shift <= 256);
         pAlignedChunk[-1] = static_cast<char>(shift & 0xFF);
         return static_cast<void*>(pAlignedChunk);
     }
diff --git a/Engine/Plugins/Animation/RigLogic/Source/RigLogicLib/Private/riglogic/ml/MachineLearnedBehavior.h b/Engine/Plugins/Animation/RigLogic/Source/RigLogicLib/Private/riglogic/ml/MachineLearnedBehavior.h
--- a/Engine/Plugins/Animation/RigLogic/Source/RigLogicLib/Private/riglogic/ml/MachineLearnedBehavior.h
+++ b/Engine/Plugins/Animation/RigLogic/Source/RigLogicLib/Private/riglogic/ml/MachineLearnedBehavior.h
@@ -4,6 +4,7 @@
 
 #include "riglogic/TypeDefs.h"
 #include "riglogic/types/Aliases.h"
+#include "riglogic/utils/Macros.h"
 #include "riglogic/ml/MachineLearnedBehaviorEvaluator.h"
 #include "riglogic/ml/MachineLearnedBehaviorOutputInstance.h"
 
@@ -32,14 +33,22 @@ class MachineLearnedBehavior {
 
         template<class Archive>
         void load(Archive& archive) {
-            evaluator->load(archive);
-            archive >> neuralNetworkIndicesPerMeshRegion;
+            #ifdef RL_BUILD_WITH_ML_EVALUATOR
+                evaluator->load(archive);
+                archive >> neuralNetworkIndicesPerMeshRegion;
+            #else
+                UNUSED(archive);
+            #endif  // RL_BUILD_WITH_ML_EVALUATOR
         }
 
         template<class Archive>
         void save(Archive& archive) {
-            evaluator->save(archive);
-            archive << neuralNetworkIndicesPerMeshRegion;
+            #ifdef RL_BUILD_WITH_ML_EVALUATOR
+                evaluator->save(archive);
+                archive << neuralNetworkIndicesPerMeshRegion;
+            #else
+                UNUSED(archive);
+            #endif  // RL_BUILD_WITH_ML_EVALUATOR
         }
 
         std::uint16_t getMeshCount() const;
diff --git a/Engine/Plugins/Animation/RigLogic/Source/RigLogicLib/Private/riglogic/ml/MachineLearnedBehaviorFactory.cpp b/Engine/Plugins/Animation/RigLogic/Source/RigLogicLib/Private/riglogic/ml/MachineLearnedBehaviorFactory.cpp
--- a/Engine/Plugins/Animation/RigLogic/Source/RigLogicLib/Private/riglogic/ml/MachineLearnedBehaviorFactory.cpp
+++ b/Engine/Plugins/Animation/RigLogic/Source/RigLogicLib/Private/riglogic/ml/MachineLearnedBehaviorFactory.cpp
@@ -6,16 +6,20 @@
 #include "riglogic/utils/Extd.h"
 #include "riglogic/ml/MachineLearnedBehaviorEvaluator.h"
 #include "riglogic/ml/MachineLearnedBehaviorNullEvaluator.h"
-#include "riglogic/ml/cpu/Factory.h"
+#ifdef RL_BUILD_WITH_ML_EVALUATOR
+    #include "riglogic/ml/cpu/Factory.h"
+#endif  // RL_BUILD_WITH_ML_EVALUATOR
 #include "riglogic/riglogic/Configuration.h"
 #include "riglogic/riglogic/RigMetrics.h"
-#include "riglogic/system/simd/Detect.h"
 #include "riglogic/system/simd/SIMD.h"
 
 #include <cstdint>
 
 namespace rl4 {
 
+// *INDENT-OFF*
+#ifdef RL_BUILD_WITH_ML_EVALUATOR
+
 #ifdef RL_USE_HALF_FLOATS
     using StorageValueType = std::uint16_t;
 #else
@@ -35,7 +39,7 @@ MachineLearnedBehaviorEvaluator::Pointer createEvaluator(const Configuration& co
     #ifdef RL_BUILD_WITH_AVX
         if (config.calculationType == CalculationType::AVX) {
             // Use 256-bit AVX registers and whatever 128-bit width type is available
-            return ml::cpu::Factory<StorageValueType, trimd::avx::F256, trimd::F128>::create(reader, memRes);
+            return ml::cpu::Factory<StorageValueType, trimd::avx::F256, trimd::sse::F128>::create(reader, memRes);
         }
     #endif  // RL_BUILD_WITH_AVX
     return ml::cpu::Factory<float, trimd::scalar::F256, trimd::scalar::F128>::create(reader, memRes);
@@ -74,4 +78,21 @@ MachineLearnedBehavior::Pointer MachineLearnedBehaviorFactory::create(const Conf
     return moduleFactory.create(createEvaluator(config, nullptr, memRes), memRes);
 }
 
+#else
+
+MachineLearnedBehavior::Pointer MachineLearnedBehaviorFactory::create(const Configuration&  /*unused*/,
+                                                                      const dna::MachineLearnedBehaviorReader*  /*unused*/,
+                                                                      MemoryResource* memRes) {
+    auto moduleFactory = UniqueInstance<MachineLearnedBehavior>::with(memRes);
+    auto evaluator = UniqueInstance<MachineLearnedBehaviorNullEvaluator, MachineLearnedBehaviorEvaluator>::with(memRes).create();
+    return moduleFactory.create(std::move(evaluator), memRes);
+}
+
+MachineLearnedBehavior::Pointer MachineLearnedBehaviorFactory::create(const Configuration& config, MemoryResource* memRes) {
+    return create(config, nullptr, memRes);
+}
+
+#endif  // RL_BUILD_WITH_ML_EVALUATOR
+// *INDENT-ON*
+
 }  // namespace rl4
diff --git a/Engine/Plugins/Animation/RigLogic/Source/RigLogicLib/Private/riglogic/ml/cpu/Block4.h b/Engine/Plugins/Animation/RigLogic/Source/RigLogicLib/Private/riglogic/ml/cpu/Block4.h
--- a/Engine/Plugins/Animation/RigLogic/Source/RigLogicLib/Private/riglogic/ml/cpu/Block4.h
+++ b/Engine/Plugins/Animation/RigLogic/Source/RigLogicLib/Private/riglogic/ml/cpu/Block4.h
@@ -2,6 +2,9 @@
 
 #pragma once
 
+// *INDENT-OFF*
+#ifdef RL_BUILD_WITH_ML_EVALUATOR
+
 #include "riglogic/TypeDefs.h"
 #include "riglogic/ml/cpu/Consts.h"
 #include "riglogic/utils/Macros.h"
@@ -214,3 +217,6 @@ static FORCE_INLINE void calculateBlock4(const NeuralNetLayer<T>& layer, ConstAr
 }  // namespace ml
 
 }  // namespace cpu
+
+#endif  // RL_BUILD_WITH_ML_EVALUATOR
+// *INDENT-ON*
diff --git a/Engine/Plugins/Animation/RigLogic/Source/RigLogicLib/Private/riglogic/ml/cpu/Consts.h b/Engine/Plugins/Animation/RigLogic/Source/RigLogicLib/Private/riglogic/ml/cpu/Consts.h
--- a/Engine/Plugins/Animation/RigLogic/Source/RigLogicLib/Private/riglogic/ml/cpu/Consts.h
+++ b/Engine/Plugins/Animation/RigLogic/Source/RigLogicLib/Private/riglogic/ml/cpu/Consts.h
@@ -2,6 +2,9 @@
 
 #pragma once
 
+// *INDENT-OFF*
+#ifdef RL_BUILD_WITH_ML_EVALUATOR
+
 #include <cstdint>
 
 namespace rl4 {
@@ -24,3 +27,6 @@ constexpr std::uint32_t block16Height = 16u;
 }  // namespace ml
 
 }  // namespace rl4
+
+#endif  // RL_BUILD_WITH_ML_EVALUATOR
+// *INDENT-ON*
diff --git a/Engine/Plugins/Animation/RigLogic/Source/RigLogicLib/Private/riglogic/ml/cpu/Evaluator.h b/Engine/Plugins/Animation/RigLogic/Source/RigLogicLib/Private/riglogic/ml/cpu/Evaluator.h
--- a/Engine/Plugins/Animation/RigLogic/Source/RigLogicLib/Private/riglogic/ml/cpu/Evaluator.h
+++ b/Engine/Plugins/Animation/RigLogic/Source/RigLogicLib/Private/riglogic/ml/cpu/Evaluator.h
@@ -2,6 +2,9 @@
 
 #pragma once
 
+// *INDENT-OFF*
+#ifdef RL_BUILD_WITH_ML_EVALUATOR
+
 #include "riglogic/TypeDefs.h"
 #include "riglogic/controls/ControlsInputInstance.h"
 #include "riglogic/ml/LODSpec.h"
@@ -95,3 +98,6 @@ class Evaluator : public MachineLearnedBehaviorEvaluator {
 }  // namespace ml
 
 }  // namespace rl4
+
+#endif  // RL_BUILD_WITH_ML_EVALUATOR
+// *INDENT-ON*
diff --git a/Engine/Plugins/Animation/RigLogic/Source/RigLogicLib/Private/riglogic/ml/cpu/Factory.h b/Engine/Plugins/Animation/RigLogic/Source/RigLogicLib/Private/riglogic/ml/cpu/Factory.h
--- a/Engine/Plugins/Animation/RigLogic/Source/RigLogicLib/Private/riglogic/ml/cpu/Factory.h
+++ b/Engine/Plugins/Animation/RigLogic/Source/RigLogicLib/Private/riglogic/ml/cpu/Factory.h
@@ -2,6 +2,9 @@
 
 #pragma once
 
+// *INDENT-OFF*
+#ifdef RL_BUILD_WITH_ML_EVALUATOR
+
 #include "riglogic/TypeDefs.h"
 #include "riglogic/ml/LODSpec.h"
 #include "riglogic/ml/MachineLearnedBehaviorEvaluator.h"
@@ -18,14 +21,6 @@ namespace ml {
 
 namespace cpu {
 
-static OutputInstance::Factory createInstanceFactory() {
-    return [](ConstArrayView<std::uint32_t> maxLayerOutputCounts, MemoryResource* memRes) {
-               return UniqueInstance<OutputInstance, MachineLearnedBehaviorOutputInstance>::with(memRes).create(
-                   maxLayerOutputCounts,
-                   memRes);
-    };
-}
-
 template<typename T, typename TF256, typename TF128>
 class Factory {
     public:
@@ -33,7 +28,10 @@ class Factory {
                                                                MemoryResource* memRes) {
             Vector<NeuralNetInference<T, TF256, TF128> > neuralNets{memRes};
             Vector<std::uint32_t> maxLayerOutputCountPerNet{memRes};
-            auto instanceFactory = createInstanceFactory();
+            auto instanceFactory = [](ConstArrayView<std::uint32_t> maxLayerOutputCounts, MemoryResource* instanceMemRes) {
+                    using OutputInstancePointer = UniqueInstance<OutputInstance, MachineLearnedBehaviorOutputInstance>;
+                    return OutputInstancePointer::with(instanceMemRes).create(maxLayerOutputCounts, instanceMemRes);
+                };
             auto factory = UniqueInstance<Evaluator<T, TF256, TF128>, MachineLearnedBehaviorEvaluator>::with(memRes);
 
             if (reader == nullptr) {
@@ -156,3 +154,6 @@ class Factory {
 }  // namespace ml
 
 }  // namespace rl4
+
+#endif  // RL_BUILD_WITH_ML_EVALUATOR
+// *INDENT-ON*
diff --git a/Engine/Plugins/Animation/RigLogic/Source/RigLogicLib/Private/riglogic/ml/cpu/Inference.h b/Engine/Plugins/Animation/RigLogic/Source/RigLogicLib/Private/riglogic/ml/cpu/Inference.h
--- a/Engine/Plugins/Animation/RigLogic/Source/RigLogicLib/Private/riglogic/ml/cpu/Inference.h
+++ b/Engine/Plugins/Animation/RigLogic/Source/RigLogicLib/Private/riglogic/ml/cpu/Inference.h
@@ -2,6 +2,9 @@
 
 #pragma once
 
+// *INDENT-OFF*
+#ifdef RL_BUILD_WITH_ML_EVALUATOR
+
 #include "riglogic/TypeDefs.h"
 #include "riglogic/ml/cpu/NeuralNet.h"
 #include "riglogic/ml/cpu/layers/LeakyReLULayerEvaluator.h"
@@ -108,3 +111,6 @@ struct NeuralNetInference {
 }  // namespace ml
 
 }  // namespace rl4
+
+#endif  // RL_BUILD_WITH_ML_EVALUATOR
+// *INDENT-ON*
diff --git a/Engine/Plugins/Animation/RigLogic/Source/RigLogicLib/Private/riglogic/ml/cpu/NeuralNet.h b/Engine/Plugins/Animation/RigLogic/Source/RigLogicLib/Private/riglogic/ml/cpu/NeuralNet.h
--- a/Engine/Plugins/Animation/RigLogic/Source/RigLogicLib/Private/riglogic/ml/cpu/NeuralNet.h
+++ b/Engine/Plugins/Animation/RigLogic/Source/RigLogicLib/Private/riglogic/ml/cpu/NeuralNet.h
@@ -2,6 +2,9 @@
 
 #pragma once
 
+// *INDENT-OFF*
+#ifdef RL_BUILD_WITH_ML_EVALUATOR
+
 #include "riglogic/TypeDefs.h"
 #include "riglogic/types/AlignedBlockView.h"
 #include "riglogic/types/Extent.h"
@@ -86,3 +89,6 @@ struct NeuralNet {
 }  // namespace ml
 
 }  // namespace rl4
+
+#endif  // RL_BUILD_WITH_ML_EVALUATOR
+// *INDENT-ON*
diff --git a/Engine/Plugins/Animation/RigLogic/Source/RigLogicLib/Private/riglogic/ml/cpu/OutputInstance.cpp b/Engine/Plugins/Animation/RigLogic/Source/RigLogicLib/Private/riglogic/ml/cpu/OutputInstance.cpp
--- a/Engine/Plugins/Animation/RigLogic/Source/RigLogicLib/Private/riglogic/ml/cpu/OutputInstance.cpp
+++ b/Engine/Plugins/Animation/RigLogic/Source/RigLogicLib/Private/riglogic/ml/cpu/OutputInstance.cpp
@@ -1,5 +1,8 @@
 // Copyright Epic Games, Inc. All Rights Reserved.
 
+// *INDENT-OFF*
+#ifdef RL_BUILD_WITH_ML_EVALUATOR
+
 #include "riglogic/ml/cpu/OutputInstance.h"
 
 #include <cassert>
@@ -47,3 +50,6 @@ ConstArrayView<float> OutputInstance::getMaskBuffer() const {
 }  // namespace ml
 
 }  // namespace rl4
+
+#endif  // RL_BUILD_WITH_ML_EVALUATOR
+// *INDENT-ON*
diff --git a/Engine/Plugins/Animation/RigLogic/Source/RigLogicLib/Private/riglogic/ml/cpu/OutputInstance.h b/Engine/Plugins/Animation/RigLogic/Source/RigLogicLib/Private/riglogic/ml/cpu/OutputInstance.h
--- a/Engine/Plugins/Animation/RigLogic/Source/RigLogicLib/Private/riglogic/ml/cpu/OutputInstance.h
+++ b/Engine/Plugins/Animation/RigLogic/Source/RigLogicLib/Private/riglogic/ml/cpu/OutputInstance.h
@@ -2,6 +2,9 @@
 
 #pragma once
 
+// *INDENT-OFF*
+#ifdef RL_BUILD_WITH_ML_EVALUATOR
+
 #include "riglogic/TypeDefs.h"
 #include "riglogic/ml/MachineLearnedBehaviorOutputInstance.h"
 
@@ -43,3 +46,6 @@ class OutputInstance : public MachineLearnedBehaviorOutputInstance {
 }  // namespace ml
 
 }  // namespace rl4
+
+#endif  // RL_BUILD_WITH_ML_EVALUATOR
+// *INDENT-ON*
diff --git a/Engine/Plugins/Animation/RigLogic/Source/RigLogicLib/Private/riglogic/ml/cpu/layers/LayerEvaluator.h b/Engine/Plugins/Animation/RigLogic/Source/RigLogicLib/Private/riglogic/ml/cpu/layers/LayerEvaluator.h
--- a/Engine/Plugins/Animation/RigLogic/Source/RigLogicLib/Private/riglogic/ml/cpu/layers/LayerEvaluator.h
+++ b/Engine/Plugins/Animation/RigLogic/Source/RigLogicLib/Private/riglogic/ml/cpu/layers/LayerEvaluator.h
@@ -2,6 +2,9 @@
 
 #pragma once
 
+// *INDENT-OFF*
+#ifdef RL_BUILD_WITH_ML_EVALUATOR
+
 #include "riglogic/TypeDefs.h"
 #include "riglogic/ml/cpu/NeuralNet.h"
 
@@ -27,3 +30,6 @@ class LayerEvaluator {
 }  // namespace ml
 
 }  // namespace rl4
+
+#endif  // RL_BUILD_WITH_ML_EVALUATOR
+// *INDENT-ON*
diff --git a/Engine/Plugins/Animation/RigLogic/Source/RigLogicLib/Private/riglogic/ml/cpu/layers/LeakyReLULayerEvaluator.h b/Engine/Plugins/Animation/RigLogic/Source/RigLogicLib/Private/riglogic/ml/cpu/layers/LeakyReLULayerEvaluator.h
--- a/Engine/Plugins/Animation/RigLogic/Source/RigLogicLib/Private/riglogic/ml/cpu/layers/LeakyReLULayerEvaluator.h
+++ b/Engine/Plugins/Animation/RigLogic/Source/RigLogicLib/Private/riglogic/ml/cpu/layers/LeakyReLULayerEvaluator.h
@@ -2,6 +2,9 @@
 
 #pragma once
 
+// *INDENT-OFF*
+#ifdef RL_BUILD_WITH_ML_EVALUATOR
+
 #include "riglogic/TypeDefs.h"
 #include "riglogic/ml/cpu/Block4.h"
 #include "riglogic/ml/cpu/Consts.h"
@@ -42,3 +45,6 @@ class LeakyReLULayerEvaluator : public LayerEvaluator<T> {
 }  // namespace ml
 
 }  // namespace rl4
+
+#endif  // RL_BUILD_WITH_ML_EVALUATOR
+// *INDENT-ON*
diff --git a/Engine/Plugins/Animation/RigLogic/Source/RigLogicLib/Private/riglogic/ml/cpu/layers/LinearLayerEvaluator.h b/Engine/Plugins/Animation/RigLogic/Source/RigLogicLib/Private/riglogic/ml/cpu/layers/LinearLayerEvaluator.h
--- a/Engine/Plugins/Animation/RigLogic/Source/RigLogicLib/Private/riglogic/ml/cpu/layers/LinearLayerEvaluator.h
+++ b/Engine/Plugins/Animation/RigLogic/Source/RigLogicLib/Private/riglogic/ml/cpu/layers/LinearLayerEvaluator.h
@@ -2,6 +2,9 @@
 
 #pragma once
 
+// *INDENT-OFF*
+#ifdef RL_BUILD_WITH_ML_EVALUATOR
+
 #include "riglogic/TypeDefs.h"
 #include "riglogic/ml/cpu/Block4.h"
 #include "riglogic/ml/cpu/Consts.h"
@@ -36,3 +39,6 @@ class LinearLayerEvaluator : public LayerEvaluator<T> {
 }  // namespace ml
 
 }  // namespace rl4
+
+#endif  // RL_BUILD_WITH_ML_EVALUATOR
+// *INDENT-ON*
diff --git a/Engine/Plugins/Animation/RigLogic/Source/RigLogicLib/Private/riglogic/ml/cpu/layers/ReLULayerEvaluator.h b/Engine/Plugins/Animation/RigLogic/Source/RigLogicLib/Private/riglogic/ml/cpu/layers/ReLULayerEvaluator.h
--- a/Engine/Plugins/Animation/RigLogic/Source/RigLogicLib/Private/riglogic/ml/cpu/layers/ReLULayerEvaluator.h
+++ b/Engine/Plugins/Animation/RigLogic/Source/RigLogicLib/Private/riglogic/ml/cpu/layers/ReLULayerEvaluator.h
@@ -2,6 +2,9 @@
 
 #pragma once
 
+// *INDENT-OFF*
+#ifdef RL_BUILD_WITH_ML_EVALUATOR
+
 #include "riglogic/TypeDefs.h"
 #include "riglogic/ml/cpu/Block4.h"
 #include "riglogic/ml/cpu/Consts.h"
@@ -40,3 +43,6 @@ class ReLULayerEvaluator : public LayerEvaluator<T> {
 }  // namespace ml
 
 }  // namespace rl4
+
+#endif  // RL_BUILD_WITH_ML_EVALUATOR
+// *INDENT-ON*
diff --git a/Engine/Plugins/Animation/RigLogic/Source/RigLogicLib/Private/riglogic/ml/cpu/layers/SigmoidLayerEvaluator.h b/Engine/Plugins/Animation/RigLogic/Source/RigLogicLib/Private/riglogic/ml/cpu/layers/SigmoidLayerEvaluator.h
--- a/Engine/Plugins/Animation/RigLogic/Source/RigLogicLib/Private/riglogic/ml/cpu/layers/SigmoidLayerEvaluator.h
+++ b/Engine/Plugins/Animation/RigLogic/Source/RigLogicLib/Private/riglogic/ml/cpu/layers/SigmoidLayerEvaluator.h
@@ -2,6 +2,9 @@
 
 #pragma once
 
+// *INDENT-OFF*
+#ifdef RL_BUILD_WITH_ML_EVALUATOR
+
 #include "riglogic/TypeDefs.h"
 #include "riglogic/ml/cpu/Block4.h"
 #include "riglogic/ml/cpu/Consts.h"
@@ -77,3 +80,6 @@ class SigmoidLayerEvaluator : public LayerEvaluator<T> {
 }  // namespace ml
 
 }  // namespace rl4
+
+#endif  // RL_BUILD_WITH_ML_EVALUATOR
+// *INDENT-ON*
diff --git a/Engine/Plugins/Animation/RigLogic/Source/RigLogicLib/Private/riglogic/ml/cpu/layers/TanHLayerEvaluator.h b/Engine/Plugins/Animation/RigLogic/Source/RigLogicLib/Private/riglogic/ml/cpu/layers/TanHLayerEvaluator.h
--- a/Engine/Plugins/Animation/RigLogic/Source/RigLogicLib/Private/riglogic/ml/cpu/layers/TanHLayerEvaluator.h
+++ b/Engine/Plugins/Animation/RigLogic/Source/RigLogicLib/Private/riglogic/ml/cpu/layers/TanHLayerEvaluator.h
@@ -2,6 +2,9 @@
 
 #pragma once
 
+// *INDENT-OFF*
+#ifdef RL_BUILD_WITH_ML_EVALUATOR
+
 #include "riglogic/TypeDefs.h"
 #include "riglogic/ml/cpu/Block4.h"
 #include "riglogic/ml/cpu/Consts.h"
@@ -78,3 +81,6 @@ class TanHLayerEvaluator : public LayerEvaluator<T> {
 }  // namespace ml
 
 }  // namespace rl4
+
+#endif  // RL_BUILD_WITH_ML_EVALUATOR
+// *INDENT-ON*
diff --git a/Engine/Plugins/Animation/RigLogic/Source/RigLogicLib/Private/riglogic/ml/cpu/layers/Utils.h b/Engine/Plugins/Animation/RigLogic/Source/RigLogicLib/Private/riglogic/ml/cpu/layers/Utils.h
--- a/Engine/Plugins/Animation/RigLogic/Source/RigLogicLib/Private/riglogic/ml/cpu/layers/Utils.h
+++ b/Engine/Plugins/Animation/RigLogic/Source/RigLogicLib/Private/riglogic/ml/cpu/layers/Utils.h
@@ -2,6 +2,9 @@
 
 #pragma once
 
+// *INDENT-OFF*
+#ifdef RL_BUILD_WITH_ML_EVALUATOR
+
 #include <type_traits>
 
 namespace rl4 {
@@ -20,3 +23,6 @@ struct HasSize {
 }  // namespace ml
 
 }  // namespace rl4
+
+#endif  // RL_BUILD_WITH_ML_EVALUATOR
+// *INDENT-ON*
diff --git a/Engine/Plugins/Animation/RigLogic/Source/RigLogicLib/Private/riglogic/riglogic/RigLogicImpl.cpp b/Engine/Plugins/Animation/RigLogic/Source/RigLogicLib/Private/riglogic/riglogic/RigLogicImpl.cpp
--- a/Engine/Plugins/Animation/RigLogic/Source/RigLogicLib/Private/riglogic/riglogic/RigLogicImpl.cpp
+++ b/Engine/Plugins/Animation/RigLogic/Source/RigLogicLib/Private/riglogic/riglogic/RigLogicImpl.cpp
@@ -38,11 +38,13 @@ static RigMetrics::Pointer computeRigMetrics(const dna::Reader* reader, MemoryRe
     metrics->guiControlCount = reader->getGUIControlCount();
     metrics->rawControlCount = reader->getRawControlCount();
     metrics->psdControlCount = reader->getPSDCount();
-    metrics->mlControlCount = reader->getMLControlCount();
     metrics->jointAttributeCount = reader->getJointRowCount();
     metrics->blendShapeCount = reader->getBlendShapeChannelCount();
     metrics->animatedMapCount = reader->getAnimatedMapCount();
-    metrics->neuralNetworkCount = reader->getNeuralNetworkCount();
+    #ifdef RL_BUILD_WITH_ML_EVALUATOR
+        metrics->mlControlCount = reader->getMLControlCount();
+        metrics->neuralNetworkCount = reader->getNeuralNetworkCount();
+    #endif  // RL_BUILD_WITH_ML_EVALUATOR
     return metrics;
 }
 
diff --git a/Engine/Plugins/Animation/RigLogic/Source/RigLogicLib/Private/riglogic/system/simd/Detect.h b/Engine/Plugins/Animation/RigLogic/Source/RigLogicLib/Private/riglogic/system/simd/Detect.h
--- a/Engine/Plugins/Animation/RigLogic/Source/RigLogicLib/Private/riglogic/system/simd/Detect.h
+++ b/Engine/Plugins/Animation/RigLogic/Source/RigLogicLib/Private/riglogic/system/simd/Detect.h
@@ -4,7 +4,9 @@
 
 #if defined(RL_AUTODETECT_SSE) && !defined(RL_BUILD_WITH_SSE)
     // Assume any x86 architecture we try to build on will support SSE
-    #if defined(i386) || defined(__i386) || defined(__i386__) || defined(__IA32__) || defined(_M_IX86) || defined(__X86__) || defined(_X86_) || defined(__THW_INTEL__) || defined(__I86__) || defined(__amd64__) || defined(__amd64) || defined(__x86_64__) || defined(__x86_64) || defined(_M_X64) || defined(_M_AMD64)
+    #if defined(i386) || defined(__i386) || defined(__i386__) || defined(__IA32__) || defined(_M_IX86) || defined(__X86__) || \
+    defined(_X86_) || defined(__THW_INTEL__) || defined(__I86__) || defined(__amd64__) || defined(__amd64) || \
+    defined(__x86_64__) || defined(__x86_64) || defined(_M_X64) || defined(_M_AMD64)
         #define RL_BUILD_WITH_SSE 1
     #endif
 #endif  // RL_AUTODETECT_SSE
diff --git a/Engine/Plugins/Animation/RigLogic/Source/RigLogicLib/Private/riglogic/system/simd/SIMD.h b/Engine/Plugins/Animation/RigLogic/Source/RigLogicLib/Private/riglogic/system/simd/SIMD.h
--- a/Engine/Plugins/Animation/RigLogic/Source/RigLogicLib/Private/riglogic/system/simd/SIMD.h
+++ b/Engine/Plugins/Animation/RigLogic/Source/RigLogicLib/Private/riglogic/system/simd/SIMD.h
@@ -8,8 +8,13 @@
     #define TRIMD_ENABLE_F16C
 #endif  // RL_USE_HALF_FLOATS
 
-#if defined(RL_BUILD_WITH_AVX) && !defined(TRIMD_ENABLE_AVX)
-    #define TRIMD_ENABLE_AVX
+#if defined(RL_BUILD_WITH_AVX)
+    #if !defined(TRIMD_ENABLE_AVX)
+        #define TRIMD_ENABLE_AVX
+    #endif
+    #if !defined(TRIMD_ENABLE_SSE)
+        #define TRIMD_ENABLE_SSE
+    #endif
 #endif  // RL_BUILD_WITH_AVX
 
 #if defined(RL_BUILD_WITH_SSE) && !defined(TRIMD_ENABLE_SSE)
diff --git a/Engine/Plugins/Animation/RigLogic/Source/RigLogicLib/Private/riglogic/types/bpcm/Optimizer.h b/Engine/Plugins/Animation/RigLogic/Source/RigLogicLib/Private/riglogic/types/bpcm/Optimizer.h
--- a/Engine/Plugins/Animation/RigLogic/Source/RigLogicLib/Private/riglogic/types/bpcm/Optimizer.h
+++ b/Engine/Plugins/Animation/RigLogic/Source/RigLogicLib/Private/riglogic/types/bpcm/Optimizer.h
@@ -3,6 +3,7 @@
 
 #include "riglogic/types/Extent.h"
 #include "riglogic/utils/Extd.h"
+#include "riglogic/utils/Macros.h"
 #include "riglogic/system/simd/SIMD.h"
 
 #include <cstdint>
@@ -92,6 +93,10 @@ struct Optimizer {
                     #endif
                 }
 
+                // GCC-7 has a bug in its `array-bounds` check, and incorrectly warns about OOB access in the
+                // below loop on `output[blkIdx]`, but `blkIdx` can never be higher than `paddedBlockHeight`,
+                // which in turn is always either below or equal to `BlockHeight`, thus OOB is impossible.
+                ASSUME_TRUE(paddedBlockHeight <= BlockHeight);
                 // Copy even the padded zero values as the representation of half-floats might
                 // differ from the storage's zero-initialized state
                 for (std::uint32_t blkIdx = 0u; blkIdx < paddedBlockHeight; ++blkIdx, ++offset) {
diff --git a/Engine/Plugins/Animation/RigLogic/Source/RigLogicLib/Private/riglogic/utils/Macros.h b/Engine/Plugins/Animation/RigLogic/Source/RigLogicLib/Private/riglogic/utils/Macros.h
--- a/Engine/Plugins/Animation/RigLogic/Source/RigLogicLib/Private/riglogic/utils/Macros.h
+++ b/Engine/Plugins/Animation/RigLogic/Source/RigLogicLib/Private/riglogic/utils/Macros.h
@@ -9,3 +9,15 @@
         #define FORCE_INLINE inline __attribute__((always_inline))
     #endif
 #endif
+
+#if !defined(ASSUME_TRUE)
+    #if defined(_MSC_VER) && !defined(__clang__)
+        #define ASSUME_TRUE(condition) __assume(condition)
+    #else
+        #define ASSUME_TRUE(condition) do {if (!(condition)) __builtin_unreachable();} while (0)
+    #endif
+#endif
+
+#if !defined(UNUSED)
+    #define UNUSED(x) static_cast<void>(x)
+#endif
diff --git a/Engine/Plugins/Animation/RigLogic/Source/RigLogicLib/Public/dna/BinaryStreamWriter.h b/Engine/Plugins/Animation/RigLogic/Source/RigLogicLib/Public/dna/BinaryStreamWriter.h
--- a/Engine/Plugins/Animation/RigLogic/Source/RigLogicLib/Public/dna/BinaryStreamWriter.h
+++ b/Engine/Plugins/Animation/RigLogic/Source/RigLogicLib/Public/dna/BinaryStreamWriter.h
@@ -41,10 +41,12 @@ class DNAAPI BinaryStreamWriter : public StreamWriter {
                              DataLayer layer = DataLayer::All,
                              UnknownLayerPolicy policy = UnknownLayerPolicy::Preserve,
                              MemoryResource* memRes = nullptr) = 0;
-        virtual void setFrom(const JSONStreamReader* source,
-                             DataLayer layer = DataLayer::All,
-                             UnknownLayerPolicy policy = UnknownLayerPolicy::Preserve,
-                             MemoryResource* memRes = nullptr) = 0;
+        #ifdef DNA_BUILD_WITH_JSON_SUPPORT
+            virtual void setFrom(const JSONStreamReader* source,
+                                 DataLayer layer = DataLayer::All,
+                                 UnknownLayerPolicy policy = UnknownLayerPolicy::Preserve,
+                                 MemoryResource* memRes = nullptr) = 0;
+        #endif  // DNA_BUILD_WITH_JSON_SUPPORT
 };
 
 }  // namespace dna
diff --git a/Engine/Plugins/Animation/RigLogic/Source/RigLogicLib/Public/dna/JSONStreamReader.h b/Engine/Plugins/Animation/RigLogic/Source/RigLogicLib/Public/dna/JSONStreamReader.h
--- a/Engine/Plugins/Animation/RigLogic/Source/RigLogicLib/Public/dna/JSONStreamReader.h
+++ b/Engine/Plugins/Animation/RigLogic/Source/RigLogicLib/Public/dna/JSONStreamReader.h
@@ -2,6 +2,9 @@
 
 #pragma once
 
+// *INDENT-OFF*
+#ifdef DNA_BUILD_WITH_JSON_SUPPORT
+
 #include "dna/Defs.h"
 #include "dna/StreamReader.h"
 #include "dna/types/Aliases.h"
@@ -49,3 +52,6 @@ struct DefaultInstanceDestroyer<dna::JSONStreamReader> {
 };
 
 }  // namespace pma
+
+#endif  // DNA_BUILD_WITH_JSON_SUPPORT
+// *INDENT-ON*
diff --git a/Engine/Plugins/Animation/RigLogic/Source/RigLogicLib/Public/dna/JSONStreamWriter.h b/Engine/Plugins/Animation/RigLogic/Source/RigLogicLib/Public/dna/JSONStreamWriter.h
--- a/Engine/Plugins/Animation/RigLogic/Source/RigLogicLib/Public/dna/JSONStreamWriter.h
+++ b/Engine/Plugins/Animation/RigLogic/Source/RigLogicLib/Public/dna/JSONStreamWriter.h
@@ -2,6 +2,9 @@
 
 #pragma once
 
+// *INDENT-OFF*
+#ifdef DNA_BUILD_WITH_JSON_SUPPORT
+
 #include "dna/Defs.h"
 #include "dna/StreamWriter.h"
 #include "dna/types/Aliases.h"
@@ -27,7 +30,8 @@ class DNAAPI JSONStreamWriter : public StreamWriter {
                 User is responsible for releasing the returned pointer by calling destroy.
             @see destroy
         */
-        static JSONStreamWriter* create(BoundedIOStream* stream, std::uint32_t indentWidth = 4u,
+        static JSONStreamWriter* create(BoundedIOStream* stream,
+                                        std::uint32_t indentWidth = 4u,
                                         MemoryResource* memRes = nullptr);
         /**
             @brief Method for freeing a JSONStreamWriter instance.
@@ -65,3 +69,6 @@ struct DefaultInstanceDestroyer<dna::JSONStreamWriter> {
 };
 
 }  // namespace pma
+
+#endif  // DNA_BUILD_WITH_JSON_SUPPORT
+// *INDENT-ON*
diff --git a/Engine/Plugins/Animation/RigLogic/Source/RigLogicLib/Public/dna/version/Version.h b/Engine/Plugins/Animation/RigLogic/Source/RigLogicLib/Public/dna/version/Version.h
--- a/Engine/Plugins/Animation/RigLogic/Source/RigLogicLib/Public/dna/version/Version.h
+++ b/Engine/Plugins/Animation/RigLogic/Source/RigLogicLib/Public/dna/version/Version.h
@@ -4,5 +4,5 @@
 
 #define DNA_MAJOR_VERSION 9
 #define DNA_MINOR_VERSION 1
-#define DNA_PATCH_VERSION 2
-#define DNA_VERSION_STRING "9.1.2"
+#define DNA_PATCH_VERSION 4
+#define DNA_VERSION_STRING "9.1.4"
diff --git a/Engine/Plugins/Animation/RigLogic/Source/RigLogicLib/Public/pma/version/Version.h b/Engine/Plugins/Animation/RigLogic/Source/RigLogicLib/Public/pma/version/Version.h
--- a/Engine/Plugins/Animation/RigLogic/Source/RigLogicLib/Public/pma/version/Version.h
+++ b/Engine/Plugins/Animation/RigLogic/Source/RigLogicLib/Public/pma/version/Version.h
@@ -4,5 +4,5 @@
 
 #define PMA_MAJOR_VERSION 1
 #define PMA_MINOR_VERSION 3
-#define PMA_PATCH_VERSION 6
-#define PMA_VERSION_STRING "1.3.6"
+#define PMA_PATCH_VERSION 7
+#define PMA_VERSION_STRING "1.3.7"
diff --git a/Engine/Plugins/Animation/RigLogic/Source/RigLogicLib/Public/riglogic/types/Aliases.h b/Engine/Plugins/Animation/RigLogic/Source/RigLogicLib/Public/riglogic/types/Aliases.h
--- a/Engine/Plugins/Animation/RigLogic/Source/RigLogicLib/Public/riglogic/types/Aliases.h
+++ b/Engine/Plugins/Animation/RigLogic/Source/RigLogicLib/Public/riglogic/types/Aliases.h
@@ -5,8 +5,6 @@
 #include <dna/DataLayer.h>
 #include <dna/BinaryStreamReader.h>
 #include <dna/BinaryStreamWriter.h>
-#include <dna/JSONStreamReader.h>
-#include <dna/JSONStreamWriter.h>
 #include <dna/types/Aliases.h>
 #include <dna/types/Vector3.h>
 #include <pma/MemoryResource.h>
@@ -30,8 +28,6 @@ using trio::MemoryStream;
 using dna::DataLayer;
 using dna::BinaryStreamReader;
 using dna::BinaryStreamWriter;
-using dna::JSONStreamReader;
-using dna::JSONStreamWriter;
 using dna::StringView;
 using dna::UnknownLayerPolicy;
 using dna::Vector3;
diff --git a/Engine/Plugins/Animation/RigLogic/Source/RigLogicLib/Public/riglogic/version/Version.h b/Engine/Plugins/Animation/RigLogic/Source/RigLogicLib/Public/riglogic/version/Version.h
--- a/Engine/Plugins/Animation/RigLogic/Source/RigLogicLib/Public/riglogic/version/Version.h
+++ b/Engine/Plugins/Animation/RigLogic/Source/RigLogicLib/Public/riglogic/version/Version.h
@@ -4,5 +4,5 @@
 
 #define RL_MAJOR_VERSION 9
 #define RL_MINOR_VERSION 1
-#define RL_PATCH_VERSION 2
-#define RL_VERSION_STRING "9.1.2"
+#define RL_PATCH_VERSION 3
+#define RL_VERSION_STRING "9.1.3"
diff --git a/Engine/Plugins/Animation/RigLogic/Source/RigLogicLib/Public/terse/version/Version.h b/Engine/Plugins/Animation/RigLogic/Source/RigLogicLib/Public/terse/version/Version.h
--- a/Engine/Plugins/Animation/RigLogic/Source/RigLogicLib/Public/terse/version/Version.h
+++ b/Engine/Plugins/Animation/RigLogic/Source/RigLogicLib/Public/terse/version/Version.h
@@ -4,5 +4,5 @@
 
 #define TERSE_MAJOR_VERSION 3
 #define TERSE_MINOR_VERSION 1
-#define TERSE_PATCH_VERSION 5
-#define TERSE_VERSION_STRING "3.1.5"
+#define TERSE_PATCH_VERSION 6
+#define TERSE_VERSION_STRING "3.1.6"
diff --git a/Engine/Plugins/Animation/RigLogic/Source/RigLogicLib/Public/trio/version/Version.h b/Engine/Plugins/Animation/RigLogic/Source/RigLogicLib/Public/trio/version/Version.h
--- a/Engine/Plugins/Animation/RigLogic/Source/RigLogicLib/Public/trio/version/Version.h
+++ b/Engine/Plugins/Animation/RigLogic/Source/RigLogicLib/Public/trio/version/Version.h
@@ -4,5 +4,5 @@
 
 #define TRIO_MAJOR_VERSION 4
 #define TRIO_MINOR_VERSION 0
-#define TRIO_PATCH_VERSION 7
-#define TRIO_VERSION_STRING "4.0.7"
+#define TRIO_PATCH_VERSION 8
+#define TRIO_VERSION_STRING "4.0.8"
diff --git a/Engine/Plugins/Animation/RigLogic/Source/RigLogicLib/RigLogicLib.Build.cs b/Engine/Plugins/Animation/RigLogic/Source/RigLogicLib/RigLogicLib.Build.cs
--- a/Engine/Plugins/Animation/RigLogic/Source/RigLogicLib/RigLogicLib.Build.cs
+++ b/Engine/Plugins/Animation/RigLogic/Source/RigLogicLib/RigLogicLib.Build.cs
@@ -51,6 +51,7 @@ namespace UnrealBuildTool.Rules
 			}
 
 			PrivateDefinitions.Add("RL_AUTODETECT_SSE=1");
+			PublicDefinitions.Add("RL_BUILD_WITH_ML_EVALUATOR=1");
 		}
 	}
 }
diff --git a/Engine/Plugins/Animation/RigLogic/Source/RigLogicLibTest/Private/dnatests/FixturesJSON.cpp b/Engine/Plugins/Animation/RigLogic/Source/RigLogicLibTest/Private/dnatests/FixturesJSON.cpp
--- a/Engine/Plugins/Animation/RigLogic/Source/RigLogicLibTest/Private/dnatests/FixturesJSON.cpp
+++ b/Engine/Plugins/Animation/RigLogic/Source/RigLogicLibTest/Private/dnatests/FixturesJSON.cpp
@@ -1,5 +1,7 @@
 // Copyright Epic Games, Inc. All Rights Reserved.
 
+#ifdef DNA_BUILD_WITH_JSON_SUPPORT
+
 #include "dnatests/FixturesJSON.h"
 
 namespace dna {
@@ -182,3 +184,5 @@ const char* jsonDNA =
 }";
 
 }  // namespace dna
+
+#endif  // DNA_BUILD_WITH_JSON_SUPPORT
diff --git a/Engine/Plugins/Animation/RigLogic/Source/RigLogicLibTest/Private/dnatests/FixturesJSON.h b/Engine/Plugins/Animation/RigLogic/Source/RigLogicLibTest/Private/dnatests/FixturesJSON.h
--- a/Engine/Plugins/Animation/RigLogic/Source/RigLogicLibTest/Private/dnatests/FixturesJSON.h
+++ b/Engine/Plugins/Animation/RigLogic/Source/RigLogicLibTest/Private/dnatests/FixturesJSON.h
@@ -2,6 +2,9 @@
 
 #pragma once
 
+// *INDENT-OFF*
+#ifdef DNA_BUILD_WITH_JSON_SUPPORT
+
 #include "dnatests/Defs.h"
 
 #include <pma/TypeDefs.h>
@@ -11,3 +14,6 @@ namespace dna {
 extern const char* jsonDNA;
 
 }  // namespace dna
+
+#endif  // DNA_BUILD_WITH_JSON_SUPPORT
+// *INDENT-ON*
diff --git a/Engine/Plugins/Animation/RigLogic/Source/RigLogicLibTest/Private/dnatests/TestStreamReadWriteIntegration.cpp b/Engine/Plugins/Animation/RigLogic/Source/RigLogicLibTest/Private/dnatests/TestStreamReadWriteIntegration.cpp
--- a/Engine/Plugins/Animation/RigLogic/Source/RigLogicLibTest/Private/dnatests/TestStreamReadWriteIntegration.cpp
+++ b/Engine/Plugins/Animation/RigLogic/Source/RigLogicLibTest/Private/dnatests/TestStreamReadWriteIntegration.cpp
@@ -3,7 +3,9 @@
 #include "dnatests/TestStreamReadWriteIntegration.h"
 
 #include "dnatests/Defs.h"
-#include "dnatests/FixturesJSON.h"
+#ifdef DNA_BUILD_WITH_JSON_SUPPORT
+    #include "dnatests/FixturesJSON.h"
+#endif  // DNA_BUILD_WITH_JSON_SUPPORT
 #include "dnatests/Fixturesv21.h"
 #include "dnatests/Fixturesv22.h"
 #include "dnatests/Fixturesv23.h"
@@ -12,8 +14,10 @@
 #include "dna/DNA.h"
 #include "dna/BinaryStreamReader.h"
 #include "dna/BinaryStreamWriter.h"
-#include "dna/JSONStreamReader.h"
-#include "dna/JSONStreamWriter.h"
+#ifdef DNA_BUILD_WITH_JSON_SUPPORT
+    #include "dna/JSONStreamReader.h"
+    #include "dna/JSONStreamWriter.h"
+#endif  // DNA_BUILD_WITH_JSON_SUPPORT
 
 namespace dna {
 
@@ -438,21 +442,23 @@ struct ReaderDataVerifier<APICopyParameters<Reader, Writer, RawV23, DecodedV23,
 };
 
 using TAPICopyTestParameters = ::testing::Types<
-    APICopyParameters<dna::BinaryStreamReader, dna::BinaryStreamWriter, RawV21, DecodedV21, 0u, 1u, 0u>,
-    APICopyParameters<dna::BinaryStreamReader, dna::BinaryStreamWriter, RawV21, DecodedV21, 0u, 1u, 1u>,
-    APICopyParameters<dna::BinaryStreamReader, dna::BinaryStreamWriter, RawV21, DecodedV21, 0u, 0u, 0u>,
-    APICopyParameters<dna::BinaryStreamReader, dna::BinaryStreamWriter, RawV21, DecodedV21, 1u, 1u, 0u>,
-    APICopyParameters<dna::BinaryStreamReader, dna::BinaryStreamWriter, RawV22, DecodedV22, 0u, 1u, 0u>,
-    APICopyParameters<dna::BinaryStreamReader, dna::BinaryStreamWriter, RawV22, DecodedV22, 0u, 1u, 1u>,
-    APICopyParameters<dna::BinaryStreamReader, dna::BinaryStreamWriter, RawV22, DecodedV22, 0u, 0u, 0u>,
-    APICopyParameters<dna::BinaryStreamReader, dna::BinaryStreamWriter, RawV22, DecodedV22, 1u, 1u, 0u>,
-    APICopyParameters<dna::BinaryStreamReader, dna::BinaryStreamWriter, RawV23, DecodedV23, 0u, 1u, 0u>,
-    APICopyParameters<dna::BinaryStreamReader, dna::BinaryStreamWriter, RawV23, DecodedV23, 0u, 1u, 1u>,
-    APICopyParameters<dna::BinaryStreamReader, dna::BinaryStreamWriter, RawV23, DecodedV23, 0u, 0u, 0u>,
-    APICopyParameters<dna::BinaryStreamReader, dna::BinaryStreamWriter, RawV23, DecodedV23, 1u, 1u, 0u>,
-    APICopyParameters<dna::JSONStreamReader, dna::JSONStreamWriter, RawV21, DecodedV21, 0u, 1u, 0u>,
-    APICopyParameters<dna::JSONStreamReader, dna::JSONStreamWriter, RawV22, DecodedV22, 0u, 1u, 0u>,
-    APICopyParameters<dna::JSONStreamReader, dna::JSONStreamWriter, RawV23, DecodedV23, 0u, 1u, 0u>
+    APICopyParameters<dna::BinaryStreamReader, dna::BinaryStreamWriter, RawV21, DecodedV21, 0u, 1u, 0u>
+    , APICopyParameters<dna::BinaryStreamReader, dna::BinaryStreamWriter, RawV21, DecodedV21, 0u, 1u, 1u>
+    , APICopyParameters<dna::BinaryStreamReader, dna::BinaryStreamWriter, RawV21, DecodedV21, 0u, 0u, 0u>
+    , APICopyParameters<dna::BinaryStreamReader, dna::BinaryStreamWriter, RawV21, DecodedV21, 1u, 1u, 0u>
+    , APICopyParameters<dna::BinaryStreamReader, dna::BinaryStreamWriter, RawV22, DecodedV22, 0u, 1u, 0u>
+    , APICopyParameters<dna::BinaryStreamReader, dna::BinaryStreamWriter, RawV22, DecodedV22, 0u, 1u, 1u>
+    , APICopyParameters<dna::BinaryStreamReader, dna::BinaryStreamWriter, RawV22, DecodedV22, 0u, 0u, 0u>
+    , APICopyParameters<dna::BinaryStreamReader, dna::BinaryStreamWriter, RawV22, DecodedV22, 1u, 1u, 0u>
+    , APICopyParameters<dna::BinaryStreamReader, dna::BinaryStreamWriter, RawV23, DecodedV23, 0u, 1u, 0u>
+    , APICopyParameters<dna::BinaryStreamReader, dna::BinaryStreamWriter, RawV23, DecodedV23, 0u, 1u, 1u>
+    , APICopyParameters<dna::BinaryStreamReader, dna::BinaryStreamWriter, RawV23, DecodedV23, 0u, 0u, 0u>
+    , APICopyParameters<dna::BinaryStreamReader, dna::BinaryStreamWriter, RawV23, DecodedV23, 1u, 1u, 0u>
+    #ifdef DNA_BUILD_WITH_JSON_SUPPORT
+        , APICopyParameters<dna::JSONStreamReader, dna::JSONStreamWriter, RawV21, DecodedV21, 0u, 1u, 0u>
+        , APICopyParameters<dna::JSONStreamReader, dna::JSONStreamWriter, RawV22, DecodedV22, 0u, 1u, 0u>
+        , APICopyParameters<dna::JSONStreamReader, dna::JSONStreamWriter, RawV23, DecodedV23, 0u, 1u, 0u>
+    #endif  // DNA_BUILD_WITH_JSON_SUPPORT
     >;
 TYPED_TEST_SUITE(StreamReadWriteAPICopyIntegrationTest, TAPICopyTestParameters, );
 
@@ -471,17 +477,19 @@ struct ReaderFactory<dna::BinaryStreamReader> {
 
 };
 
-template<>
-struct ReaderFactory<dna::JSONStreamReader> {
-    static pma::ScopedPtr<dna::JSONStreamReader> create(trio::BoundedIOStream* stream,
-                                                        dna::DataLayer  /*unused*/,
-                                                        dna::UnknownLayerPolicy  /*unused*/,
-                                                        std::uint16_t  /*unused*/,
-                                                        std::uint16_t  /*unused*/) {
-        return pma::makeScoped<dna::JSONStreamReader>(stream);
-    }
+#ifdef DNA_BUILD_WITH_JSON_SUPPORT
+    template<>
+    struct ReaderFactory<dna::JSONStreamReader> {
+        static pma::ScopedPtr<dna::JSONStreamReader> create(trio::BoundedIOStream* stream,
+                                                            dna::DataLayer  /*unused*/,
+                                                            dna::UnknownLayerPolicy  /*unused*/,
+                                                            std::uint16_t  /*unused*/,
+                                                            std::uint16_t  /*unused*/) {
+            return pma::makeScoped<dna::JSONStreamReader>(stream);
+        }
 
-};
+    };
+#endif  // DNA_BUILD_WITH_JSON_SUPPORT
 
 TYPED_TEST(StreamReadWriteAPICopyIntegrationTest, VerifyAllDNADataAfterSetFromThroughAPI) {
     using CurrentParameters = typename TestFixture::Parameters;
@@ -580,35 +588,37 @@ TYPED_TEST(StreamReadWriteRawCopyIntegrationTest, VerifySetFromCopiesEvenUnknown
     ASSERT_EQ(expectedBytes, copiedBytes);
 }
 
-TEST(StreamReadWriteIntegrationTest, ReadWriteJSON) {
-    auto stream = pma::makeScoped<trio::MemoryStream>();
-    auto writer = pma::makeScoped<JSONStreamWriter>(stream.get(), 4u);
-
-    writer->setMeshName(0, "mesh0");
-    const Position vertices[] = {Position{0.0f, 1.0f, 2.0}, Position{3.0f, 4.0f, 5.0}};
-    writer->setVertexPositions(0u, vertices, 2u);
-    writer->write();
-
-    #if !defined(__clang__) && defined(__GNUC__)
-        #pragma GCC diagnostic push
-        #pragma GCC diagnostic ignored "-Wuseless-cast"
-    #endif
-    pma::Vector<char> json(static_cast<std::size_t>(stream->size()));
-    #if !defined(__clang__) && defined(__GNUC__)
-        #pragma GCC diagnostic pop
-    #endif
-
-    pma::String<char> expected = jsonDNA;
-    stream->seek(0ul);
-    stream->read(json.data(), json.size());
-    ASSERT_EQ(json.size(), expected.size());
-    ASSERT_ELEMENTS_EQ(json.data(), expected.data(), expected.size());
-
-    stream->seek(0ul);
-    auto reader = pma::makeScoped<JSONStreamReader>(stream.get());
-    reader->read();
-    ASSERT_TRUE(dna::Status::isOk());
-}
+#ifdef DNA_BUILD_WITH_JSON_SUPPORT
+    TEST(StreamReadWriteIntegrationTest, ReadWriteJSON) {
+        auto stream = pma::makeScoped<trio::MemoryStream>();
+        auto writer = pma::makeScoped<JSONStreamWriter>(stream.get(), 4u);
+
+        writer->setMeshName(0, "mesh0");
+        const Position vertices[] = {Position{0.0f, 1.0f, 2.0}, Position{3.0f, 4.0f, 5.0}};
+        writer->setVertexPositions(0u, vertices, 2u);
+        writer->write();
+
+        #if !defined(__clang__) && defined(__GNUC__)
+            #pragma GCC diagnostic push
+            #pragma GCC diagnostic ignored "-Wuseless-cast"
+        #endif
+        pma::Vector<char> json(static_cast<std::size_t>(stream->size()));
+        #if !defined(__clang__) && defined(__GNUC__)
+            #pragma GCC diagnostic pop
+        #endif
+
+        pma::String<char> expected = jsonDNA;
+        stream->seek(0ul);
+        stream->read(json.data(), json.size());
+        ASSERT_EQ(json.size(), expected.size());
+        ASSERT_ELEMENTS_EQ(json.data(), expected.data(), expected.size());
+
+        stream->seek(0ul);
+        auto reader = pma::makeScoped<JSONStreamReader>(stream.get());
+        reader->read();
+        ASSERT_TRUE(dna::Status::isOk());
+    }
+#endif  // DNA_BUILD_WITH_JSON_SUPPORT
 
 using TReadWriteMultipleParameters = ::testing::Types<
     ReadWriteMultipleParameters<RawV21>,
diff --git a/Engine/Plugins/Animation/RigLogic/Source/RigLogicLibTest/Private/rltests/ml/cpu/FixturesBlock4.cpp b/Engine/Plugins/Animation/RigLogic/Source/RigLogicLibTest/Private/rltests/ml/cpu/FixturesBlock4.cpp
--- a/Engine/Plugins/Animation/RigLogic/Source/RigLogicLibTest/Private/rltests/ml/cpu/FixturesBlock4.cpp
+++ b/Engine/Plugins/Animation/RigLogic/Source/RigLogicLibTest/Private/rltests/ml/cpu/FixturesBlock4.cpp
@@ -1,5 +1,8 @@
 // Copyright Epic Games, Inc. All Rights Reserved.
 
+// *INDENT-OFF*
+#ifdef RL_BUILD_WITH_ML_EVALUATOR
+
 #include "rltests/ml/cpu/FixturesBlock4.h"
 
 #include "riglogic/TypeDefs.h"
@@ -1100,3 +1103,6 @@ CanonicalReader::~CanonicalReader() = default;
 }  // namespace ml
 
 }  // namespace rltests
+
+#endif  // RL_BUILD_WITH_ML_EVALUATOR
+// *INDENT-ON*
diff --git a/Engine/Plugins/Animation/RigLogic/Source/RigLogicLibTest/Private/rltests/ml/cpu/FixturesBlock4.h b/Engine/Plugins/Animation/RigLogic/Source/RigLogicLibTest/Private/rltests/ml/cpu/FixturesBlock4.h
--- a/Engine/Plugins/Animation/RigLogic/Source/RigLogicLibTest/Private/rltests/ml/cpu/FixturesBlock4.h
+++ b/Engine/Plugins/Animation/RigLogic/Source/RigLogicLibTest/Private/rltests/ml/cpu/FixturesBlock4.h
@@ -2,6 +2,9 @@
 
 #pragma once
 
+// *INDENT-OFF*
+#ifdef RL_BUILD_WITH_ML_EVALUATOR
+
 #include "rltests/Defs.h"
 #include "rltests/StorageValueType.h"
 #include "rltests/dna/FakeReader.h"
@@ -135,3 +138,6 @@ class CanonicalReader : public dna::FakeReader {
 }  // namespace ml
 
 }  // namespace rltests
+
+#endif  // RL_BUILD_WITH_ML_EVALUATOR
+// *INDENT-ON*
diff --git a/Engine/Plugins/Animation/RigLogic/Source/RigLogicLibTest/Private/rltests/ml/cpu/TestInference.cpp b/Engine/Plugins/Animation/RigLogic/Source/RigLogicLibTest/Private/rltests/ml/cpu/TestInference.cpp
--- a/Engine/Plugins/Animation/RigLogic/Source/RigLogicLibTest/Private/rltests/ml/cpu/TestInference.cpp
+++ b/Engine/Plugins/Animation/RigLogic/Source/RigLogicLibTest/Private/rltests/ml/cpu/TestInference.cpp
@@ -1,5 +1,8 @@
 // Copyright Epic Games, Inc. All Rights Reserved.
 
+// *INDENT-OFF*
+#ifdef RL_BUILD_WITH_ML_EVALUATOR
+
 #include "rltests/Defs.h"
 #include "rltests/StorageValueType.h"
 #include "rltests/controls/ControlFixtures.h"
@@ -21,13 +24,22 @@ namespace {
 template<typename TTestTypes>
 class MLBInferenceTest : public ::testing::Test {
     protected:
-        void SetUp() {
-            using T = typename std::tuple_element<0, TTestTypes>::type;
-            using TF256 = typename std::tuple_element<1, TTestTypes>::type;
-            using TF128 = typename std::tuple_element<2, TTestTypes>::type;
+        void SetUp() override {
+            MLBInferenceTest::SetUpImpl();
+        }
+
+        template<typename TestTypes = TTestTypes>
+        typename std::enable_if<std::tuple_size<TestTypes>::value != 0ul, void>::type SetUpImpl() {
+            using T = typename std::tuple_element<0, TestTypes>::type;
+            using TF256 = typename std::tuple_element<1, TestTypes>::type;
+            using TF128 = typename std::tuple_element<2, TestTypes>::type;
             evaluator = rl4::ml::cpu::Factory<T, TF256, TF128>::create(&reader, &memRes);
         }
 
+        template<typename TestTypes = TTestTypes>
+        typename std::enable_if<std::tuple_size<TestTypes>::value == 0ul, void>::type SetUpImpl() {
+        }
+
     protected:
         pma::AlignedMemoryResource memRes;
         rltests::ml::block4::CanonicalReader reader;
@@ -40,10 +52,10 @@ class MLBInferenceTest : public ::testing::Test {
 #if defined(RL_BUILD_WITH_AVX) && defined(RL_BUILD_WITH_SSE)
     using InferenceTypeList = ::testing::Types<
     #ifdef RL_USE_HALF_FLOATS
-            std::tuple<StorageValueType, trimd::avx::F256, trimd::F128>,
+            std::tuple<StorageValueType, trimd::avx::F256, trimd::sse::F128>,
             std::tuple<StorageValueType, trimd::sse::F256, trimd::sse::F128>
     #else
-            std::tuple<StorageValueType, trimd::avx::F256, trimd::F128>,
+            std::tuple<StorageValueType, trimd::avx::F256, trimd::sse::F128>,
             std::tuple<StorageValueType, trimd::sse::F256, trimd::sse::F128>,
             std::tuple<StorageValueType, trimd::scalar::F256, trimd::scalar::F128>
     #endif  // RL_USE_HALF_FLOATS
@@ -51,9 +63,9 @@ class MLBInferenceTest : public ::testing::Test {
 #elif defined(RL_BUILD_WITH_AVX)
     using InferenceTypeList = ::testing::Types<
     #ifdef RL_USE_HALF_FLOATS
-            std::tuple<StorageValueType, trimd::avx::F256, trimd::F128>
+            std::tuple<StorageValueType, trimd::avx::F256, trimd::sse::F128>
     #else
-            std::tuple<StorageValueType, trimd::avx::F256, trimd::F128>,
+            std::tuple<StorageValueType, trimd::avx::F256, trimd::sse::F128>,
             std::tuple<StorageValueType, trimd::scalar::F256, trimd::scalar::F128>
     #endif  // RL_USE_HALF_FLOATS
         >;
@@ -70,13 +82,17 @@ class MLBInferenceTest : public ::testing::Test {
     #ifndef RL_USE_HALF_FLOATS
         using InferenceTypeList = ::testing::Types<std::tuple<StorageValueType, trimd::scalar::F256, trimd::scalar::F128> >;
     #else
-        using InferenceTypeList = ::testing::Types<>;
+        using InferenceTypeList = ::testing::Types<std::tuple<> >;
     #endif  // RL_USE_HALF_FLOATS
 #endif
 
 TYPED_TEST_SUITE(MLBInferenceTest, InferenceTypeList, );
 
 TYPED_TEST(MLBInferenceTest, InferencePerLOD) {
+    if (this->evaluator == nullptr) {
+        return;
+    }
+
     auto inputInstanceFactory = ControlsFactory::getInstanceFactory(0,
                                                                     rltests::ml::block4::unoptimized::rawControlCount,
                                                                     0,
@@ -105,3 +121,6 @@ TYPED_TEST(MLBInferenceTest, InferencePerLOD) {
 #ifdef _MSC_VER
     #pragma warning(pop)
 #endif
+
+#endif  // RL_BUILD_WITH_ML_EVALUATOR
+// *INDENT-ON*
diff --git a/Engine/Plugins/Animation/RigLogic/Source/RigLogicLibTest/Private/rltests/ml/cpu/TestStorageBuilder.cpp b/Engine/Plugins/Animation/RigLogic/Source/RigLogicLibTest/Private/rltests/ml/cpu/TestStorageBuilder.cpp
--- a/Engine/Plugins/Animation/RigLogic/Source/RigLogicLibTest/Private/rltests/ml/cpu/TestStorageBuilder.cpp
+++ b/Engine/Plugins/Animation/RigLogic/Source/RigLogicLibTest/Private/rltests/ml/cpu/TestStorageBuilder.cpp
@@ -1,5 +1,8 @@
 // Copyright Epic Games, Inc. All Rights Reserved.
 
+// *INDENT-OFF*
+#ifdef RL_BUILD_WITH_ML_EVALUATOR
+
 #include "rltests/Defs.h"
 #include "rltests/StorageValueType.h"
 #include "rltests/ml/cpu/FixturesBlock4.h"
@@ -65,7 +68,8 @@ namespace {
 template<typename TTestTypes>
 class MLBSStorageBuilderTest : public ::testing::Test {
     protected:
-        void buildStorage() {
+        template<typename TestTypes = TTestTypes>
+        typename std::enable_if<std::tuple_size<TestTypes>::value != 0ul, void>::type buildStorage() {
             using T = typename std::tuple_element<0, TTestTypes>::type;
             using TF256 = typename std::tuple_element<1, TTestTypes>::type;
             using TF128 = typename std::tuple_element<2, TTestTypes>::type;
@@ -74,6 +78,10 @@ class MLBSStorageBuilderTest : public ::testing::Test {
             rl4::ml::cpu::Evaluator<T, TF256, TF128>::Accessor::assertRawDataEqual(*evaluatorImpl);
         }
 
+        template<typename TestTypes = TTestTypes>
+        typename std::enable_if<std::tuple_size<TestTypes>::value == 0ul, void>::type buildStorage() {
+        }
+
     protected:
         pma::AlignedMemoryResource memRes;
         rltests::ml::block4::CanonicalReader reader;
@@ -85,10 +93,10 @@ class MLBSStorageBuilderTest : public ::testing::Test {
 #if defined(RL_BUILD_WITH_AVX) && defined(RL_BUILD_WITH_SSE)
     using StorageValueTypeList = ::testing::Types<
     #ifdef RL_USE_HALF_FLOATS
-            std::tuple<StorageValueType, trimd::avx::F256, trimd::F128>,
+            std::tuple<StorageValueType, trimd::avx::F256, trimd::sse::F128>,
             std::tuple<StorageValueType, trimd::sse::F256, trimd::sse::F128>
     #else
-            std::tuple<StorageValueType, trimd::avx::F256, trimd::F128>,
+            std::tuple<StorageValueType, trimd::avx::F256, trimd::sse::F128>,
             std::tuple<StorageValueType, trimd::sse::F256, trimd::sse::F128>,
             std::tuple<StorageValueType, trimd::scalar::F256, trimd::scalar::F128>
     #endif  // RL_USE_HALF_FLOATS
@@ -96,9 +104,9 @@ class MLBSStorageBuilderTest : public ::testing::Test {
 #elif defined(RL_BUILD_WITH_AVX)
     using StorageValueTypeList = ::testing::Types<
     #ifdef RL_USE_HALF_FLOATS
-            std::tuple<StorageValueType, trimd::avx::F256, trimd::F128>
+            std::tuple<StorageValueType, trimd::avx::F256, trimd::sse::F128>
     #else
-            std::tuple<StorageValueType, trimd::avx::F256, trimd::F128>,
+            std::tuple<StorageValueType, trimd::avx::F256, trimd::sse::F128>,
             std::tuple<StorageValueType, trimd::scalar::F256, trimd::scalar::F128>
     #endif  // RL_USE_HALF_FLOATS
         >;
@@ -115,7 +123,7 @@ class MLBSStorageBuilderTest : public ::testing::Test {
     #ifndef RL_USE_HALF_FLOATS
         using StorageValueTypeList = ::testing::Types<std::tuple<StorageValueType, trimd::scalar::F256, trimd::scalar::F128> >;
     #else
-        using StorageValueTypeList = ::testing::Types<>;
+        using StorageValueTypeList = ::testing::Types<std::tuple<> >;
     #endif  // RL_USE_HALF_FLOATS
 #endif
 
@@ -128,3 +136,6 @@ TYPED_TEST(MLBSStorageBuilderTest, LayoutOptimization) {
 #ifdef _MSC_VER
     #pragma warning(pop)
 #endif
+
+#endif  // RL_BUILD_WITH_ML_EVALUATOR
+// *INDENT-ON*
