Prevent reinitialization of RigLogic in each frame

diff --git a/Engine/Plugins/Animation/RigLogic/Source/RigLogicModule/Private/RigUnit_RigLogic.cpp b/Engine/Plugins/Animation/RigLogic/Source/RigLogicModule/Private/RigUnit_RigLogic.cpp
--- a/Engine/Plugins/Animation/RigLogic/Source/RigLogicModule/Private/RigUnit_RigLogic.cpp
+++ b/Engine/Plugins/Animation/RigLogic/Source/RigLogicModule/Private/RigUnit_RigLogic.cpp
@@ -85,6 +85,7 @@ void FRigUnit_RigLogic_Data::InitializeRigLogic(const URigHierarchy* InHierarchy
 {
 	DECLARE_SCOPE_HIERARCHICAL_COUNTER_FUNC()
 	QUICK_SCOPE_CYCLE_COUNTER(STAT_RigUnit_RigLogic_InitializeRigLogic);
+
 	if (!NewContext.IsValid())
 	{
 		UE_LOG(LogRigLogicUnit, Warning, TEXT("No valid DNA file found, abort initialization."));
@@ -115,7 +116,8 @@ void FRigUnit_RigLogic_Data::MapInputCurveIndices(const URigHierarchy* InHierarc
 {
 	DECLARE_SCOPE_HIERARCHICAL_COUNTER_FUNC()
 	QUICK_SCOPE_CYCLE_COUNTER(STAT_RigUnit_RigLogic_MapInputCurveIndices);
-	const IBehaviorReader* DNABehavior = LocalRigRuntimeContext->BehaviorReader.Get();
+
+	const IDNAReader* DNABehavior = LocalRigRuntimeContext->BehaviorReader.Get();
 	const uint32 ControlCount = DNABehavior->GetRawControlCount();
 	InputCurveIndices.Reset(ControlCount);
 	for (uint32_t ControlIndex = 0; ControlIndex < ControlCount; ++ControlIndex)
@@ -136,7 +138,8 @@ void FRigUnit_RigLogic_Data::MapJoints(const URigHierarchy* Hierarchy)
 {
 	DECLARE_SCOPE_HIERARCHICAL_COUNTER_FUNC()
 	QUICK_SCOPE_CYCLE_COUNTER(STAT_RigUnit_RigLogic_MapJoints);
-	const IBehaviorReader* DNABehavior = LocalRigRuntimeContext->BehaviorReader.Get();
+
+	const IDNAReader* DNABehavior = LocalRigRuntimeContext->BehaviorReader.Get();
 	const uint16 JointCount = DNABehavior->GetJointCount();
 	HierarchyBoneIndices.Reset(JointCount);
 	for (uint16 JointIndex = 0; JointIndex < JointCount ; ++JointIndex)
@@ -152,7 +155,8 @@ void FRigUnit_RigLogic_Data::MapMorphTargets(const URigHierarchy* InHierarchy)
 {
 	DECLARE_SCOPE_HIERARCHICAL_COUNTER_FUNC()
 	QUICK_SCOPE_CYCLE_COUNTER(STAT_RigUnit_RigLogic_MapMorphTargets);
-	const IBehaviorReader* DNABehavior = LocalRigRuntimeContext->BehaviorReader.Get();
+
+	const IDNAReader* DNABehavior = LocalRigRuntimeContext->BehaviorReader.Get();
 	const uint16 LODCount = DNABehavior->GetLODCount();
 
 	MorphTargetCurveIndices.Reset();
@@ -186,7 +190,8 @@ void FRigUnit_RigLogic_Data::MapMaskMultipliers(const URigHierarchy* InHierarchy
 {
 	DECLARE_SCOPE_HIERARCHICAL_COUNTER_FUNC()
 	QUICK_SCOPE_CYCLE_COUNTER(STAT_RigUnit_RigLogic_MapMaskMultipliers);
-	const IBehaviorReader* DNABehavior = LocalRigRuntimeContext->BehaviorReader.Get();
+
+	const IDNAReader* DNABehavior = LocalRigRuntimeContext->BehaviorReader.Get();
 	const uint16 LODCount = DNABehavior->GetLODCount();
 	CurveElementIndicesForAnimMaps.Reset();
 	CurveElementIndicesForAnimMaps.AddDefaulted(LODCount);
@@ -230,6 +235,7 @@ void FRigUnit_RigLogic_Data::CalculateRigLogic(const URigHierarchy* InHierarchy)
 		const float Value = InHierarchy->GetCurveValue(CurveIndex);
 		RigInstance->SetRawControl(ControlIndex, Value);
 	}
+
 	LocalRigRuntimeContext->RigLogic->Calculate(RigInstance.Get());
 }
 
@@ -265,6 +271,7 @@ void FRigUnit_RigLogic_Data::UpdateBlendShapeCurves(URigHierarchy* InHierarchy,
 {
 	DECLARE_SCOPE_HIERARCHICAL_COUNTER_FUNC()
 	QUICK_SCOPE_CYCLE_COUNTER(STAT_RigUnit_RigLogic_UpdateBlendShapeCurves);
+
 	// set output blend shapes
 	if (BlendShapeIndices.IsValidIndex(CurrentLOD) && MorphTargetCurveIndices.IsValidIndex(CurrentLOD))
 	{
@@ -294,6 +301,7 @@ void FRigUnit_RigLogic_Data::UpdateAnimMapCurves(URigHierarchy* InHierarchy, TAr
 {
 	DECLARE_SCOPE_HIERARCHICAL_COUNTER_FUNC()
 	QUICK_SCOPE_CYCLE_COUNTER(STAT_RigUnit_RigLogic_UpdateAnimMapCurves);
+
 	// set output mask multipliers
 	// In case curves are not imported yet into CL, AnimatedMapsCurveIndices will be empty, so we need to check
 	// array bounds before trying to access it:
@@ -357,6 +365,7 @@ FRigUnit_RigLogic_Execute()
 			// Context is initialized with a BehaviorReader, which can be imported into SkeletalMesh from DNA file
 			// or overwritten by GeneSplicer when making a new character
 			Data.InitializeRigLogic(Hierarchy, RigRuntimeContext);
+			bIsInitialized = true;
 			//const double delta = FPlatformTime::Seconds() - startTime;
 			//UE_LOG(LogRigLogicUnit, Warning, TEXT("RigLogic::Init execution time: %f"), delta);
 		}
diff --git a/Engine/Plugins/Animation/RigLogic/Source/RigLogicModule/Private/SharedRigRuntimeContext.h b/Engine/Plugins/Animation/RigLogic/Source/RigLogicModule/Private/SharedRigRuntimeContext.h
--- a/Engine/Plugins/Animation/RigLogic/Source/RigLogicModule/Private/SharedRigRuntimeContext.h
+++ b/Engine/Plugins/Animation/RigLogic/Source/RigLogicModule/Private/SharedRigRuntimeContext.h
@@ -15,7 +15,7 @@ struct FSharedRigRuntimeContext
 	};
 
 	// Redundant but stored here as well so the runtime context can be updated / queried atomically
-	TSharedPtr<IBehaviorReader> BehaviorReader;
+	TSharedPtr<IDNAReader> BehaviorReader;
 
 	/** RigLogic itself is stateless, and is designed to be shared between
 	  * multiple rig instances based on the same DNA.
diff --git a/Engine/Plugins/Animation/RigLogic/Source/RigLogicModule/Public/RigUnit_RigLogic.h b/Engine/Plugins/Animation/RigLogic/Source/RigLogicModule/Public/RigUnit_RigLogic.h
--- a/Engine/Plugins/Animation/RigLogic/Source/RigLogicModule/Public/RigUnit_RigLogic.h
+++ b/Engine/Plugins/Animation/RigLogic/Source/RigLogicModule/Public/RigUnit_RigLogic.h
@@ -57,7 +57,7 @@ struct FRigUnit_RigLogic_Data
 	  * if value is valid, it is cast to appropriate uint type
 	**/
 
-	/** RL input index to ControlRig's input curve index for each LOD **/
+	/** RL input index to ControlRig's input curve index **/
 	UPROPERTY(transient)
 	TArray<int32> InputCurveIndices;
 
