Reduce default alignment requirement of memory allocations in RigLogic

diff --git a/Engine/Plugins/Animation/RigLogic/Source/RigLogicLib/Public/pma/PolyAllocator.h b/Engine/Plugins/Animation/RigLogic/Source/RigLogicLib/Public/pma/PolyAllocator.h
--- a/Engine/Plugins/Animation/RigLogic/Source/RigLogicLib/Public/pma/PolyAllocator.h
+++ b/Engine/Plugins/Animation/RigLogic/Source/RigLogicLib/Public/pma/PolyAllocator.h
@@ -24,6 +24,8 @@ namespace pma {
 
 namespace impl {
 
+using min_align_t = void*;
+
 template<typename T, typename U>
 struct max_align_of {
     using type = typename std::conditional<(alignof(T) > alignof(U)), T, U>::type;
@@ -124,7 +126,7 @@ class PolyAllocator {
 }  // namespace impl
 
 template<typename T,
-         std::size_t Alignment = alignof(typename impl::max_align_of<T, std::max_align_t>::type),
+         std::size_t Alignment = alignof(typename impl::max_align_of<T, impl::min_align_t>::type),
          class TDefaultMemoryResource = DefaultMemoryResource>
 class PolyAllocator : public std::scoped_allocator_adaptor<impl::PolyAllocator<T, Alignment, TDefaultMemoryResource> > {
     private:
