[MH-9368] Add LOD threshold to rig logic anim graph node

Compute for facial animation can possibly be LOD'ed earlier than other skeletal post-processing, which might be part of the same PP AnimBP. In order to do that, we'll introduce a LOD Threshold, so that we can get more fine control over that.

diff --git a/Engine/Plugins/Animation/RigLogic/Source/RigLogicModule/Private/AnimNode_RigLogic.cpp b/Engine/Plugins/Animation/RigLogic/Source/RigLogicModule/Private/AnimNode_RigLogic.cpp
--- a/Engine/Plugins/Animation/RigLogic/Source/RigLogicModule/Private/AnimNode_RigLogic.cpp
+++ b/Engine/Plugins/Animation/RigLogic/Source/RigLogicModule/Private/AnimNode_RigLogic.cpp
@@ -128,13 +128,17 @@ void FAnimNode_RigLogic::Evaluate_AnyThread(FPoseContext& OutputContext)
 	{
 		return;
 	}
-	UpdateControlCurves(OutputContext, LocalDNAIndexMapping.Get());
-	CalculateRigLogic(LocalRigRuntimeContext->RigLogic.Get());
-	const uint16 CurrentLOD = RigInstance->GetLOD();
-	TArrayView<const uint16> VariableJointIndices = LocalRigRuntimeContext->VariableJointIndicesPerLOD[CurrentLOD].Values;
-	UpdateJoints(VariableJointIndices, LocalRigRuntimeContext->RigLogic->GetRawNeutralJointValues(), RigInstance->GetRawJointOutputs(), OutputContext);
-	UpdateBlendShapeCurves(LocalDNAIndexMapping.Get(), RigInstance->GetBlendShapeOutputs(), OutputContext);
-	UpdateAnimMapCurves(LocalDNAIndexMapping.Get(), RigInstance->GetAnimatedMapOutputs(), OutputContext);
+
+	if (IsLODEnabled(OutputContext.AnimInstanceProxy))
+	{
+		UpdateControlCurves(OutputContext, LocalDNAIndexMapping.Get());
+		CalculateRigLogic(LocalRigRuntimeContext->RigLogic.Get());
+		const uint16 CurrentLOD = RigInstance->GetLOD();
+		TArrayView<const uint16> VariableJointIndices = LocalRigRuntimeContext->VariableJointIndicesPerLOD[CurrentLOD].Values;
+		UpdateJoints(VariableJointIndices, LocalRigRuntimeContext->RigLogic->GetRawNeutralJointValues(), RigInstance->GetRawJointOutputs(), OutputContext);
+		UpdateBlendShapeCurves(LocalDNAIndexMapping.Get(), RigInstance->GetBlendShapeOutputs(), OutputContext);
+		UpdateAnimMapCurves(LocalDNAIndexMapping.Get(), RigInstance->GetAnimatedMapOutputs(), OutputContext);
+	}
 }
 
 void FAnimNode_RigLogic::GatherDebugData(FNodeDebugData& DebugData)
@@ -201,9 +205,11 @@ void FAnimNode_RigLogic::UpdateBlendShapeCurves(const FDNAIndexMapping* DNAIndex
 	const FDNAIndexMapping::FCachedIndexedCurve& MorphTargetCurve = DNAIndexMapping->MorphTargetCurvesPerLOD[LOD];
 	UE::Anim::FNamedValueArrayUtils::Union(OutputContext.Curve, MorphTargetCurve, [&BlendShapeValues](UE::Anim::FCurveElement& InOutResult, const UE::Anim::FCurveElementIndexed& InSource, UE::Anim::ENamedValueUnionFlags InFlags)
 	{
-		check(BlendShapeValues.IsValidIndex(InSource.Index));
-		InOutResult.Value = BlendShapeValues[InSource.Index];
-		InOutResult.Flags |= UE::Anim::ECurveElementFlags::MorphTarget;
+		if (BlendShapeValues.IsValidIndex(InSource.Index))
+		{
+			InOutResult.Value = BlendShapeValues[InSource.Index];
+			InOutResult.Flags |= UE::Anim::ECurveElementFlags::MorphTarget;
+		}
 	});
 }
 
@@ -216,9 +222,11 @@ void FAnimNode_RigLogic::UpdateAnimMapCurves(const FDNAIndexMapping* DNAIndexMap
 	const FDNAIndexMapping::FCachedIndexedCurve& MaskMultiplierCurve = DNAIndexMapping->MaskMultiplierCurvesPerLOD[LOD];
 	UE::Anim::FNamedValueArrayUtils::Union(OutputContext.Curve, MaskMultiplierCurve, [&AnimMapOutputs](UE::Anim::FCurveElement& InOutResult, const UE::Anim::FCurveElementIndexed& InSource, UE::Anim::ENamedValueUnionFlags InFlags)
 	{
-		check(AnimMapOutputs.IsValidIndex(InSource.Index));
-		InOutResult.Value = AnimMapOutputs[InSource.Index];
-		InOutResult.Flags |= UE::Anim::ECurveElementFlags::Material;
+		if (AnimMapOutputs.IsValidIndex(InSource.Index))
+		{
+			InOutResult.Value = AnimMapOutputs[InSource.Index];
+			InOutResult.Flags |= UE::Anim::ECurveElementFlags::Material;
+		}
 	});
 }
 
diff --git a/Engine/Plugins/Animation/RigLogic/Source/RigLogicModule/Public/AnimNode_RigLogic.h b/Engine/Plugins/Animation/RigLogic/Source/RigLogicModule/Public/AnimNode_RigLogic.h
--- a/Engine/Plugins/Animation/RigLogic/Source/RigLogicModule/Public/AnimNode_RigLogic.h
+++ b/Engine/Plugins/Animation/RigLogic/Source/RigLogicModule/Public/AnimNode_RigLogic.h
@@ -30,6 +30,7 @@ public:
 	void Update_AnyThread(const FAnimationUpdateContext& Context) override;
 	void Evaluate_AnyThread(FPoseContext& Output) override;
 	void GatherDebugData(FNodeDebugData& DebugData) override;
+	int32 GetLODThreshold() const override { return LODThreshold; }
 
 public:
 	UPROPERTY(EditAnywhere, BlueprintReadWrite, Category = Links)
@@ -43,6 +44,14 @@ private:
 	void UpdateAnimMapCurves(const FDNAIndexMapping* DNAIndexMapping, TArrayView<const float> AnimMapOutputs, FPoseContext& OutputContext);
 
 private:
+	/*
+	 * Max LOD level that post-process AnimBPs are evaluated.
+	 * For example if you have the threshold set to 2, it will evaluate until including LOD 2 (based on 0 index). In case the LOD level gets set to 3, it will stop evaluating the post-process AnimBP.
+	 * Setting it to -1 will always evaluate it and disable LODing.
+	 */
+	UPROPERTY(EditAnywhere, Category = RigLogic)
+	int32 LODThreshold = INDEX_NONE;
+
 	TSharedPtr<FSharedRigRuntimeContext> LocalRigRuntimeContext;
 	TSharedPtr<FDNAIndexMapping> LocalDNAIndexMapping;
 	FRigInstance* RigInstance;
