Move RigLogic initialization out of AnimNode CacheBones_AnyThread function

diff --git a/Engine/Plugins/Animation/RigLogic/Source/RigLogicModule/Private/AnimNode_RigLogic.cpp b/Engine/Plugins/Animation/RigLogic/Source/RigLogicModule/Private/AnimNode_RigLogic.cpp
--- a/Engine/Plugins/Animation/RigLogic/Source/RigLogicModule/Private/AnimNode_RigLogic.cpp
+++ b/Engine/Plugins/Animation/RigLogic/Source/RigLogicModule/Private/AnimNode_RigLogic.cpp
@@ -75,7 +75,7 @@ void FAnimNode_RigLogic::CacheBones_AnyThread(const FAnimationCacheBonesContext&
 		return;
 	}
 
-	TSharedPtr<FSharedRigRuntimeContext> SharedRigRuntimeContext = DNAAsset->GetRigRuntimeContext(UDNAAsset::EDNARetentionPolicy::Unload);
+	TSharedPtr<FSharedRigRuntimeContext> SharedRigRuntimeContext = DNAAsset->GetRigRuntimeContext();
 	if (!SharedRigRuntimeContext.IsValid())
 	{
 		return;
diff --git a/Engine/Plugins/Animation/RigLogic/Source/RigLogicModule/Private/DNAAsset.cpp b/Engine/Plugins/Animation/RigLogic/Source/RigLogicModule/Private/DNAAsset.cpp
--- a/Engine/Plugins/Animation/RigLogic/Source/RigLogicModule/Private/DNAAsset.cpp
+++ b/Engine/Plugins/Animation/RigLogic/Source/RigLogicModule/Private/DNAAsset.cpp
@@ -96,71 +96,72 @@ UDNAAsset::~UDNAAsset() = default;
 
 TSharedPtr<IDNAReader> UDNAAsset::GetBehaviorReader()
 {
-	FScopeLock ScopeLock{&DNAUpdateSection};
+	FReadScopeLock DNAScopeLock{DNAUpdateLock};
 	return BehaviorReader;
 }
 
 #if WITH_EDITORONLY_DATA
 TSharedPtr<IDNAReader> UDNAAsset::GetGeometryReader()
 {
-	FScopeLock ScopeLock{&DNAUpdateSection};
+	FReadScopeLock DNAScopeLock{DNAUpdateLock};
 	return GeometryReader;
 }
 #endif
 
 void UDNAAsset::SetBehaviorReader(TSharedPtr<IDNAReader> SourceDNAReader)
 {
-	FScopeLock ScopeLock{&DNAUpdateSection};
+	FWriteScopeLock DNAScopeLock{DNAUpdateLock};
 	const size_t PredictedSize = (SourceDNAReader->GetNeuralNetworkCount() != 0) ? AVG_BEHAVIOR_SIZE + AVG_MACHINE_LEARNED_BEHAVIOR_SIZE : AVG_BEHAVIOR_SIZE;
 	BehaviorReader = CopyDNALayer(SourceDNAReader.Get(), EDNADataLayer::Behavior | EDNADataLayer::MachineLearnedBehavior, PredictedSize);
 	InvalidateRigRuntimeContext();
+	InitializeRigRuntimeContext();
 }
 
 void UDNAAsset::SetGeometryReader(TSharedPtr<IDNAReader> SourceDNAReader)
 {
 #if WITH_EDITORONLY_DATA
-	FScopeLock ScopeLock{&DNAUpdateSection};
+	FWriteScopeLock DNAScopeLock{DNAUpdateLock};
 	GeometryReader = CopyDNALayer(SourceDNAReader.Get(), EDNADataLayer::Geometry, AVG_GEOMETRY_SIZE);
 #endif // #if WITH_EDITORONLY_DATA
 }
 
 void UDNAAsset::InvalidateRigRuntimeContext()
 {
-	FWriteScopeLock ScopeLock{RigRuntimeContextUpdateLock};
+	FWriteScopeLock ContextScopeLock{RigRuntimeContextUpdateLock};
 	FWriteScopeLock MappingScopeLock(DNAIndexMappingUpdateLock);
 	RigRuntimeContext = nullptr;
 	DNAIndexMappingContainer.Empty(1);
 }
 
-TSharedPtr<FSharedRigRuntimeContext> UDNAAsset::GetRigRuntimeContext(EDNARetentionPolicy Policy)
+void UDNAAsset::InitializeRigRuntimeContext()
 {
 	LLM_SCOPE_BYNAME(TEXT("Animation/RigLogic"));
 
-	FScopeLock DNAScopeLock{&DNAUpdateSection};
-	FRWScopeLock ContextScopeLock(RigRuntimeContextUpdateLock, SLT_ReadOnly);
-	if (!RigRuntimeContext.IsValid())
+	// Assumes DNAUpdateLock is locked by caller
+
+	TSharedPtr<FSharedRigRuntimeContext> NewContext = MakeShared<FSharedRigRuntimeContext>();
+	if (BehaviorReader.IsValid() && (BehaviorReader->GetJointCount() != 0))
 	{
-		ContextScopeLock.ReleaseReadOnlyLockAndAcquireWriteLock_USE_WITH_CAUTION();
-		if (!RigRuntimeContext.IsValid())
+		NewContext->BehaviorReader = BehaviorReader;
+		NewContext->RigLogic = MakeShared<FRigLogic>(BehaviorReader.Get());
+		NewContext->CacheVariableJointIndices();
 		{
-			TSharedPtr<FSharedRigRuntimeContext> NewContext = MakeShared<FSharedRigRuntimeContext>();
-			if (BehaviorReader.IsValid() && (BehaviorReader->GetJointCount() != 0))
-			{
-				NewContext->BehaviorReader = BehaviorReader;
-				NewContext->RigLogic = MakeShared<FRigLogic>(BehaviorReader.Get());
-				NewContext->CacheVariableJointIndices();
-				RigRuntimeContext = NewContext;
+			FWriteScopeLock ContextScopeLock{RigRuntimeContextUpdateLock};
+			RigRuntimeContext = NewContext;
+		}
 #if !WITH_EDITOR
-				if (Policy == EDNARetentionPolicy::Unload)
-				{
-					BehaviorReader->Unload(EDNADataLayer::Behavior);
-					BehaviorReader->Unload(EDNADataLayer::Geometry);
-					BehaviorReader->Unload(EDNADataLayer::MachineLearnedBehavior);
-				}
+		BehaviorReader->Unload(EDNADataLayer::Behavior);
+		BehaviorReader->Unload(EDNADataLayer::Geometry);
+		BehaviorReader->Unload(EDNADataLayer::MachineLearnedBehavior);
 #endif  // !WITH_EDITOR
-			}
-		}
 	}
+}
+
+TSharedPtr<FSharedRigRuntimeContext> UDNAAsset::GetRigRuntimeContext()
+{
+	LLM_SCOPE_BYNAME(TEXT("Animation/RigLogic"));
+
+	FReadScopeLock ContextScopeLock{RigRuntimeContextUpdateLock};
 	return RigRuntimeContext;
 }
 
@@ -170,7 +171,7 @@ TSharedPtr<FDNAIndexMapping> UDNAAsset::GetDNAIndexMapping(const USkeleton* Skel
 {
 	LLM_SCOPE_BYNAME(TEXT("Animation/RigLogic"));
 
-	FScopeLock DNAScopeLock{&DNAUpdateSection};
+	FReadScopeLock DNAScopeLock{DNAUpdateLock};
 	FWriteScopeLock MappingScopeLock(DNAIndexMappingUpdateLock);
 
 	// Find currently needed mapping and also clean stale objects along the way (requires only one iteration over the map)
@@ -232,7 +233,7 @@ bool UDNAAsset::Init(const FString& DNAFilename)
 		return false;
 	}
 	
-	FScopeLock ScopeLock{&DNAUpdateSection};
+	FWriteScopeLock DNAScopeLock{DNAUpdateLock};
 
 	// Load run-time data (behavior) from whole-DNA buffer into BehaviorReader
 	BehaviorReader = ReadDNAFromBuffer(&TempFileBuffer, EDNADataLayer::Behavior | EDNADataLayer::MachineLearnedBehavior, 0u); //0u = MaxLOD
@@ -242,6 +243,7 @@ bool UDNAAsset::Init(const FString& DNAFilename)
 	}
 
 	InvalidateRigRuntimeContext();
+	InitializeRigRuntimeContext();
 
 #if WITH_EDITORONLY_DATA
 	//We use geometry part of the data in MHC only (for updating the SkeletalMesh with
@@ -268,7 +270,7 @@ void UDNAAsset::Serialize(FArchive& Ar)
 
 	Ar.UsingCustomVersion(FDNAAssetCustomVersion::GUID);
 
-	FScopeLock ScopeLock{&DNAUpdateSection};
+	FWriteScopeLock DNAScopeLock{DNAUpdateLock};
 
 	if (Ar.CustomVer(FDNAAssetCustomVersion::GUID) >= FDNAAssetCustomVersion::BeforeCustomVersionWasAdded)
 	{
@@ -286,6 +288,7 @@ void UDNAAsset::Serialize(FArchive& Ar)
 #endif // #if WITH_EDITORONLY_DATA
 
 			InvalidateRigRuntimeContext();
+			InitializeRigRuntimeContext();
 		}
 
 		if (Ar.IsSaving())
diff --git a/Engine/Plugins/Animation/RigLogic/Source/RigLogicModule/Private/RigUnit_RigLogic.cpp b/Engine/Plugins/Animation/RigLogic/Source/RigLogicModule/Private/RigUnit_RigLogic.cpp
--- a/Engine/Plugins/Animation/RigLogic/Source/RigLogicModule/Private/RigUnit_RigLogic.cpp
+++ b/Engine/Plugins/Animation/RigLogic/Source/RigLogicModule/Private/RigUnit_RigLogic.cpp
@@ -333,7 +333,7 @@ TSharedPtr<FSharedRigRuntimeContext> FRigUnit_RigLogic::GetSharedRigRuntimeConte
 		return nullptr;
 	}
 	UDNAAsset* DNAAsset = Cast<UDNAAsset>(UserData);
-	return DNAAsset->GetRigRuntimeContext(UDNAAsset::EDNARetentionPolicy::Keep);
+	return DNAAsset->GetRigRuntimeContext();
 }
 
 FRigUnit_RigLogic_Execute()
diff --git a/Engine/Plugins/Animation/RigLogic/Source/RigLogicModule/Public/DNAAsset.h b/Engine/Plugins/Animation/RigLogic/Source/RigLogicModule/Public/DNAAsset.h
--- a/Engine/Plugins/Animation/RigLogic/Source/RigLogicModule/Public/DNAAsset.h
+++ b/Engine/Plugins/Animation/RigLogic/Source/RigLogicModule/Public/DNAAsset.h
@@ -67,21 +67,16 @@ private:
 	friend struct FAnimNode_RigLogic;
 	friend struct FRigUnit_RigLogic;
 
-	enum class EDNARetentionPolicy
-	{
-		Keep,
-		Unload
-	};
-
-	TSharedPtr<FSharedRigRuntimeContext> GetRigRuntimeContext(EDNARetentionPolicy Policy);
+	TSharedPtr<FSharedRigRuntimeContext> GetRigRuntimeContext();
 	void InvalidateRigRuntimeContext();
+	void InitializeRigRuntimeContext();
 	TSharedPtr<FDNAIndexMapping> GetDNAIndexMapping(const USkeleton* Skeleton,
 													const USkeletalMesh* SkeletalMesh,
 													const USkeletalMeshComponent* SkeletalMeshComponent);
 
 private:
 	// Synchronize DNA updates
-	FCriticalSection DNAUpdateSection;
+	FRWLock DNAUpdateLock;
 
 	// Synchronize Rig Runtime Context updates
 	FRWLock RigRuntimeContextUpdateLock;
