Cache filtered joint list for current LOD in RigLogic

diff --git a/Engine/Plugins/Animation/RigLogic/Source/RigLogicModule/Private/AnimNode_RigLogic.cpp b/Engine/Plugins/Animation/RigLogic/Source/RigLogicModule/Private/AnimNode_RigLogic.cpp
--- a/Engine/Plugins/Animation/RigLogic/Source/RigLogicModule/Private/AnimNode_RigLogic.cpp
+++ b/Engine/Plugins/Animation/RigLogic/Source/RigLogicModule/Private/AnimNode_RigLogic.cpp
@@ -98,12 +98,16 @@ void FAnimNode_RigLogic::CacheBones_AnyThread(const FAnimationCacheBonesContext&
 	{
 		RigInstance->SetLOD(Context.AnimInstanceProxy->GetLODLevel());
 		const uint16 CurrentLOD = RigInstance->GetLOD();
-		const TArray<uint16>& VariableJointIndices = DNAAsset->RigRuntimeContext->VariableJointIndicesPerLOD[CurrentLOD].Values;
-		JointsMapDNAIndicesToCompactPoseBoneIndices.Reset(VariableJointIndices.Num());
+		const TArray<uint16>& VariableJointIndices = LocalRigRuntimeContext->VariableJointIndicesPerLOD[CurrentLOD].Values;
+		JointsMapDNAIndicesToCompactPoseBoneIndices.Empty();
+		JointsMapDNAIndicesToCompactPoseBoneIndices.Reserve(VariableJointIndices.Num());
 		for (const uint16 JointIndex : VariableJointIndices)
 		{
 			const FMeshPoseBoneIndex MeshPoseBoneIndex = LocalDNAIndexMapping->JointsMapDNAIndicesToMeshPoseBoneIndices[JointIndex];
-			JointsMapDNAIndicesToCompactPoseBoneIndices.Add(RequiredBones.MakeCompactPoseIndex(MeshPoseBoneIndex));
+			const FCompactPoseBoneIndex CompactPoseBoneIndex = RequiredBones.MakeCompactPoseIndex(MeshPoseBoneIndex);
+			if (CompactPoseBoneIndex != INDEX_NONE) {
+				JointsMapDNAIndicesToCompactPoseBoneIndices.Add({JointIndex, CompactPoseBoneIndex});
+			}
 		}
 	}
 }
@@ -177,22 +181,16 @@ void FAnimNode_RigLogic::UpdateJoints(TArrayView<const uint16> VariableJointIndi
 
 	const float* N = NeutralJointValues.GetData();
 	const float* D = DeltaJointValues.GetData();
-	check(VariableJointIndices.Num() == JointsMapDNAIndicesToCompactPoseBoneIndices.Num());
-	for (int32 MappingIndex = 0; MappingIndex < JointsMapDNAIndicesToCompactPoseBoneIndices.Num(); ++MappingIndex)
+	for (const FJointCompactPoseBoneMapping& Mapping : JointsMapDNAIndicesToCompactPoseBoneIndices)
 	{
-		const uint16 JointIndex = VariableJointIndices[MappingIndex];
-		const FCompactPoseBoneIndex CompactPoseBoneIndex = JointsMapDNAIndicesToCompactPoseBoneIndices[MappingIndex];
-		if (CompactPoseBoneIndex != INDEX_NONE)
-		{
-			const uint16 AttrIndex = JointIndex * ATTR_COUNT_PER_JOINT;
-			FTransform& CompactPose = OutputContext.Pose[CompactPoseBoneIndex];
-			// Translation: X = X, Y = -Y, Z = Z
-			CompactPose.SetTranslation(FVector((N[AttrIndex + 0] + D[AttrIndex + 0]), -(N[AttrIndex + 1] + D[AttrIndex + 1]), (N[AttrIndex + 2] + D[AttrIndex + 2])));
-			// Rotation: X = -Y, Y = -Z, Z = X
-			CompactPose.SetRotation(FQuat(FRotator(-N[AttrIndex + 4], -N[AttrIndex + 5], N[AttrIndex + 3])) * FQuat(FRotator(-D[AttrIndex + 4], -D[AttrIndex + 5], D[AttrIndex + 3])));
-			// Scale: X = X, Y = Y, Z = Z
-			CompactPose.SetScale3D(FVector((N[AttrIndex + 6] + D[AttrIndex + 6]), (N[AttrIndex + 7] + D[AttrIndex + 7]), (N[AttrIndex + 8] + D[AttrIndex + 8])));
-		}
+		const uint16 AttrIndex = Mapping.JointIndex * ATTR_COUNT_PER_JOINT;
+		FTransform& CompactPose = OutputContext.Pose[Mapping.CompactPoseBoneIndex];
+		// Translation: X = X, Y = -Y, Z = Z
+		CompactPose.SetTranslation(FVector((N[AttrIndex + 0] + D[AttrIndex + 0]), -(N[AttrIndex + 1] + D[AttrIndex + 1]), (N[AttrIndex + 2] + D[AttrIndex + 2])));
+		// Rotation: X = -Y, Y = -Z, Z = X
+		CompactPose.SetRotation(FQuat(FRotator(-N[AttrIndex + 4], -N[AttrIndex + 5], N[AttrIndex + 3])) * FQuat(FRotator(-D[AttrIndex + 4], -D[AttrIndex + 5], D[AttrIndex + 3])));
+		// Scale: X = X, Y = Y, Z = Z
+		CompactPose.SetScale3D(FVector((N[AttrIndex + 6] + D[AttrIndex + 6]), (N[AttrIndex + 7] + D[AttrIndex + 7]), (N[AttrIndex + 8] + D[AttrIndex + 8])));
 	}
 }
 
diff --git a/Engine/Plugins/Animation/RigLogic/Source/RigLogicModule/Public/AnimNode_RigLogic.h b/Engine/Plugins/Animation/RigLogic/Source/RigLogicModule/Public/AnimNode_RigLogic.h
--- a/Engine/Plugins/Animation/RigLogic/Source/RigLogicModule/Public/AnimNode_RigLogic.h
+++ b/Engine/Plugins/Animation/RigLogic/Source/RigLogicModule/Public/AnimNode_RigLogic.h
@@ -43,6 +43,12 @@ private:
 	void UpdateBlendShapeCurves(const FDNAIndexMapping* DNAIndexMapping, TArrayView<const float> BlendShapeValues, FPoseContext& OutputContext);
 	void UpdateAnimMapCurves(const FDNAIndexMapping* DNAIndexMapping, TArrayView<const float> AnimMapOutputs, FPoseContext& OutputContext);
 
+private:
+	struct FJointCompactPoseBoneMapping {
+		uint16 JointIndex;
+		FCompactPoseBoneIndex CompactPoseBoneIndex;
+	};
+
 private:
 	/*
 	 * Max LOD level that post-process AnimBPs are evaluated.
@@ -52,8 +58,9 @@ private:
 	UPROPERTY(EditAnywhere, Category = RigLogic)
 	int32 LODThreshold = INDEX_NONE;
 
+private:
 	TSharedPtr<FSharedRigRuntimeContext> LocalRigRuntimeContext;
 	TSharedPtr<FDNAIndexMapping> LocalDNAIndexMapping;
 	FRigInstance* RigInstance;
-	TArray<FCompactPoseBoneIndex> JointsMapDNAIndicesToCompactPoseBoneIndices;
+	TArray<FJointCompactPoseBoneMapping> JointsMapDNAIndicesToCompactPoseBoneIndices;
 };
