Forward neural network mask curves to RigLogic from both anim node and control rig implementations

diff --git a/Engine/Plugins/Animation/RigLogic/Source/RigLogicModule/Private/AnimNode_RigLogic.cpp b/Engine/Plugins/Animation/RigLogic/Source/RigLogicModule/Private/AnimNode_RigLogic.cpp
--- a/Engine/Plugins/Animation/RigLogic/Source/RigLogicModule/Private/AnimNode_RigLogic.cpp
+++ b/Engine/Plugins/Animation/RigLogic/Source/RigLogicModule/Private/AnimNode_RigLogic.cpp
@@ -160,11 +160,16 @@ void FAnimNode_RigLogic::UpdateControlCurves(const FPoseContext& InputContext, c
 	UE::Anim::FNamedValueArrayUtils::Union(InputContext.Curve, DNAIndexMapping->ControlAttributeCurves,
 		[this](const UE::Anim::FCurveElement& InCurveElement, const UE::Anim::FCurveElementIndexed& InControlAttributeCurveElement, UE::Anim::ENamedValueUnionFlags InFlags)
 		{
-			if(InControlAttributeCurveElement.Index != INDEX_NONE)
- 			{
- 				RigInstance->SetRawControl(InControlAttributeCurveElement.Index, InCurveElement.Value);
- 			}
+			RigInstance->SetRawControl(InControlAttributeCurveElement.Index, InCurveElement.Value);
 		});
+
+	if (RigInstance->GetNeuralNetworkCount() != 0) {
+		UE::Anim::FNamedValueArrayUtils::Union(InputContext.Curve, DNAIndexMapping->NeuralNetworkMaskCurves,
+			[this](const UE::Anim::FCurveElement& InCurveElement, const UE::Anim::FCurveElementIndexed& InControlAttributeCurveElement, UE::Anim::ENamedValueUnionFlags InFlags)
+			{
+				RigInstance->SetNeuralNetworkMask(InControlAttributeCurveElement.Index, InCurveElement.Value);
+			});
+	}
 }
 
 void FAnimNode_RigLogic::CalculateRigLogic(FRigLogic* RigLogic)
diff --git a/Engine/Plugins/Animation/RigLogic/Source/RigLogicModule/Private/DNAAsset.cpp b/Engine/Plugins/Animation/RigLogic/Source/RigLogicModule/Private/DNAAsset.cpp
--- a/Engine/Plugins/Animation/RigLogic/Source/RigLogicModule/Private/DNAAsset.cpp
+++ b/Engine/Plugins/Animation/RigLogic/Source/RigLogicModule/Private/DNAAsset.cpp
@@ -198,6 +198,7 @@ TSharedPtr<FDNAIndexMapping> UDNAAsset::GetDNAIndexMapping(const USkeleton* Skel
 		DNAIndexMapping = MakeShared<FDNAIndexMapping>();
 		DNAIndexMapping->SkeletonGuid = SkeletonGuid;
 		DNAIndexMapping->MapControlCurves(BehaviorReader.Get(), Skeleton);
+		DNAIndexMapping->MapNeuralNetworkMaskCurves(BehaviorReader.Get(), Skeleton);
 		DNAIndexMapping->MapJoints(BehaviorReader.Get(), SkeletalMeshComponent);
 		DNAIndexMapping->MapMorphTargets(BehaviorReader.Get(), Skeleton, SkeletalMesh);
 		DNAIndexMapping->MapMaskMultipliers(BehaviorReader.Get(), Skeleton);
diff --git a/Engine/Plugins/Animation/RigLogic/Source/RigLogicModule/Private/DNAIndexMapping.cpp b/Engine/Plugins/Animation/RigLogic/Source/RigLogicModule/Private/DNAIndexMapping.cpp
--- a/Engine/Plugins/Animation/RigLogic/Source/RigLogicModule/Private/DNAIndexMapping.cpp
+++ b/Engine/Plugins/Animation/RigLogic/Source/RigLogicModule/Private/DNAIndexMapping.cpp
@@ -31,18 +31,18 @@ static FString CreateCurveName(const FString& NameToSplit, const FString& Format
 	return CurveName;
 }
 
-void FDNAIndexMapping::MapControlCurves(const IBehaviorReader* DNABehavior, const USkeleton* Skeleton)
+void FDNAIndexMapping::MapControlCurves(const IDNAReader* DNAReader, const USkeleton* Skeleton)
 {
 	LLM_SCOPE_BYNAME(TEXT("Animation/RigLogic"));
 
-	const uint32 ControlCount = DNABehavior->GetRawControlCount();
+	const uint32 ControlCount = DNAReader->GetRawControlCount();
 
 	ControlAttributeCurves.Empty();
 	ControlAttributeCurves.Reserve(ControlCount);
 
 	for (uint32_t ControlIndex = 0; ControlIndex < ControlCount; ++ControlIndex)
 	{
-		const FString DNAControlName = DNABehavior->GetRawControlName(ControlIndex);
+		const FString DNAControlName = DNAReader->GetRawControlName(ControlIndex);
 		const FString AnimatedControlName = CreateCurveName(DNAControlName, TEXT("<obj>_<attr>"));
 		if (AnimatedControlName == TEXT(""))
 		{
@@ -52,15 +52,42 @@ void FDNAIndexMapping::MapControlCurves(const IBehaviorReader* DNABehavior, cons
 	}
 }
 
-void FDNAIndexMapping::MapJoints(const IBehaviorReader* DNABehavior, const USkeletalMeshComponent* SkeletalMeshComponent)
+void FDNAIndexMapping::MapNeuralNetworkMaskCurves(const IDNAReader* DNAReader, const USkeleton* Skeleton)
 {
 	LLM_SCOPE_BYNAME(TEXT("Animation/RigLogic"));
 
-	const uint16 JointCount = DNABehavior->GetJointCount();
+	const uint16 NeuralNetworkCount = DNAReader->GetNeuralNetworkCount();
+
+	NeuralNetworkMaskCurves.Empty();
+	NeuralNetworkMaskCurves.Reserve(NeuralNetworkCount);
+
+	const uint16 MeshCount = DNAReader->GetMeshCount();
+	for (uint16 MeshIndex = 0; MeshIndex < MeshCount; ++MeshIndex)
+	{
+		const uint16 MeshRegionCount = DNAReader->GetMeshRegionCount(MeshIndex);
+		for (uint16 RegionIndex = 0; RegionIndex < MeshRegionCount; ++RegionIndex)
+		{
+			const FString& MeshRegionName = DNAReader->GetMeshRegionName(MeshIndex, RegionIndex);
+			TArrayView<const uint16> NeuralNetworkIndices = DNAReader->GetNeuralNetworkIndicesForMeshRegion(MeshIndex, RegionIndex);
+			const FString MaskCurveName = TEXT("CTRL_ML_") + MeshRegionName;
+			for (const auto NeuralNetworkIndex : NeuralNetworkIndices)
+			{
+				NeuralNetworkMaskCurves.Add(*MaskCurveName, NeuralNetworkIndex);
+			}
+		}
+	}
+}
+
+
+void FDNAIndexMapping::MapJoints(const IDNAReader* DNAReader, const USkeletalMeshComponent* SkeletalMeshComponent)
+{
+	LLM_SCOPE_BYNAME(TEXT("Animation/RigLogic"));
+
+	const uint16 JointCount = DNAReader->GetJointCount();
 	JointsMapDNAIndicesToMeshPoseBoneIndices.Reset(JointCount);
 	for (uint16 JointIndex = 0; JointIndex < JointCount; ++JointIndex)
 	{
-		const FString JointName = DNABehavior->GetJointName(JointIndex);
+		const FString JointName = DNAReader->GetJointName(JointIndex);
 		const FName BoneName = FName(*JointName);
 		const int32 BoneIndex = SkeletalMeshComponent->GetBoneIndex(BoneName);
 		// BoneIndex may be INDEX_NONE, but it's handled properly by the Evaluate method
@@ -68,11 +95,11 @@ void FDNAIndexMapping::MapJoints(const IBehaviorReader* DNABehavior, const USkel
 	}
 }
 
-void FDNAIndexMapping::MapMorphTargets(const IBehaviorReader* DNABehavior, const USkeleton* Skeleton, const USkeletalMesh* SkeletalMesh)
+void FDNAIndexMapping::MapMorphTargets(const IDNAReader* DNAReader, const USkeleton* Skeleton, const USkeletalMesh* SkeletalMesh)
 {
 	LLM_SCOPE_BYNAME(TEXT("Animation/RigLogic"));
 
-	const uint16 LODCount = DNABehavior->GetLODCount();
+	const uint16 LODCount = DNAReader->GetLODCount();
 	const TMap<FName, int32>& MorphTargetIndexMap = SkeletalMesh->GetMorphTargetIndexMap();
 	const TArray<UMorphTarget*>& MorphTargets = SkeletalMesh->GetMorphTargets();
 
@@ -81,15 +108,15 @@ void FDNAIndexMapping::MapMorphTargets(const IBehaviorReader* DNABehavior, const
 
 	for (uint16 LODIndex = 0; LODIndex < LODCount; ++LODIndex)
 	{
-		TArrayView<const uint16> MappingIndicesForLOD = DNABehavior->GetMeshBlendShapeChannelMappingIndicesForLOD(LODIndex);
+		TArrayView<const uint16> MappingIndicesForLOD = DNAReader->GetMeshBlendShapeChannelMappingIndicesForLOD(LODIndex);
 
 		MorphTargetCurvesPerLOD[LODIndex].Reserve(MappingIndicesForLOD.Num());
 
 		for (uint16 MappingIndex : MappingIndicesForLOD)
 		{
-			const FMeshBlendShapeChannelMapping Mapping = DNABehavior->GetMeshBlendShapeChannelMapping(MappingIndex);
-			const FString MeshName = DNABehavior->GetMeshName(Mapping.MeshIndex);
-			const FString BlendShapeName = DNABehavior->GetBlendShapeChannelName(Mapping.BlendShapeChannelIndex);
+			const FMeshBlendShapeChannelMapping Mapping = DNAReader->GetMeshBlendShapeChannelMapping(MappingIndex);
+			const FString MeshName = DNAReader->GetMeshName(Mapping.MeshIndex);
+			const FString BlendShapeName = DNAReader->GetBlendShapeChannelName(Mapping.BlendShapeChannelIndex);
 			const FString MorphTargetStr = MeshName + TEXT("__") + BlendShapeName;
 			const FName MorphTargetName(*MorphTargetStr);
 			const int32* MorphTargetIndex = MorphTargetIndexMap.Find(MorphTargetName);
@@ -102,24 +129,24 @@ void FDNAIndexMapping::MapMorphTargets(const IBehaviorReader* DNABehavior, const
 	}
 }
 
-void FDNAIndexMapping::MapMaskMultipliers(const IBehaviorReader* DNABehavior, const USkeleton* Skeleton)
+void FDNAIndexMapping::MapMaskMultipliers(const IDNAReader* DNAReader, const USkeleton* Skeleton)
 {
 	LLM_SCOPE_BYNAME(TEXT("Animation/RigLogic"));
 
-	const uint16 LODCount = DNABehavior->GetLODCount();
+	const uint16 LODCount = DNAReader->GetLODCount();
 
 	MaskMultiplierCurvesPerLOD.Reset();
 	MaskMultiplierCurvesPerLOD.AddDefaulted(LODCount);
 
 	for (uint16 LODIndex = 0; LODIndex < LODCount; ++LODIndex)
 	{
-		TArrayView<const uint16> IndicesPerLOD = DNABehavior->GetAnimatedMapIndicesForLOD(LODIndex);
+		TArrayView<const uint16> IndicesPerLOD = DNAReader->GetAnimatedMapIndicesForLOD(LODIndex);
 
 		MaskMultiplierCurvesPerLOD[LODIndex].Reserve(IndicesPerLOD.Num());
 
 		for (uint16 AnimMapIndex : IndicesPerLOD)
 		{
-			const FString AnimatedMapName = DNABehavior->GetAnimatedMapName(AnimMapIndex);
+			const FString AnimatedMapName = DNAReader->GetAnimatedMapName(AnimMapIndex);
 			const FString MaskMultiplierNameStr = CreateCurveName(AnimatedMapName, TEXT("<obj>_<attr>"));
 			if (MaskMultiplierNameStr == "")
 			{
diff --git a/Engine/Plugins/Animation/RigLogic/Source/RigLogicModule/Private/DNAIndexMapping.h b/Engine/Plugins/Animation/RigLogic/Source/RigLogicModule/Private/DNAIndexMapping.h
--- a/Engine/Plugins/Animation/RigLogic/Source/RigLogicModule/Private/DNAIndexMapping.h
+++ b/Engine/Plugins/Animation/RigLogic/Source/RigLogicModule/Private/DNAIndexMapping.h
@@ -9,7 +9,7 @@
 
 #include "DNAIndexMapping.generated.h"
 
-class IBehaviorReader;
+class IDNAReader;
 class USkeleton;
 class USkeletalMesh;
 class USkeletalMeshComponent;
@@ -28,14 +28,16 @@ struct FDNAIndexMapping
 
 	// all the control attributes that we will need to extract, alongside their control index
 	FCachedIndexedCurve ControlAttributeCurves;
+	FCachedIndexedCurve NeuralNetworkMaskCurves;
 	TArray<FMeshPoseBoneIndex> JointsMapDNAIndicesToMeshPoseBoneIndices;
 	TArray<FCachedIndexedCurve> MorphTargetCurvesPerLOD;
 	TArray<FCachedIndexedCurve> MaskMultiplierCurvesPerLOD;
 
-	void MapControlCurves(const IBehaviorReader* DNABehavior, const USkeleton* Skeleton);
-	void MapJoints(const IBehaviorReader* DNABehavior, const USkeletalMeshComponent* SkeletalMeshComponent);
-	void MapMorphTargets(const IBehaviorReader* DNABehavior, const USkeleton* Skeleton, const USkeletalMesh* SkeletalMesh);
-	void MapMaskMultipliers(const IBehaviorReader* DNABehavior, const USkeleton* Skeleton);
+	void MapControlCurves(const IDNAReader* DNAReader, const USkeleton* Skeleton);
+	void MapNeuralNetworkMaskCurves(const IDNAReader* DNAReader, const USkeleton* Skeleton);
+	void MapJoints(const IDNAReader* DNAReader, const USkeletalMeshComponent* SkeletalMeshComponent);
+	void MapMorphTargets(const IDNAReader* DNAReader, const USkeleton* Skeleton, const USkeletalMesh* SkeletalMesh);
+	void MapMaskMultipliers(const IDNAReader* DNAReader, const USkeleton* Skeleton);
 
 };
 
diff --git a/Engine/Plugins/Animation/RigLogic/Source/RigLogicModule/Private/RigUnit_RigLogic.cpp b/Engine/Plugins/Animation/RigLogic/Source/RigLogicModule/Private/RigUnit_RigLogic.cpp
--- a/Engine/Plugins/Animation/RigLogic/Source/RigLogicModule/Private/RigUnit_RigLogic.cpp
+++ b/Engine/Plugins/Animation/RigLogic/Source/RigLogicModule/Private/RigUnit_RigLogic.cpp
@@ -67,6 +67,7 @@ FRigUnit_RigLogic_Data& FRigUnit_RigLogic_Data::operator=(const FRigUnit_RigLogi
 	LocalRigRuntimeContext = nullptr;
 	RigInstance = nullptr;
 	InputCurveIndices = Other.InputCurveIndices;
+	NeuralNetMaskCurveIndices = Other.NeuralNetMaskCurveIndices;
 	HierarchyBoneIndices = Other.HierarchyBoneIndices;
 	MorphTargetCurveIndices = Other.MorphTargetCurveIndices;
 	BlendShapeIndices = Other.BlendShapeIndices;
@@ -106,6 +107,7 @@ void FRigUnit_RigLogic_Data::InitializeRigLogic(const URigHierarchy* InHierarchy
 
 		MapJoints(InHierarchy);
 		MapInputCurveIndices(InHierarchy);
+		MapNeuralNetMaskCurveIndices(InHierarchy);
 		MapMorphTargets(InHierarchy);
 		MapMaskMultipliers(InHierarchy);
 	}
@@ -134,6 +136,34 @@ void FRigUnit_RigLogic_Data::MapInputCurveIndices(const URigHierarchy* InHierarc
 	}
 }
 
+void FRigUnit_RigLogic_Data::MapNeuralNetMaskCurveIndices(const URigHierarchy* InHierarchy)
+{
+	DECLARE_SCOPE_HIERARCHICAL_COUNTER_FUNC()
+	QUICK_SCOPE_CYCLE_COUNTER(STAT_RigUnit_RigLogic_MapInputCurveIndices);
+
+	const IDNAReader* DNABehavior = LocalRigRuntimeContext->BehaviorReader.Get();
+	const uint16 NeuralNetworkCount = DNABehavior->GetNeuralNetworkCount();
+	NeuralNetMaskCurveIndices.SetNum(NeuralNetworkCount);
+
+	const uint16 MeshCount = DNABehavior->GetMeshCount();
+	for (uint16 MeshIndex = 0; MeshIndex < MeshCount; ++MeshIndex)
+	{
+		const uint16 MeshRegionCount = DNABehavior->GetMeshRegionCount(MeshIndex);
+		for (uint16 RegionIndex = 0; RegionIndex < MeshRegionCount; ++RegionIndex)
+		{
+			const FString& MeshRegionName = DNABehavior->GetMeshRegionName(MeshIndex, RegionIndex);
+			TArrayView<const uint16> NeuralNetworkIndices = DNABehavior->GetNeuralNetworkIndicesForMeshRegion(MeshIndex, RegionIndex);
+			const FString MaskCurveName = TEXT("CTRL_ML_") + MeshRegionName;
+			const FName CurveFName(*MaskCurveName);
+			const int32 CurveIndex = InHierarchy ? InHierarchy->GetIndex(FRigElementKey(CurveFName, ERigElementType::Curve)) : INDEX_NONE;
+			for (const auto NeuralNetworkIndex : NeuralNetworkIndices)
+			{
+				NeuralNetMaskCurveIndices[NeuralNetworkIndex] = CurveIndex;  // Can be INDEX_NONE
+			}
+		}
+	}
+}
+
 void FRigUnit_RigLogic_Data::MapJoints(const URigHierarchy* Hierarchy)
 {
 	DECLARE_SCOPE_HIERARCHICAL_COUNTER_FUNC()
@@ -236,6 +266,17 @@ void FRigUnit_RigLogic_Data::CalculateRigLogic(const URigHierarchy* InHierarchy)
 		RigInstance->SetRawControl(ControlIndex, Value);
 	}
 
+	const int32 NeuralNetworkCount = RigInstance->GetNeuralNetworkCount();
+	for (int32 NeuralNetworkIndex = 0; NeuralNetworkIndex < NeuralNetworkCount; ++NeuralNetworkIndex)
+	{
+		const uint32 CurveIndex = NeuralNetMaskCurveIndices[NeuralNetworkIndex];
+		if (InHierarchy->IsCurveValueSetByIndex(CurveIndex))
+		{
+			const float Value = InHierarchy->GetCurveValue(CurveIndex);
+			RigInstance->SetNeuralNetworkMask(NeuralNetworkIndex, Value);
+		}
+	}
+
 	LocalRigRuntimeContext->RigLogic->Calculate(RigInstance.Get());
 }
 
diff --git a/Engine/Plugins/Animation/RigLogic/Source/RigLogicModule/Public/RigUnit_RigLogic.h b/Engine/Plugins/Animation/RigLogic/Source/RigLogicModule/Public/RigUnit_RigLogic.h
--- a/Engine/Plugins/Animation/RigLogic/Source/RigLogicModule/Public/RigUnit_RigLogic.h
+++ b/Engine/Plugins/Animation/RigLogic/Source/RigLogicModule/Public/RigUnit_RigLogic.h
@@ -61,6 +61,10 @@ struct FRigUnit_RigLogic_Data
 	UPROPERTY(transient)
 	TArray<int32> InputCurveIndices;
 
+	/** RL neural network mask index to ControlRig's input curve index **/
+	UPROPERTY(transient)
+	TArray<int32> NeuralNetMaskCurveIndices;
+
 	/** RL joint index to ControlRig's hierarchy bone index **/
 	UPROPERTY(transient)
 	TArray<int32> HierarchyBoneIndices;
@@ -92,6 +96,8 @@ struct FRigUnit_RigLogic_Data
 
 	/** Makes a map of input curve indices from DNA file to the control rig curves **/
 	void MapInputCurveIndices(const URigHierarchy* InHierarchy);
+	/** Makes a map of neural net mask indices from DNA file to the control rig curves **/
+	void MapNeuralNetMaskCurveIndices(const URigHierarchy* InHierarchy);
 	/** Uses names to map joint indices from DNA file to the indices of bones in control rig hierarchy **/
 	void MapJoints(const URigHierarchy* Hierarchy);
 	/** Uses names of blend shapes and meshes in DNA file, for all LODs, to map their indices to the indices of
