Updated MetaHuman source data regeneration in RilLogic plugin

diff --git a/Engine/Plugins/Animation/RigLogic/Source/RigLogicModule/Private/SkelMeshDNAUtils.cpp b/Engine/Plugins/Animation/RigLogic/Source/RigLogicModule/Private/SkelMeshDNAUtils.cpp
--- a/Engine/Plugins/Animation/RigLogic/Source/RigLogicModule/Private/SkelMeshDNAUtils.cpp
+++ b/Engine/Plugins/Animation/RigLogic/Source/RigLogicModule/Private/SkelMeshDNAUtils.cpp
@@ -458,7 +458,9 @@ void USkelMeshDNAUtils::UpdateJointBehavior(USkeletalMeshComponent* InSkelMeshCo
 	InSkelMeshComponent->InitAnim(true);
 }
 
-void USkelMeshDNAUtils::UpdateSourceData(USkeletalMesh* InSkelMesh)
+// BEGIN: MHC MODIFICATION
+void USkelMeshDNAUtils::RecreateSourceData(USkeletalMesh* InSkelMesh)
+// END: MHC MODIFICATION
 {
 	IMeshUtilities& MeshUtilities = FModuleManager::Get().LoadModuleChecked<IMeshUtilities>("MeshUtilities");
 	// Source data must be updated during cooking.
@@ -466,7 +468,9 @@ void USkelMeshDNAUtils::UpdateSourceData(USkeletalMesh* InSkelMesh)
 	MeshUtilities.CreateImportDataFromLODModel(InSkelMesh);
 }
 
-void USkelMeshDNAUtils::UpdateSourceData(USkeletalMesh* InSkelMesh, class IDNAReader* InDNAReader, class FDNAToSkelMeshMap* InDNAToSkelMeshMap)
+// BEGIN: MHC MODIFICATION
+void USkelMeshDNAUtils::UpdateSourceData(USkeletalMesh* InSkelMesh)
+// END: MHC MODIFICATION
 {
 	FSkeletalMeshModel* ImportedModel = InSkelMesh->GetImportedModel();
 	const int32 LODCount = ImportedModel->LODModels.Num();
@@ -483,54 +487,79 @@ void USkelMeshDNAUtils::UpdateSourceData(USkeletalMesh* InSkelMesh, class IDNARe
 		TArray<FSoftSkinVertex> LODVertices;
 		LODModel.GetVertices(LODVertices);
 
-		TArray<SkeletalMeshImportData::FRawBoneInfluence> NewInfluences;
-		TArray<bool> HasOverlappingVertices;
-		HasOverlappingVertices.AddZeroed(LODMeshVtxCount);
+// BEGIN: MHC MODIFICATION
+		ImportData.Influences.Empty();
 		for (int32 LODMeshVtxIndex = 0; LODMeshVtxIndex < LODMeshVtxCount; LODMeshVtxIndex++)
 		{
-			// Update points.
+			// Update points but keep vertex order in the original MeshToImportVertex map to preserve DNA vertex ordering in the original FBX.
 			int32 FbxVertexIndex = LODModel.MeshToImportVertexMap[LODMeshVtxIndex];
-			if (!HasOverlappingVertices[FbxVertexIndex])
+
+			if (FbxVertexIndex <= LODModel.MaxImportVertex)
 			{
-				HasOverlappingVertices[FbxVertexIndex] = true;
-				if (FbxVertexIndex <= LODModel.MaxImportVertex)
+				ImportData.Points[FbxVertexIndex] = LODVertices[LODMeshVtxIndex].Position;
+			}
+
+			// Update influences.
+			int32 SectionIdx;
+			int32 VertexIdx;
+			LODModel.GetSectionFromVertexIndex(LODMeshVtxIndex, SectionIdx, VertexIdx);
+			const FSoftSkinVertex& Vertex = LODModel.Sections[SectionIdx].SoftVertices[VertexIdx];
+			for (int32 InfluenceIndex = 0; InfluenceIndex < MAX_TOTAL_INFLUENCES; ++InfluenceIndex)
+			{
+				// Have we reached the end of valid weights?
+				if (Vertex.InfluenceWeights[InfluenceIndex] == 0)
 				{
-					ImportData.Points[FbxVertexIndex] = LODVertices[LODMeshVtxIndex].Position;
+					break;
 				}
+				SkeletalMeshImportData::FRawBoneInfluence& Influence = ImportData.Influences.AddZeroed_GetRef();
+				Influence.VertexIndex = FbxVertexIndex;
+				Influence.BoneIndex = LODModel.Sections[SectionIdx].BoneMap[Vertex.InfluenceBones[InfluenceIndex]];
+				Influence.Weight = static_cast<float>(Vertex.InfluenceWeights[InfluenceIndex]) / 65535.0f;
+			}
+		}
 
-				// Update influences.
-				int32 SectionIdx;
-				int32 VertexIdx;
-				LODModel.GetSectionFromVertexIndex(LODMeshVtxIndex, SectionIdx, VertexIdx);
-				if (LODModel.Sections[SectionIdx].SoftVertices[VertexIdx].Color.B != 0)
-				{
-					int32 DNAMeshIndex = InDNAToSkelMeshMap->ImportVtxToDNAMeshIndex[LODIndex][LODMeshVtxIndex];
-					int32 DNAVertexIndex = InDNAToSkelMeshMap->ImportVtxToDNAVtxIndex[LODIndex][LODMeshVtxIndex];
+		// Update tangents.
+		ImportData.NumTexCoords = FMath::Min(LODModel.NumTexCoords, static_cast<uint32>(MAX_TEXCOORDS));
 
-					if (DNAVertexIndex >= 0)
-					{
-						TArrayView<const float> DNASkinWeights = InDNAReader->GetSkinWeightsValues(DNAMeshIndex, DNAVertexIndex);
-						TArrayView<const uint16> DNASkinJoints = InDNAReader->GetSkinWeightsJointIndices(DNAMeshIndex, DNAVertexIndex);
-						uint16 SkinJointNum = DNASkinJoints.Num();
-						for (uint16 InfluenceIndex = 0; InfluenceIndex < SkinJointNum; ++InfluenceIndex)
-						{
-							float InfluenceWeight = DNASkinWeights[InfluenceIndex];
-							int32 UpdatedBoneId = InDNAToSkelMeshMap->GetUEBoneIndex(DNASkinJoints[InfluenceIndex]);
-
-							SkeletalMeshImportData::FRawBoneInfluence Influence;
-							Influence.VertexIndex = FbxVertexIndex;
-							Influence.BoneIndex = UpdatedBoneId;
-							Influence.Weight = InfluenceWeight;
-							NewInfluences.Add(Influence);
-						}
-						ImportData.Influences.RemoveAll([FbxVertexIndex](const SkeletalMeshImportData::FRawBoneInfluence& BoneInfluence) { return FbxVertexIndex == BoneInfluence.VertexIndex; });
-					}
+		const int32 VertexCount = LODVertices.Num();
+		const int32 TriangleCount = LODModel.IndexBuffer.Num() / 3;
+		ImportData.Wedges.Empty(TriangleCount);
+		ImportData.Faces.Empty(VertexCount);
+		for (int32 TriangleIndex = 0; TriangleIndex < TriangleCount; ++TriangleIndex)
+		{
+			SkeletalMeshImportData::FTriangle& Face = ImportData.Faces.AddZeroed_GetRef();
+			int32 IndexBufferIndex = TriangleIndex * 3;
+			int32 SectionIndex;
+			int32 VerticeIndex[3];
+			int32 SectionVertexIndex[3];
+			for (int32 Corner = 0; Corner < 3; ++Corner)
+			{
+				const int32 WedgeIndex = ImportData.Wedges.Num();
+				VerticeIndex[Corner] = LODModel.IndexBuffer[IndexBufferIndex + Corner];
+				LODModel.GetSectionFromVertexIndex(VerticeIndex[Corner], SectionIndex, SectionVertexIndex[Corner]);
+				const FSoftSkinVertex& Vertex = LODVertices[VerticeIndex[Corner]];
+				SkeletalMeshImportData::FVertex& Wedge = ImportData.Wedges.AddZeroed_GetRef();
+				Wedge.Color = Vertex.Color;
+				for (int32 UVIndex = 0; UVIndex < static_cast<int32>(ImportData.NumTexCoords); ++UVIndex)
+				{
+					Wedge.UVs[UVIndex] = Vertex.UVs[UVIndex];
 				}
+				Wedge.VertexIndex = LODModel.MeshToImportVertexMap[VerticeIndex[Corner]];
+				Wedge.MatIndex = LODModel.Sections[SectionIndex].MaterialIndex;
+
+				Face.WedgeIndex[Corner] = WedgeIndex;
+				Face.TangentX[Corner] = Vertex.TangentX;
+				Face.TangentY[Corner] = Vertex.TangentY;
+				Face.TangentZ[Corner] = Vertex.TangentZ;
 			}
+			Face.MatIndex = LODModel.Sections[SectionIndex].MaterialIndex;
+			//Faceted by default
+			Face.SmoothingGroups = 0x0;
 		}
-		ImportData.Influences.Append(NewInfluences);
-		// Sort influences by vertex index.
-		FLODUtilities::ProcessImportMeshInfluences(ImportData.Wedges.Num(), ImportData.Influences, InSkelMesh->GetPathName());
+
+		//Recreate smooth group
+		ImportData.ComputeSmoothGroupFromNormals();
+// END: MHC MODIFICATION
 
 		// Update reference pose.
 		const int32 JointCount = LODModel.RequiredBones.Num();
diff --git a/Engine/Plugins/Animation/RigLogic/Source/RigLogicModule/Public/SkelMeshDNAUtils.h b/Engine/Plugins/Animation/RigLogic/Source/RigLogicModule/Public/SkelMeshDNAUtils.h
--- a/Engine/Plugins/Animation/RigLogic/Source/RigLogicModule/Public/SkelMeshDNAUtils.h
+++ b/Engine/Plugins/Animation/RigLogic/Source/RigLogicModule/Public/SkelMeshDNAUtils.h
@@ -56,10 +56,12 @@ public:
 	/** Gets the DNA asset embedded in the mesh */
 	static UDNAAsset* GetMeshDNA(USkeletalMesh* InSkelMesh);
 
-	/** Updates source skeleton data for the purpose of character cooking and export.  */
-	static void UpdateSourceData(USkeletalMesh* InSkelMesh);
+// BEGIN: MHC MODIFICATION
+	/** Recreates source skeleton data for the purpose of character cooking and export.  */
+	static void RecreateSourceData(USkeletalMesh* InSkelMesh);
 	/** Updates ImportData structs for SkeletalMesh in non destructive way to preserve the original mesh */
-	static void UpdateSourceData(USkeletalMesh* InSkelMesh, class IDNAReader* InDNAReader, class FDNAToSkelMeshMap* InDNAToSkelMeshMap);
+	static void UpdateSourceData(USkeletalMesh* InSkelMesh);
+// END: MHC MODIFICATION
 #endif // WITH_EDITORONLY_DATA
 	/** Converts DNA vertex coordinates to UE4 coordinate system **/
 	inline static FVector ConvertDNAVertexToUE4CoordSystem(FVector InVertexPositionInDNA)
