Handle redirected/renamed asset or asset saved with unloaded plugin RigLogicEditor
* Crash report has mangled callstack but this is likely cause.
* Problem is that AssetData::GetClass can return nullptr if the asset was renamed or otherwise has a redirector ~or~ if the asset was saved by a plugin that is not loaded
* This CL should follow the redirector by way of EResolveClass::Yes to handle the redirector case
* The fix will also prevent the menu from being created if the dependent plugin is not found.

diff --git a/Engine/Plugins/Animation/RigLogic/Source/RigLogicEditor/Private/RigLogicEditor.cpp b/Engine/Plugins/Animation/RigLogic/Source/RigLogicEditor/Private/RigLogicEditor.cpp
--- a/Engine/Plugins/Animation/RigLogic/Source/RigLogicEditor/Private/RigLogicEditor.cpp
+++ b/Engine/Plugins/Animation/RigLogic/Source/RigLogicEditor/Private/RigLogicEditor.cpp
@@ -53,8 +53,9 @@ TSharedRef<FExtender> FRigLogicEditor::OnExtendSkelMeshWithDNASelectionMenu(cons
 void FRigLogicEditor::CreateDnaActionsSubMenu(FMenuBuilder& MenuBuilder, const TArray<FAssetData> SelectedAssets)
 {
 	const FAssetData& Asset = SelectedAssets[0];
+	const UClass* AssetClass = Asset.GetClass(EResolveClass::Yes);
 
-	if (Asset.GetClass()->IsChildOf(USkeletalMesh::StaticClass()))
+	if (AssetClass != nullptr && AssetClass->IsChildOf(USkeletalMesh::StaticClass()))
 	{
 		MenuBuilder.AddSubMenu(
 			LOCTEXT("DNASkeletalMeshSubmenu", "MetaHuman DNA"),
