Add missing null checks when generating context menus by RigLogicEditor plugin

diff --git a/Engine/Plugins/Animation/RigLogic/Source/RigLogicEditor/Private/RigLogicEditor.cpp b/Engine/Plugins/Animation/RigLogic/Source/RigLogicEditor/Private/RigLogicEditor.cpp
--- a/Engine/Plugins/Animation/RigLogic/Source/RigLogicEditor/Private/RigLogicEditor.cpp
+++ b/Engine/Plugins/Animation/RigLogic/Source/RigLogicEditor/Private/RigLogicEditor.cpp
@@ -52,10 +52,11 @@ TSharedRef<FExtender> FRigLogicEditor::OnExtendSkelMeshWithDNASelectionMenu(cons
 
 void FRigLogicEditor::CreateDnaActionsSubMenu(FMenuBuilder& MenuBuilder, const TArray<FAssetData> SelectedAssets)
 {
+	check(!SelectedAssets.IsEmpty());
 	const FAssetData& Asset = SelectedAssets[0];
-	const UClass* AssetClass = Asset.GetClass(EResolveClass::Yes);
+	UClass* AssetClass = Asset.GetClass();
 
-	if (AssetClass != nullptr && AssetClass->IsChildOf(USkeletalMesh::StaticClass()))
+	if ((AssetClass != nullptr) && AssetClass->IsChildOf(USkeletalMesh::StaticClass()))
 	{
 		MenuBuilder.AddSubMenu(
 			LOCTEXT("DNASkeletalMeshSubmenu", "MetaHuman DNA"),
@@ -70,6 +71,7 @@ void FRigLogicEditor::CreateDnaActionsSubMenu(FMenuBuilder& MenuBuilder, const T
 
 void FRigLogicEditor::GetDNAMenu(FMenuBuilder& MenuBuilder, const TArray<FAssetData> SelectedAssets)
 {
+	check(!SelectedAssets.IsEmpty());
 	auto Mesh = SelectedAssets[0].GetAsset();
 
 	MenuBuilder.AddMenuEntry(
@@ -146,7 +148,7 @@ void FRigLogicEditor::ExecuteDNAReimport(class UObject* Mesh)
 
 void FRigLogicEditor::GetAssetRegistryTagsForDNA(const class UObject* Object, TArray<UObject::FAssetRegistryTag>& OutTags)
 {
-	if (Object->GetClass()->IsChildOf(USkeletalMesh::StaticClass()))
+	if ((Object != nullptr) && (Object->GetClass() != nullptr) && Object->GetClass()->IsChildOf(USkeletalMesh::StaticClass()))
 	{
 		FString DNAname = (LOCTEXT("DnaNotOnSkeletalMesh", "No DNA Attached")).ToString();
 		USkeletalMesh* SkelMesh = const_cast<USkeletalMesh*>(Cast<USkeletalMesh>(Object));
