Update RigLogic to v9.2.2

diff --git a/Engine/Plugins/Animation/RigLogic/RigLogic.uplugin b/Engine/Plugins/Animation/RigLogic/RigLogic.uplugin
--- a/Engine/Plugins/Animation/RigLogic/RigLogic.uplugin
+++ b/Engine/Plugins/Animation/RigLogic/RigLogic.uplugin
@@ -1,9 +1,9 @@
 {
 	"FileVersion": 3,
-	"Version": 9.1,
-	"VersionName": "9.1.3",
-	"FriendlyName": "RigLogic Plugin v9.1.3",
-	"Description": "3Lateral RigLogic Plugin for Facial Animation v9.1.3",
+	"Version": 9.2,
+	"VersionName": "9.2.2",
+	"FriendlyName": "RigLogic Plugin v9.2.2",
+	"Description": "3Lateral RigLogic Plugin for Facial Animation v9.2.2",
 	"Category": "Animation",
 	"CreatedBy": "3Lateral",
 	"CreatedByURL": "https://www.3lateral.com/",
diff --git a/Engine/Plugins/Animation/RigLogic/Source/RigLogicLib/Private/dna/filters/AnimatedMapFilter.cpp b/Engine/Plugins/Animation/RigLogic/Source/RigLogicLib/Private/dna/filters/AnimatedMapFilter.cpp
--- a/Engine/Plugins/Animation/RigLogic/Source/RigLogicLib/Private/dna/filters/AnimatedMapFilter.cpp
+++ b/Engine/Plugins/Animation/RigLogic/Source/RigLogicLib/Private/dna/filters/AnimatedMapFilter.cpp
@@ -18,6 +18,7 @@ AnimatedMapFilter::AnimatedMapFilter(MemoryResource* memRes_) :
 void AnimatedMapFilter::configure(std::uint16_t animatedMapCount, UnorderedSet<std::uint16_t> allowedAnimatedMapIndices) {
     passingIndices = std::move(allowedAnimatedMapIndices);
     // Fill the structure that maps indices prior to deletion to indices after deletion
+    remappedIndices.clear();
     remap(animatedMapCount, passingIndices, remappedIndices);
 }
 
diff --git a/Engine/Plugins/Animation/RigLogic/Source/RigLogicLib/Private/dna/filters/BlendShapeFilter.cpp b/Engine/Plugins/Animation/RigLogic/Source/RigLogicLib/Private/dna/filters/BlendShapeFilter.cpp
--- a/Engine/Plugins/Animation/RigLogic/Source/RigLogicLib/Private/dna/filters/BlendShapeFilter.cpp
+++ b/Engine/Plugins/Animation/RigLogic/Source/RigLogicLib/Private/dna/filters/BlendShapeFilter.cpp
@@ -18,6 +18,7 @@ BlendShapeFilter::BlendShapeFilter(MemoryResource* memRes_) :
 void BlendShapeFilter::configure(std::uint16_t blendShapeCount, UnorderedSet<std::uint16_t> allowedBlendShapeIndices) {
     passingIndices = std::move(allowedBlendShapeIndices);
     // Fill the structure that maps indices prior to deletion to indices after deletion
+    remappedIndices.clear();
     remap(blendShapeCount, passingIndices, remappedIndices);
 }
 
@@ -37,8 +38,24 @@ void BlendShapeFilter::apply(RawDefinition& dest) {
     dest.meshBlendShapeChannelMapping.updateTo(remappedIndices);
 }
 
+void BlendShapeFilter::apply(RawMesh& dest) {
+    // Remove blend shape targets of blend shapes to remove
+    extd::filter(dest.blendShapeTargets, [this](const RawBlendShapeTarget& bsTarget, std::size_t  /*unused*/) {
+            return passes(bsTarget.blendShapeChannelIndex);
+        });
+
+    // Remap blend shape targets
+    for (auto& bsTarget : dest.blendShapeTargets) {
+        bsTarget.blendShapeChannelIndex = remapped(bsTarget.blendShapeChannelIndex);
+    }
+}
+
 bool BlendShapeFilter::passes(std::uint16_t index) const {
     return extd::contains(passingIndices, index);
 }
 
+std::uint16_t BlendShapeFilter::remapped(std::uint16_t index) const {
+    return remappedIndices.at(index);
+}
+
 }  // namespace dna
diff --git a/Engine/Plugins/Animation/RigLogic/Source/RigLogicLib/Private/dna/filters/BlendShapeFilter.h b/Engine/Plugins/Animation/RigLogic/Source/RigLogicLib/Private/dna/filters/BlendShapeFilter.h
--- a/Engine/Plugins/Animation/RigLogic/Source/RigLogicLib/Private/dna/filters/BlendShapeFilter.h
+++ b/Engine/Plugins/Animation/RigLogic/Source/RigLogicLib/Private/dna/filters/BlendShapeFilter.h
@@ -9,13 +9,16 @@
 namespace dna {
 
 struct RawDefinition;
+struct RawMesh;
 
 class BlendShapeFilter {
     public:
         explicit BlendShapeFilter(MemoryResource* memRes_);
         void configure(std::uint16_t blendShapeCount, UnorderedSet<std::uint16_t> allowedBlendShapeIndices);
         void apply(RawDefinition& dest);
+        void apply(RawMesh& dest);
         bool passes(std::uint16_t index) const;
+        std::uint16_t remapped(std::uint16_t index) const;
 
     private:
         MemoryResource* memRes;
diff --git a/Engine/Plugins/Animation/RigLogic/Source/RigLogicLib/Private/dna/filters/JointFilter.cpp b/Engine/Plugins/Animation/RigLogic/Source/RigLogicLib/Private/dna/filters/JointFilter.cpp
--- a/Engine/Plugins/Animation/RigLogic/Source/RigLogicLib/Private/dna/filters/JointFilter.cpp
+++ b/Engine/Plugins/Animation/RigLogic/Source/RigLogicLib/Private/dna/filters/JointFilter.cpp
@@ -21,6 +21,7 @@ void JointFilter::configure(std::uint16_t jointCount, UnorderedSet<std::uint16_t
     option = option_;
     passingIndices = std::move(allowedJointIndices);
     // Fill the structure that maps indices prior to deletion to indices after deletion
+    remappedIndices.clear();
     remap(jointCount, passingIndices, remappedIndices);
 }
 
diff --git a/Engine/Plugins/Animation/RigLogic/Source/RigLogicLib/Private/dna/filters/MeshFilter.cpp b/Engine/Plugins/Animation/RigLogic/Source/RigLogicLib/Private/dna/filters/MeshFilter.cpp
--- a/Engine/Plugins/Animation/RigLogic/Source/RigLogicLib/Private/dna/filters/MeshFilter.cpp
+++ b/Engine/Plugins/Animation/RigLogic/Source/RigLogicLib/Private/dna/filters/MeshFilter.cpp
@@ -18,6 +18,7 @@ MeshFilter::MeshFilter(MemoryResource* memRes_) :
 void MeshFilter::configure(std::uint16_t meshCount, UnorderedSet<std::uint16_t> allowedMeshIndices) {
     passingIndices = std::move(allowedMeshIndices);
     // Fill the structure that maps indices prior to deletion to indices after deletion
+    remappedIndices.clear();
     remap(meshCount, passingIndices, remappedIndices);
 }
 
diff --git a/Engine/Plugins/Animation/RigLogic/Source/RigLogicLib/Private/dna/stream/BinaryStreamReaderImpl.cpp b/Engine/Plugins/Animation/RigLogic/Source/RigLogicLib/Private/dna/stream/BinaryStreamReaderImpl.cpp
--- a/Engine/Plugins/Animation/RigLogic/Source/RigLogicLib/Private/dna/stream/BinaryStreamReaderImpl.cpp
+++ b/Engine/Plugins/Animation/RigLogic/Source/RigLogicLib/Private/dna/stream/BinaryStreamReaderImpl.cpp
@@ -21,15 +21,6 @@
 
 namespace dna {
 
-#ifdef __clang__
-    #pragma clang diagnostic push
-    #pragma clang diagnostic ignored "-Wglobal-constructors"
-#endif
-sc::StatusProvider BinaryStreamReaderImpl::status{SignatureMismatchError, VersionMismatchError, InvalidDataError};
-#ifdef __clang__
-    #pragma clang diagnostic pop
-#endif
-
 BinaryStreamReader::~BinaryStreamReader() = default;
 
 BinaryStreamReader* BinaryStreamReader::create(BoundedIOStream* stream,
diff --git a/Engine/Plugins/Animation/RigLogic/Source/RigLogicLib/Private/dna/stream/BinaryStreamReaderImpl.h b/Engine/Plugins/Animation/RigLogic/Source/RigLogicLib/Private/dna/stream/BinaryStreamReaderImpl.h
--- a/Engine/Plugins/Animation/RigLogic/Source/RigLogicLib/Private/dna/stream/BinaryStreamReaderImpl.h
+++ b/Engine/Plugins/Animation/RigLogic/Source/RigLogicLib/Private/dna/stream/BinaryStreamReaderImpl.h
@@ -6,12 +6,13 @@
 #include "dna/ReaderImpl.h"
 #include "dna/TypeDefs.h"
 #include "dna/stream/FilteredBinaryInputArchive.h"
+#include "dna/stream/StreamReaderStatus.h"
 
 #include <status/Provider.h>
 
 namespace dna {
 
-class BinaryStreamReaderImpl : public ReaderImpl<BinaryStreamReader> {
+class BinaryStreamReaderImpl : public ReaderImpl<BinaryStreamReader>, private StreamReaderStatus {
     public:
         BinaryStreamReaderImpl(BoundedIOStream* stream_,
                                DataLayer layer_,
@@ -27,8 +28,6 @@ class BinaryStreamReaderImpl : public ReaderImpl<BinaryStreamReader> {
         bool isLODConstrained() const;
 
     private:
-        static sc::StatusProvider status;
-
         BoundedIOStream* stream;
         FilteredBinaryInputArchive archive;
         bool lodConstrained;
diff --git a/Engine/Plugins/Animation/RigLogic/Source/RigLogicLib/Private/dna/stream/FilteredBinaryInputArchive.cpp b/Engine/Plugins/Animation/RigLogic/Source/RigLogicLib/Private/dna/stream/FilteredBinaryInputArchive.cpp
--- a/Engine/Plugins/Animation/RigLogic/Source/RigLogicLib/Private/dna/stream/FilteredBinaryInputArchive.cpp
+++ b/Engine/Plugins/Animation/RigLogic/Source/RigLogicLib/Private/dna/stream/FilteredBinaryInputArchive.cpp
@@ -37,6 +37,15 @@ static UnorderedMap<U, U> remappedPositions(const Vector<T>& target, const Unord
     return mapping;
 }
 
+template<typename TContainer>
+static void removeByIndex(TContainer& container, std::size_t index) {
+    assert(index < container.size());
+    for (std::size_t i = index; i < (container.size() - 1ul); ++i) {
+        container[i] = container[i + 1ul];
+    }
+    container.resize(container.size() - 1ul);
+}
+
 FilteredBinaryInputArchive::FilteredBinaryInputArchive(BoundedIOStream* stream_,
                                                        DataLayer layer_,
                                                        std::uint16_t maxLOD_,
@@ -51,6 +60,7 @@ FilteredBinaryInputArchive::FilteredBinaryInputArchive(BoundedIOStream* stream_,
     memRes{memRes_},
     layerBitmask{computeDataLayerBitmask(layer_)},
     lodConstraint{maxLOD_, minLOD_, memRes},
+    loadedControls(memRes),
     unconstrainedLODCount{},
     malformed{false} {
 }
@@ -68,6 +78,7 @@ FilteredBinaryInputArchive::FilteredBinaryInputArchive(BoundedIOStream* stream_,
     memRes{memRes_},
     layerBitmask{computeDataLayerBitmask(layer_)},
     lodConstraint{lods_, memRes},
+    loadedControls(memRes),
     unconstrainedLODCount{},
     malformed{false} {
 }
@@ -113,6 +124,9 @@ void FilteredBinaryInputArchive::process(RawDefinition& dest) {
     }
     // Load all data
     BaseArchive::process(dest);
+
+    loadedControls.resize(dest.rawControlNames.size(), true);
+
     // No filtering is done, unless LOD constraint may have some effect
     if (!lodConstraint.hasImpactOn(unconstrainedLODCount)) {
         return;
@@ -149,6 +163,11 @@ void FilteredBinaryInputArchive::process(RawDefinition& dest) {
     AnimatedMapFilter::apply(dest);
 }
 
+void FilteredBinaryInputArchive::process(RawControls& dest) {
+    BaseArchive::process(dest);
+    loadedControls.resize(loadedControls.size() + dest.psdCount, true);
+}
+
 void FilteredBinaryInputArchive::process(RawJoints& dest) {
     process(dest.rowCount);
     process(dest.colCount);
@@ -345,6 +364,8 @@ void FilteredBinaryInputArchive::process(RawMachineLearnedBehavior& dest) {
         process(dest.lodNeuralNetworkMapping);
         process(dest.neuralNetworkToMeshRegion);
 
+        loadedControls.resize(loadedControls.size() + dest.mlControlNames.size(), true);
+
         if (!lodConstraint.hasImpactOn(unconstrainedLODCount)) {
             process(dest.neuralNetworks);
             return;
@@ -392,6 +413,58 @@ void FilteredBinaryInputArchive::process(RawMachineLearnedBehavior& dest) {
                 process(netSizeMarker);
             }
         }
+    } else {
+        const auto mlControlCount = processSize();
+        loadedControls.resize(loadedControls.size() + mlControlCount, false);
+    }
+}
+
+void FilteredBinaryInputArchive::process(DNA& dest) {
+    BaseArchive::process(dest);
+    // Don't run control-based post-load filtering for delta DNA files
+    if (!loadedControls.empty()) {
+        removeUnreferencedBlendShapes(dest);
+    }
+}
+
+void FilteredBinaryInputArchive::removeUnreferencedBlendShapes(DNA& dest) {
+    auto& bsc = dest.behavior.blendShapeChannels;
+
+    const auto originalLODs = bsc.lods;
+    Vector<std::uint16_t> unreferencedChannels{memRes};
+    for (std::size_t iPlusOne = bsc.inputIndices.size(); iPlusOne > 0ul; --iPlusOne) {
+        const auto i = iPlusOne - 1ul;
+        const auto controlIndex = bsc.inputIndices[i];
+        if (!loadedControls[controlIndex]) {
+            unreferencedChannels.push_back(bsc.outputIndices[i]);
+            // Remove behavior data
+            removeByIndex(bsc.inputIndices, i);
+            removeByIndex(bsc.outputIndices, i);
+            for (std::uint16_t lod = {}; lod < bsc.lods.size(); ++lod) {
+                if (i < originalLODs[lod]) {
+                    --bsc.lods[lod];
+                }
+            }
+        }
+    }
+
+    // Remove channel from definition
+    dest.definition.lodBlendShapeMapping.filterIndices([&unreferencedChannels](std::uint16_t index) {
+        return !extd::contains(unreferencedChannels, index);
+    });
+
+    BlendShapeFilter::configure(static_cast<std::uint16_t>(dest.definition.blendShapeChannelNames.size()),
+                                dest.definition.lodBlendShapeMapping.getCombinedDistinctIndices(memRes));
+    BlendShapeFilter::apply(dest.definition);
+
+    // Remap remaining output indices
+    for (auto& outputIdx : bsc.outputIndices) {
+        outputIdx = BlendShapeFilter::remapped(outputIdx);
+    }
+
+    // Remove associated blend shape targets from geometry
+    for (auto& mesh : dest.geometry.meshes) {
+        BlendShapeFilter::apply(mesh);
     }
 }
 
diff --git a/Engine/Plugins/Animation/RigLogic/Source/RigLogicLib/Private/dna/stream/FilteredBinaryInputArchive.h b/Engine/Plugins/Animation/RigLogic/Source/RigLogicLib/Private/dna/stream/FilteredBinaryInputArchive.h
--- a/Engine/Plugins/Animation/RigLogic/Source/RigLogicLib/Private/dna/stream/FilteredBinaryInputArchive.h
+++ b/Engine/Plugins/Animation/RigLogic/Source/RigLogicLib/Private/dna/stream/FilteredBinaryInputArchive.h
@@ -61,6 +61,7 @@ class FilteredBinaryInputArchive final : public AnimatedMapFilter, public BlendS
             }
         }
 
+        void process(RawControls& dest);
         void process(RawJoints& dest);
         void process(RawBlendShapeChannels& dest);
         void process(RawAnimatedMaps& dest);
@@ -83,6 +84,9 @@ class FilteredBinaryInputArchive final : public AnimatedMapFilter, public BlendS
         void process(RawVertexSkinWeights& dest);
         void process(RawMachineLearnedBehavior& dest);
 
+        void process(DNA& dest);
+        void removeUnreferencedBlendShapes(DNA& dest);
+
         template<typename ... Args>
         void process(Args&& ... args) {
             BaseArchive::process(std::forward<Args>(args)...);
@@ -96,6 +100,7 @@ class FilteredBinaryInputArchive final : public AnimatedMapFilter, public BlendS
         MemoryResource* memRes;
         DataLayerBitmask layerBitmask;
         LODConstraint lodConstraint;
+        Vector<bool> loadedControls;
         std::uint16_t unconstrainedLODCount;
         bool malformed;
 
diff --git a/Engine/Plugins/Animation/RigLogic/Source/RigLogicLib/Private/dna/stream/JSONStreamReaderImpl.cpp b/Engine/Plugins/Animation/RigLogic/Source/RigLogicLib/Private/dna/stream/JSONStreamReaderImpl.cpp
--- a/Engine/Plugins/Animation/RigLogic/Source/RigLogicLib/Private/dna/stream/JSONStreamReaderImpl.cpp
+++ b/Engine/Plugins/Animation/RigLogic/Source/RigLogicLib/Private/dna/stream/JSONStreamReaderImpl.cpp
@@ -23,10 +23,6 @@
 
 namespace dna {
 
-// Note: JSONStreamReader::status hasn't been initialized deliberately, as it uses the same error codes that were
-// already registered in BinaryStreamReader, and since they are all registered in a single, global error code registry,
-// this would trigger an assert there.
-
 JSONStreamReader::~JSONStreamReader() = default;
 
 JSONStreamReader* JSONStreamReader::create(BoundedIOStream* stream, MemoryResource* memRes) {
diff --git a/Engine/Plugins/Animation/RigLogic/Source/RigLogicLib/Private/dna/stream/JSONStreamReaderImpl.h b/Engine/Plugins/Animation/RigLogic/Source/RigLogicLib/Private/dna/stream/JSONStreamReaderImpl.h
--- a/Engine/Plugins/Animation/RigLogic/Source/RigLogicLib/Private/dna/stream/JSONStreamReaderImpl.h
+++ b/Engine/Plugins/Animation/RigLogic/Source/RigLogicLib/Private/dna/stream/JSONStreamReaderImpl.h
@@ -8,13 +8,14 @@
 #include "dna/JSONStreamReader.h"
 #include "dna/ReaderImpl.h"
 #include "dna/TypeDefs.h"
+#include "dna/stream/StreamReaderStatus.h"
 
 #include <status/Provider.h>
 #include <terse/archives/json/InputArchive.h>
 
 namespace dna {
 
-class JSONStreamReaderImpl : public ReaderImpl<JSONStreamReader> {
+class JSONStreamReaderImpl : public ReaderImpl<JSONStreamReader>, private StreamReaderStatus {
     public:
         JSONStreamReaderImpl(BoundedIOStream* stream_, MemoryResource* memRes_);
 
@@ -22,8 +23,6 @@ class JSONStreamReaderImpl : public ReaderImpl<JSONStreamReader> {
         void read() override;
 
     private:
-        static sc::StatusProvider status;
-
         BoundedIOStream* stream;
         terse::JSONInputArchive<BoundedIOStream> archive;
 };
diff --git a/Engine/Plugins/Animation/RigLogic/Source/RigLogicLib/Private/dna/stream/StreamReaderStatus.cpp b/Engine/Plugins/Animation/RigLogic/Source/RigLogicLib/Private/dna/stream/StreamReaderStatus.cpp
new file mode 100644
--- /dev/null
+++ b/Engine/Plugins/Animation/RigLogic/Source/RigLogicLib/Private/dna/stream/StreamReaderStatus.cpp
@@ -0,0 +1,22 @@
+// Copyright Epic Games, Inc. All Rights Reserved.
+
+#include "dna/stream/StreamReaderStatus.h"
+
+#include "dna/StreamReader.h"
+
+#include <status/Provider.h>
+
+namespace dna {
+
+#ifdef __clang__
+    #pragma clang diagnostic push
+    #pragma clang diagnostic ignored "-Wglobal-constructors"
+#endif
+sc::StatusProvider StreamReaderStatus::status{StreamReader::SignatureMismatchError,
+                                              StreamReader::VersionMismatchError,
+                                              StreamReader::InvalidDataError};
+#ifdef __clang__
+    #pragma clang diagnostic pop
+#endif
+
+}  // namespace dna
diff --git a/Engine/Plugins/Animation/RigLogic/Source/RigLogicLib/Private/dna/stream/StreamReaderStatus.h b/Engine/Plugins/Animation/RigLogic/Source/RigLogicLib/Private/dna/stream/StreamReaderStatus.h
new file mode 100644
--- /dev/null
+++ b/Engine/Plugins/Animation/RigLogic/Source/RigLogicLib/Private/dna/stream/StreamReaderStatus.h
@@ -0,0 +1,13 @@
+// Copyright Epic Games, Inc. All Rights Reserved.
+
+#pragma once
+
+#include <status/Provider.h>
+
+namespace dna {
+
+struct StreamReaderStatus {
+    static sc::StatusProvider status;
+};
+
+}  // namespace dna
diff --git a/Engine/Plugins/Animation/RigLogic/Source/RigLogicLib/Private/riglogic/joints/JointsBuilder.cpp b/Engine/Plugins/Animation/RigLogic/Source/RigLogicLib/Private/riglogic/joints/JointsBuilder.cpp
--- a/Engine/Plugins/Animation/RigLogic/Source/RigLogicLib/Private/riglogic/joints/JointsBuilder.cpp
+++ b/Engine/Plugins/Animation/RigLogic/Source/RigLogicLib/Private/riglogic/joints/JointsBuilder.cpp
@@ -3,41 +3,56 @@
 #include "riglogic/joints/JointsBuilder.h"
 
 #include "riglogic/TypeDefs.h"
-#include "riglogic/joints/bpcm/builders/Scalar.h"
+#include "riglogic/joints/bpcm/Builder.h"
 #include "riglogic/system/simd/Detect.h"
-#ifdef RL_BUILD_WITH_AVX
-    #ifdef RL_USE_HALF_FLOATS
-        #include "riglogic/joints/bpcm/builders/AVXHalfFloat.h"
-    #else
-        #include "riglogic/joints/bpcm/builders/AVX.h"
-    #endif  // RL_USE_HALF_FLOATS
-#endif  // RL_BUILD_WITH_AVX
-#ifdef RL_BUILD_WITH_SSE
-    #ifdef RL_USE_HALF_FLOATS
-        #include "riglogic/joints/bpcm/builders/SSEHalfFloat.h"
-    #else
-        #include "riglogic/joints/bpcm/builders/SSE.h"
-    #endif  // RL_USE_HALF_FLOATS
-#endif  // RL_BUILD_WITH_SSE
 
 namespace rl4 {
 
 JointsBuilder::~JointsBuilder() = default;
 
 JointsBuilder::Pointer JointsBuilder::create(Configuration config, MemoryResource* memRes) {
+    #if defined(__clang__)
+        #pragma clang diagnostic push
+        #pragma clang diagnostic ignored "-Wunused-local-typedef"
+    #endif
+    #if !defined(__clang__) && defined(__GNUC__)
+        #pragma GCC diagnostic push
+        #pragma GCC diagnostic ignored "-Wunused-local-typedefs"
+    #endif
+    #ifdef RL_USE_HALF_FLOATS
+        using StorageType = std::uint16_t;
+    #else
+        using StorageType = float;
+    #endif  // RL_USE_HALF_FLOATS
+    #if !defined(__clang__) && defined(__GNUC__)
+        #pragma GCC diagnostic pop
+    #endif
+    #if defined(__clang__)
+        #pragma clang diagnostic pop
+    #endif
+
     // Work around unused parameter warning when building without SSE and AVX
     static_cast<void>(config);
     #ifdef RL_BUILD_WITH_SSE
         if (config.calculationType == CalculationType::SSE) {
-            return UniqueInstance<bpcm::SSEJointsBuilder, JointsBuilder>::with(memRes).create(memRes);
+            using SSEJointsBuilder = bpcm::BPCMJointsBuilder<StorageType, trimd::sse::F128>;
+            return UniqueInstance<SSEJointsBuilder, JointsBuilder>::with(memRes).create(memRes);
         }
     #endif  // RL_BUILD_WITH_SSE
     #ifdef RL_BUILD_WITH_AVX
         if (config.calculationType == CalculationType::AVX) {
-            return UniqueInstance<bpcm::AVXJointsBuilder, JointsBuilder>::with(memRes).create(memRes);
+            using AVXJointsBuilder = bpcm::BPCMJointsBuilder<StorageType, trimd::avx::F256>;
+            return UniqueInstance<AVXJointsBuilder, JointsBuilder>::with(memRes).create(memRes);
         }
     #endif  // RL_BUILD_WITH_AVX
-    return UniqueInstance<bpcm::ScalarJointsBuilder, JointsBuilder>::with(memRes).create(memRes);
+    #ifdef RL_BUILD_WITH_NEON
+        if (config.calculationType == CalculationType::NEON) {
+            using NEONJointsBuilder = bpcm::BPCMJointsBuilder<StorageType, trimd::neon::F128>;
+            return UniqueInstance<NEONJointsBuilder, JointsBuilder>::with(memRes).create(memRes);
+        }
+    #endif  // RL_BUILD_WITH_NEON
+    using ScalarJointsBuilder = bpcm::BPCMJointsBuilder<float, trimd::scalar::F128>;
+    return UniqueInstance<ScalarJointsBuilder, JointsBuilder>::with(memRes).create(memRes);
 }
 
 }  // namespace rl4
diff --git a/Engine/Plugins/Animation/RigLogic/Source/RigLogicLib/Private/riglogic/joints/bpcm/Builder.h b/Engine/Plugins/Animation/RigLogic/Source/RigLogicLib/Private/riglogic/joints/bpcm/Builder.h
new file mode 100644
--- /dev/null
+++ b/Engine/Plugins/Animation/RigLogic/Source/RigLogicLib/Private/riglogic/joints/bpcm/Builder.h
@@ -0,0 +1,177 @@
+// Copyright Epic Games, Inc. All Rights Reserved.
+
+#pragma once
+
+#include "riglogic/TypeDefs.h"
+#include "riglogic/joints/JointsBuilder.h"
+#include "riglogic/joints/bpcm/CalculationStrategy.h"
+#include "riglogic/joints/bpcm/Evaluator.h"
+#include "riglogic/joints/bpcm/BPCMOutputInstance.h"
+#include "riglogic/joints/bpcm/Storage.h"
+#include "riglogic/joints/bpcm/StorageSize.h"
+#include "riglogic/joints/bpcm/strategies/Traits.h"
+#include "riglogic/riglogic/RigMetrics.h"
+#include "riglogic/types/Aliases.h"
+#include "riglogic/types/bpcm/Optimizer.h"
+#include "riglogic/utils/Extd.h"
+
+#ifdef _MSC_VER
+    #pragma warning(push)
+    #pragma warning(disable : 4365 4987)
+#endif
+#include <algorithm>
+#include <cstddef>
+#include <cstdint>
+#include <functional>
+#include <memory>
+#include <utility>
+#ifdef _MSC_VER
+    #pragma warning(pop)
+#endif
+
+namespace rl4 {
+
+namespace bpcm {
+
+template<typename TValue, typename TFVec>
+class BPCMJointsBuilder : public JointsBuilder {
+    public:
+        explicit BPCMJointsBuilder(MemoryResource* memRes_);
+
+        void computeStorageRequirements(const RigMetrics& source) override;
+        void computeStorageRequirements(const dna::BehaviorReader* source) override;
+        void allocateStorage(const dna::BehaviorReader* source) override;
+        void fillStorage(const dna::BehaviorReader* source) override;
+        JointsEvaluator::Pointer build() override;
+
+    private:
+        void setValues(const dna::BehaviorReader* source);
+        void setInputIndices(const dna::BehaviorReader* source);
+        void setOutputIndices(const dna::BehaviorReader* source);
+        void setLODs(const dna::BehaviorReader* source);
+        static constexpr std::uint32_t BlockHeight() {
+            return static_cast<std::uint32_t>(TFVec::size() * 2ul);
+        }
+
+        static constexpr std::uint32_t PadTo() {
+            return static_cast<std::uint32_t>(TFVec::size());
+        }
+
+    private:
+        MemoryResource* memRes;
+        StorageSize sizeReqs;
+        JointStorage<TValue> storage;
+};
+
+template<typename TValue, typename TFVec>
+BPCMJointsBuilder<TValue, TFVec>::BPCMJointsBuilder(MemoryResource* memRes_) :
+    memRes{memRes_},
+    sizeReqs{memRes},
+    storage{memRes} {
+}
+
+template<typename TValue, typename TFVec>
+void BPCMJointsBuilder<TValue, TFVec>::computeStorageRequirements(const RigMetrics& source) {
+    // This is incomplete, but enough to create valid joint output instances when restoring from a dump.
+    // Joint attribute count is the only needed data by the instance factory.
+    sizeReqs.attributeCount = source.jointAttributeCount;
+}
+
+template<typename TValue, typename TFVec>
+void BPCMJointsBuilder<TValue, TFVec>::computeStorageRequirements(const dna::BehaviorReader* source) {
+    sizeReqs.computeFrom(source, PadTo());
+}
+
+template<typename TValue, typename TFVec>
+void BPCMJointsBuilder<TValue, TFVec>::allocateStorage(const dna::BehaviorReader*  /*unused*/) {
+    storage.values.resize(sizeReqs.valueCount);
+    storage.inputIndices.resize(sizeReqs.inputIndexCount);
+    storage.outputIndices.resize(sizeReqs.outputIndexCount);
+    storage.lodRegions.reserve(sizeReqs.lodRegionCount);
+    storage.jointGroups.resize(sizeReqs.jointGroups.size());
+}
+
+template<typename TValue, typename TFVec>
+void BPCMJointsBuilder<TValue, TFVec>::fillStorage(const dna::BehaviorReader* source) {
+    setValues(source);
+    setInputIndices(source);
+    setOutputIndices(source);
+    setLODs(source);
+}
+
+template<typename TValue, typename TFVec>
+void BPCMJointsBuilder<TValue, TFVec>::setValues(const dna::BehaviorReader* source) {
+    std::uint32_t offset = 0ul;
+    for (std::uint16_t i = 0u; i < source->getJointGroupCount(); ++i) {
+        const auto values = source->getJointGroupValues(i);
+        const auto jointGroupSize = sizeReqs.getJointGroupSize(i);
+        storage.jointGroups[i].valuesOffset = offset;
+        storage.jointGroups[i].valuesSize = jointGroupSize.padded.size();
+        offset += Optimizer<TFVec, BlockHeight(), PadTo()>::optimize(storage.values.data() + offset,
+                                                                     values.data(),
+                                                                     jointGroupSize.original);
+    }
+}
+
+template<typename TValue, typename TFVec>
+void BPCMJointsBuilder<TValue, TFVec>::setInputIndices(const dna::BehaviorReader* source) {
+    std::uint32_t offset = 0ul;
+    for (std::uint16_t i = 0u; i < source->getJointGroupCount(); ++i) {
+        const auto jointGroupSize = sizeReqs.getJointGroupSize(i);
+        const auto indices = source->getJointGroupInputIndices(i);
+        std::copy(indices.begin(), indices.end(), extd::advanced(storage.inputIndices.begin(), offset));
+        storage.jointGroups[i].inputIndicesOffset = offset;
+        storage.jointGroups[i].inputIndicesSize = jointGroupSize.padded.cols;
+        storage.jointGroups[i].inputIndicesSizeAlignedTo4 = jointGroupSize.padded.cols - (jointGroupSize.padded.cols % 4u);
+        storage.jointGroups[i].inputIndicesSizeAlignedTo8 = jointGroupSize.padded.cols - (jointGroupSize.padded.cols % 8u);
+        offset += jointGroupSize.padded.cols;
+    }
+}
+
+template<typename TValue, typename TFVec>
+void BPCMJointsBuilder<TValue, TFVec>::setOutputIndices(const dna::BehaviorReader* source) {
+    std::uint32_t offset = 0ul;
+    for (std::uint16_t i = 0u; i < source->getJointGroupCount(); ++i) {
+        const auto jointGroupSize = sizeReqs.getJointGroupSize(i);
+        const auto indices = source->getJointGroupOutputIndices(i);
+        std::copy(indices.begin(), indices.end(), extd::advanced(storage.outputIndices.begin(), offset));
+        storage.jointGroups[i].outputIndicesOffset = offset;
+        offset += jointGroupSize.padded.rows;
+    }
+}
+
+template<typename TValue, typename TFVec>
+void BPCMJointsBuilder<TValue, TFVec>::setLODs(const dna::BehaviorReader* source) {
+    std::uint32_t offset = 0ul;
+    for (std::uint16_t i = 0u; i < source->getJointGroupCount(); ++i) {
+        const auto jointGroupSize = sizeReqs.getJointGroupSize(i);
+        const auto dstRowCount = jointGroupSize.padded.rows;
+        auto makeLODRegion = [dstRowCount](std::uint16_t lodRowCount) {
+                return LODRegion{lodRowCount, dstRowCount, BlockHeight(), PadTo()};
+            };
+        const auto lods = source->getJointGroupLODs(i);
+        std::transform(lods.begin(), lods.end(), std::back_inserter(storage.lodRegions), makeLODRegion);
+        storage.jointGroups[i].lodsOffset = offset;
+        offset += source->getLODCount();
+    }
+}
+
+template<typename TValue, typename TFVec>
+JointsEvaluator::Pointer BPCMJointsBuilder<TValue, TFVec>::build() {
+    const auto attributeCount = static_cast<std::uint16_t>(sizeReqs.attributeCount);
+    auto instanceFactory = [attributeCount](MemoryResource* instanceMemRes) {
+            return UniqueInstance<OutputInstance, JointsOutputInstance>::with(instanceMemRes).create(attributeCount,
+                                                                                                     instanceMemRes);
+        };
+
+    using CalculationStrategyBase = JointCalculationStrategy<TValue>;
+    using CalculationStrategy = typename Strategy<TFVec::size()>::template Type<TValue, TFVec>;
+    auto strategy = UniqueInstance<CalculationStrategy, CalculationStrategyBase>::with(memRes).create();
+
+    auto factory = UniqueInstance<Evaluator<TValue>, JointsEvaluator>::with(memRes);
+    return factory.create(std::move(storage), std::move(strategy), instanceFactory, memRes);
+}
+
+}  // namespace bpcm
+
+}  // namespace rl4
diff --git a/Engine/Plugins/Animation/RigLogic/Source/RigLogicLib/Private/riglogic/joints/bpcm/BuilderCommon.h b/Engine/Plugins/Animation/RigLogic/Source/RigLogicLib/Private/riglogic/joints/bpcm/BuilderCommon.h
deleted file mode 100644
diff --git a/Engine/Plugins/Animation/RigLogic/Source/RigLogicLib/Private/riglogic/joints/bpcm/StorageSize.cpp b/Engine/Plugins/Animation/RigLogic/Source/RigLogicLib/Private/riglogic/joints/bpcm/StorageSize.cpp
--- a/Engine/Plugins/Animation/RigLogic/Source/RigLogicLib/Private/riglogic/joints/bpcm/StorageSize.cpp
+++ b/Engine/Plugins/Animation/RigLogic/Source/RigLogicLib/Private/riglogic/joints/bpcm/StorageSize.cpp
@@ -38,10 +38,17 @@ void StorageSize::computeFrom(const dna::BehaviorReader* src, std::size_t padTo)
         inputIndexCount += columnCount;
         outputIndexCount += (rowCount + padding);
         lodRegionCount += lodCount;
+        #if !defined(__clang__) && defined(__GNUC__)
+            #pragma GCC diagnostic push
+            #pragma GCC diagnostic ignored "-Wuseless-cast"
+        #endif
         jointGroups.push_back({
                     {static_cast<std::uint32_t>(rowCount), static_cast<std::uint32_t>(columnCount)},
                     {static_cast<std::uint32_t>(rowCount + padding), static_cast<std::uint32_t>(columnCount)}
                 });
+        #if !defined(__clang__) && defined(__GNUC__)
+            #pragma GCC diagnostic pop
+        #endif
     }
 }
 
diff --git a/Engine/Plugins/Animation/RigLogic/Source/RigLogicLib/Private/riglogic/joints/bpcm/builders/AVX.cpp b/Engine/Plugins/Animation/RigLogic/Source/RigLogicLib/Private/riglogic/joints/bpcm/builders/AVX.cpp
deleted file mode 100644
diff --git a/Engine/Plugins/Animation/RigLogic/Source/RigLogicLib/Private/riglogic/joints/bpcm/builders/AVX.h b/Engine/Plugins/Animation/RigLogic/Source/RigLogicLib/Private/riglogic/joints/bpcm/builders/AVX.h
deleted file mode 100644
diff --git a/Engine/Plugins/Animation/RigLogic/Source/RigLogicLib/Private/riglogic/joints/bpcm/builders/AVXHalfFloat.cpp b/Engine/Plugins/Animation/RigLogic/Source/RigLogicLib/Private/riglogic/joints/bpcm/builders/AVXHalfFloat.cpp
deleted file mode 100644
diff --git a/Engine/Plugins/Animation/RigLogic/Source/RigLogicLib/Private/riglogic/joints/bpcm/builders/AVXHalfFloat.h b/Engine/Plugins/Animation/RigLogic/Source/RigLogicLib/Private/riglogic/joints/bpcm/builders/AVXHalfFloat.h
deleted file mode 100644
diff --git a/Engine/Plugins/Animation/RigLogic/Source/RigLogicLib/Private/riglogic/joints/bpcm/builders/SSE.cpp b/Engine/Plugins/Animation/RigLogic/Source/RigLogicLib/Private/riglogic/joints/bpcm/builders/SSE.cpp
deleted file mode 100644
diff --git a/Engine/Plugins/Animation/RigLogic/Source/RigLogicLib/Private/riglogic/joints/bpcm/builders/SSE.h b/Engine/Plugins/Animation/RigLogic/Source/RigLogicLib/Private/riglogic/joints/bpcm/builders/SSE.h
deleted file mode 100644
diff --git a/Engine/Plugins/Animation/RigLogic/Source/RigLogicLib/Private/riglogic/joints/bpcm/builders/SSEHalfFloat.cpp b/Engine/Plugins/Animation/RigLogic/Source/RigLogicLib/Private/riglogic/joints/bpcm/builders/SSEHalfFloat.cpp
deleted file mode 100644
diff --git a/Engine/Plugins/Animation/RigLogic/Source/RigLogicLib/Private/riglogic/joints/bpcm/builders/SSEHalfFloat.h b/Engine/Plugins/Animation/RigLogic/Source/RigLogicLib/Private/riglogic/joints/bpcm/builders/SSEHalfFloat.h
deleted file mode 100644
diff --git a/Engine/Plugins/Animation/RigLogic/Source/RigLogicLib/Private/riglogic/joints/bpcm/builders/Scalar.cpp b/Engine/Plugins/Animation/RigLogic/Source/RigLogicLib/Private/riglogic/joints/bpcm/builders/Scalar.cpp
deleted file mode 100644
diff --git a/Engine/Plugins/Animation/RigLogic/Source/RigLogicLib/Private/riglogic/joints/bpcm/builders/Scalar.h b/Engine/Plugins/Animation/RigLogic/Source/RigLogicLib/Private/riglogic/joints/bpcm/builders/Scalar.h
deleted file mode 100644
diff --git a/Engine/Plugins/Animation/RigLogic/Source/RigLogicLib/Private/riglogic/joints/bpcm/strategies/AVX.h b/Engine/Plugins/Animation/RigLogic/Source/RigLogicLib/Private/riglogic/joints/bpcm/strategies/AVX.h
deleted file mode 100644
diff --git a/Engine/Plugins/Animation/RigLogic/Source/RigLogicLib/Private/riglogic/joints/bpcm/strategies/Block4.h b/Engine/Plugins/Animation/RigLogic/Source/RigLogicLib/Private/riglogic/joints/bpcm/strategies/Block4.h
--- a/Engine/Plugins/Animation/RigLogic/Source/RigLogicLib/Private/riglogic/joints/bpcm/strategies/Block4.h
+++ b/Engine/Plugins/Animation/RigLogic/Source/RigLogicLib/Private/riglogic/joints/bpcm/strategies/Block4.h
@@ -16,6 +16,8 @@ namespace rl4 {
 
 namespace bpcm {
 
+namespace block4 {
+
 /*
  * Process the remainder portion after 8x4 blocks
  *
@@ -333,7 +335,7 @@ static FORCE_INLINE void processJointGroupBlock4(const JointGroupView<T>& jointG
     }
 }
 
-template<typename TF128, typename T>
+template<typename T, typename TF128>
 struct Block4JointCalculationStrategy : public JointCalculationStrategy<T> {
 
     using JointGroupArrayView = typename JointCalculationStrategy<T>::JointGroupArrayView;
@@ -359,6 +361,8 @@ struct Block4JointCalculationStrategy : public JointCalculationStrategy<T> {
 
 };
 
+}  // namespace block4
+
 }  // namespace bpcm
 
 }  // namespace rl4
diff --git a/Engine/Plugins/Animation/RigLogic/Source/RigLogicLib/Private/riglogic/joints/bpcm/strategies/Block8.h b/Engine/Plugins/Animation/RigLogic/Source/RigLogicLib/Private/riglogic/joints/bpcm/strategies/Block8.h
--- a/Engine/Plugins/Animation/RigLogic/Source/RigLogicLib/Private/riglogic/joints/bpcm/strategies/Block8.h
+++ b/Engine/Plugins/Animation/RigLogic/Source/RigLogicLib/Private/riglogic/joints/bpcm/strategies/Block8.h
@@ -19,6 +19,8 @@ namespace rl4 {
 
 namespace bpcm {
 
+namespace block8 {
+
 template<typename TF256, typename T>
 static FORCE_INLINE void processBlocks16x1(const std::uint16_t* inputIndicesEndAlignedTo4,
                                            const std::uint16_t* inputIndicesEnd,
@@ -231,7 +233,7 @@ static FORCE_INLINE void processJointGroupBlock8(const JointGroupView<T>& jointG
     }
 }
 
-template<typename TF256, typename T>
+template<typename T, typename TF256>
 struct Block8JointCalculationStrategy : public JointCalculationStrategy<T> {
 
     using JointGroupArrayView = typename JointCalculationStrategy<T>::JointGroupArrayView;
@@ -257,6 +259,8 @@ struct Block8JointCalculationStrategy : public JointCalculationStrategy<T> {
 
 };
 
+}  // namespace block8
+
 }  // namespace bpcm
 
 }  // namespace rl4
diff --git a/Engine/Plugins/Animation/RigLogic/Source/RigLogicLib/Private/riglogic/joints/bpcm/strategies/SSE.h b/Engine/Plugins/Animation/RigLogic/Source/RigLogicLib/Private/riglogic/joints/bpcm/strategies/SSE.h
deleted file mode 100644
diff --git a/Engine/Plugins/Animation/RigLogic/Source/RigLogicLib/Private/riglogic/joints/bpcm/strategies/Scalar.h b/Engine/Plugins/Animation/RigLogic/Source/RigLogicLib/Private/riglogic/joints/bpcm/strategies/Scalar.h
deleted file mode 100644
diff --git a/Engine/Plugins/Animation/RigLogic/Source/RigLogicLib/Private/riglogic/joints/bpcm/strategies/Traits.h b/Engine/Plugins/Animation/RigLogic/Source/RigLogicLib/Private/riglogic/joints/bpcm/strategies/Traits.h
new file mode 100644
--- /dev/null
+++ b/Engine/Plugins/Animation/RigLogic/Source/RigLogicLib/Private/riglogic/joints/bpcm/strategies/Traits.h
@@ -0,0 +1,31 @@
+// Copyright Epic Games, Inc. All Rights Reserved.
+
+#pragma once
+
+#include "riglogic/joints/bpcm/strategies/Block4.h"
+#include "riglogic/joints/bpcm/strategies/Block8.h"
+#include "riglogic/system/simd/SIMD.h"
+
+namespace rl4 {
+
+namespace bpcm {
+
+template<std::size_t BlockHeight>
+struct Strategy;
+
+template<>
+struct Strategy<4> {
+    template<typename T, typename TFVec>
+    using Type = block4::Block4JointCalculationStrategy<T, TFVec>;
+};
+
+template<>
+struct Strategy<8> {
+    template<typename T, typename TFVec>
+    using Type = block8::Block8JointCalculationStrategy<T, TFVec>;
+};
+
+
+}  // namespace bpcm
+
+}  // namespace rl4
diff --git a/Engine/Plugins/Animation/RigLogic/Source/RigLogicLib/Private/riglogic/ml/MachineLearnedBehaviorFactory.cpp b/Engine/Plugins/Animation/RigLogic/Source/RigLogicLib/Private/riglogic/ml/MachineLearnedBehaviorFactory.cpp
--- a/Engine/Plugins/Animation/RigLogic/Source/RigLogicLib/Private/riglogic/ml/MachineLearnedBehaviorFactory.cpp
+++ b/Engine/Plugins/Animation/RigLogic/Source/RigLogicLib/Private/riglogic/ml/MachineLearnedBehaviorFactory.cpp
@@ -42,6 +42,11 @@ MachineLearnedBehaviorEvaluator::Pointer createEvaluator(const Configuration& co
             return ml::cpu::Factory<StorageValueType, trimd::avx::F256, trimd::sse::F128>::create(reader, memRes);
         }
     #endif  // RL_BUILD_WITH_AVX
+    #ifdef RL_BUILD_WITH_NEON
+        if (config.calculationType == CalculationType::NEON) {
+            return ml::cpu::Factory<StorageValueType, trimd::neon::F256, trimd::neon::F128>::create(reader, memRes);
+        }
+    #endif  // RL_BUILD_WITH_NEON
     return ml::cpu::Factory<float, trimd::scalar::F256, trimd::scalar::F128>::create(reader, memRes);
 }
 
diff --git a/Engine/Plugins/Animation/RigLogic/Source/RigLogicLib/Private/riglogic/ml/cpu/Block4.h b/Engine/Plugins/Animation/RigLogic/Source/RigLogicLib/Private/riglogic/ml/cpu/Block4.h
--- a/Engine/Plugins/Animation/RigLogic/Source/RigLogicLib/Private/riglogic/ml/cpu/Block4.h
+++ b/Engine/Plugins/Animation/RigLogic/Source/RigLogicLib/Private/riglogic/ml/cpu/Block4.h
@@ -202,13 +202,15 @@ static FORCE_INLINE void calculateBlock4(const NeuralNetLayer<T>& layer, ConstAr
     }
 
     for (; outputVector < outputVectorEnd; outputVector += block4Height, biases += block4Height) {
+        alignas(TF128::alignment()) float outbuf[block4Height];
         processBlocks4x8<TF128, TActivationFunction<TF128> >(inputVector,
                                                              inputVectorEndAlignedTo8,
                                                              inputVectorEnd,
                                                              weights,
                                                              biases,
                                                              activationParams,
-                                                             outputVector);
+                                                             outbuf);
+        std::memcpy(outputVector, outbuf, sizeof(outbuf));
     }
 }
 
diff --git a/Engine/Plugins/Animation/RigLogic/Source/RigLogicLib/Private/riglogic/system/simd/SIMD.h b/Engine/Plugins/Animation/RigLogic/Source/RigLogicLib/Private/riglogic/system/simd/SIMD.h
--- a/Engine/Plugins/Animation/RigLogic/Source/RigLogicLib/Private/riglogic/system/simd/SIMD.h
+++ b/Engine/Plugins/Animation/RigLogic/Source/RigLogicLib/Private/riglogic/system/simd/SIMD.h
@@ -4,11 +4,10 @@
 
 #include "riglogic/system/simd/Detect.h"
 
-#if defined(RL_USE_HALF_FLOATS) && !defined(TRIMD_ENABLE_F16C)
-    #define TRIMD_ENABLE_F16C
-#endif  // RL_USE_HALF_FLOATS
-
 #if defined(RL_BUILD_WITH_AVX)
+    #if defined(RL_USE_HALF_FLOATS) && !defined(TRIMD_ENABLE_F16C)
+        #define TRIMD_ENABLE_F16C
+    #endif  // RL_USE_HALF_FLOATS
     #if !defined(TRIMD_ENABLE_AVX)
         #define TRIMD_ENABLE_AVX
     #endif
@@ -18,7 +17,17 @@
 #endif  // RL_BUILD_WITH_AVX
 
 #if defined(RL_BUILD_WITH_SSE) && !defined(TRIMD_ENABLE_SSE)
+    #if defined(RL_USE_HALF_FLOATS) && !defined(TRIMD_ENABLE_F16C)
+        #define TRIMD_ENABLE_F16C
+    #endif  // RL_USE_HALF_FLOATS
     #define TRIMD_ENABLE_SSE
 #endif  // RL_BUILD_WITH_SSE
 
+#if defined(RL_BUILD_WITH_NEON) && !defined(TRIMD_ENABLE_NEON)
+    #if defined(RL_USE_HALF_FLOATS) && !defined(TRIMD_ENABLE_NEON_FP16)
+        #define TRIMD_ENABLE_NEON_FP16
+    #endif  // RL_USE_HALF_FLOATS
+    #define TRIMD_ENABLE_NEON
+#endif  // RL_BUILD_WITH_NEON
+
 #include <trimd/TRiMD.h>
diff --git a/Engine/Plugins/Animation/RigLogic/Source/RigLogicLib/Private/trio/streams/FileStreamImpl.cpp b/Engine/Plugins/Animation/RigLogic/Source/RigLogicLib/Private/trio/streams/FileStreamImpl.cpp
--- a/Engine/Plugins/Animation/RigLogic/Source/RigLogicLib/Private/trio/streams/FileStreamImpl.cpp
+++ b/Engine/Plugins/Animation/RigLogic/Source/RigLogicLib/Private/trio/streams/FileStreamImpl.cpp
@@ -63,6 +63,7 @@ FileStreamImpl::FileStreamImpl(const char* path_, AccessMode accessMode_, OpenMo
     fileAccessMode{accessMode_},
     fileOpenMode{openMode_},
     fileSize{getFileSizeStd(filePath.c_str())},
+    filePos{},
     memRes{memRes_} {
 }
 
@@ -97,7 +98,7 @@ void FileStreamImpl::close() {
 }
 
 std::uint64_t FileStreamImpl::tell() {
-    return static_cast<std::uint64_t>(file.tellp());
+    return filePos;
 }
 
 void FileStreamImpl::seek(std::uint64_t position) {
@@ -108,27 +109,41 @@ void FileStreamImpl::seek(std::uint64_t position) {
     }
 
     file.seekp(static_cast<std::streamoff>(position));
+    filePos = position;
     if (!file.good()) {
         status->set(SeekError, filePath.c_str());
     }
 }
 
 std::size_t FileStreamImpl::read(char* destination, std::size_t size) {
+    if (file.eof()) {
+        return 0ul;
+    }
+
     if ((destination == nullptr) || file.fail() || !file.is_open() || !contains(fileAccessMode, AccessMode::Read)) {
         status->set(ReadError, filePath.c_str());
         return 0ul;
     }
 
     file.read(destination, static_cast<std::streamsize>(size));
-    if (!file.good() && !file.eof()) {
+    if (file.fail() && !file.eof()) {
         status->set(ReadError, filePath.c_str());
     }
 
     const auto bytesRead = file.gcount();
-    return (bytesRead > 0 ? static_cast<std::size_t>(bytesRead) : 0ul);
+    if (bytesRead < 0) {
+        return 0ul;
+    }
+
+    filePos += static_cast<std::size_t>(bytesRead);
+    return static_cast<std::size_t>(bytesRead);
 }
 
 std::size_t FileStreamImpl::read(Writable* destination, std::size_t size) {
+    if (file.eof()) {
+        return 0ul;
+    }
+
     if ((destination == nullptr) || file.fail() || !file.is_open() || !contains(fileAccessMode, AccessMode::Read)) {
         status->set(ReadError, filePath.c_str());
         return 0ul;
@@ -143,11 +158,17 @@ std::size_t FileStreamImpl::read(Writable* destination, std::size_t size) {
     file.read(buffer, static_cast<std::streamsize>(size));
     destination->write(buffer, size);
 
-    if (!file.good() && !file.eof()) {
+    if (file.fail() && !file.eof()) {
         status->set(ReadError, filePath.c_str());
     }
+
     const auto bytesRead = file.gcount();
-    return (bytesRead > 0 ? static_cast<std::size_t>(bytesRead) : 0ul);
+    if (bytesRead < 0) {
+        return 0ul;
+    }
+
+    filePos += static_cast<std::size_t>(bytesRead);
+    return static_cast<std::size_t>(bytesRead);
 }
 
 std::size_t FileStreamImpl::write(const char* source, std::size_t size) {
@@ -156,17 +177,20 @@ std::size_t FileStreamImpl::write(const char* source, std::size_t size) {
         return 0ul;
     }
 
-    const auto preWritePos = file.tellp();
-    file.write(source, static_cast<std::streamsize>(size));
-    const auto postWritePos = file.tellp();
+    #if defined(_MSC_VER) && (_MSC_VER < 1930)
+    // MSVC stdlib or compiler bug workaround
+    file.seekp(file.tellp());
+    #endif  // _MSC_VER
 
+    file.write(source, static_cast<std::streamsize>(size));
     if (!file.good()) {
         status->set(WriteError, filePath.c_str());
-    } else {
-        fileSize = std::max(static_cast<std::uint64_t>(postWritePos), fileSize);
+        return 0ul;
     }
 
-    return (postWritePos > preWritePos ? static_cast<std::size_t>(postWritePos - preWritePos) : 0ul);
+    fileSize = std::max(filePos + size, fileSize);
+    filePos += size;
+    return size;
 }
 
 std::size_t FileStreamImpl::write(Readable* source, std::size_t size) {
@@ -175,7 +199,10 @@ std::size_t FileStreamImpl::write(Readable* source, std::size_t size) {
         return 0ul;
     }
 
-    const auto preWritePos = file.tellp();
+    #if defined(_MSC_VER) && (_MSC_VER < 1930)
+    // MSVC stdlib or compiler bug workaround
+    file.seekp(file.tellp());
+    #endif  // _MSC_VER
 
     char buffer[bufferSize];
     while (size > bufferSize) {
@@ -185,16 +212,14 @@ std::size_t FileStreamImpl::write(Readable* source, std::size_t size) {
     }
     source->read(buffer, size);
     file.write(buffer, static_cast<std::streamsize>(size));
-
-    const auto postWritePos = file.tellp();
-
     if (!file.good()) {
         status->set(WriteError, filePath.c_str());
-    } else {
-        fileSize = std::max(static_cast<std::uint64_t>(postWritePos), fileSize);
+        return 0ul;
     }
 
-    return (postWritePos > preWritePos ? static_cast<std::size_t>(postWritePos - preWritePos) : 0ul);
+    fileSize = std::max(filePos + size, fileSize);
+    filePos += size;
+    return size;
 }
 
 std::uint64_t FileStreamImpl::size() {
diff --git a/Engine/Plugins/Animation/RigLogic/Source/RigLogicLib/Private/trio/streams/FileStreamImpl.h b/Engine/Plugins/Animation/RigLogic/Source/RigLogicLib/Private/trio/streams/FileStreamImpl.h
--- a/Engine/Plugins/Animation/RigLogic/Source/RigLogicLib/Private/trio/streams/FileStreamImpl.h
+++ b/Engine/Plugins/Animation/RigLogic/Source/RigLogicLib/Private/trio/streams/FileStreamImpl.h
@@ -45,6 +45,7 @@ class FileStreamImpl : public FileStream {
         AccessMode fileAccessMode;
         OpenMode fileOpenMode;
         std::uint64_t fileSize;
+        std::uint64_t filePos;
         MemoryResource* memRes;
         StreamStatus status;
 };
diff --git a/Engine/Plugins/Animation/RigLogic/Source/RigLogicLib/Public/dna/Defs.h b/Engine/Plugins/Animation/RigLogic/Source/RigLogicLib/Public/dna/Defs.h
--- a/Engine/Plugins/Animation/RigLogic/Source/RigLogicLib/Public/dna/Defs.h
+++ b/Engine/Plugins/Animation/RigLogic/Source/RigLogicLib/Public/dna/Defs.h
@@ -2,17 +2,28 @@
 
 #pragma once
 
-#if defined(_WIN32) || defined(__CYGWIN__)
-    #if defined(__GNUC__)
-        #define DLL_EXPORT __attribute__((dllexport))
-        #define DLL_IMPORT __attribute__((dllimport))
-    #else
-        #define DLL_EXPORT __declspec(dllexport)
-        #define DLL_IMPORT __declspec(dllimport)
-    #endif
-#elif defined(__GNUC__)
+#if defined(FORCE_DLL_ATTRIBUTE_EXPORT_IMPORT)
+    #define DLL_EXPORT __attribute__((dllexport))
+    #define DLL_IMPORT __attribute__((dllimport))
+#elif defined(FORCE_DLL_DECLSPEC_EXPORT_IMPORT)
+    #define DLL_EXPORT __declspec(dllexport)
+    #define DLL_IMPORT __declspec(dllimport)
+#elif defined(FORCE_DLL_VISIBILITY_EXPORT_IMPORT)
     #define DLL_EXPORT __attribute__((visibility("default")))
     #define DLL_IMPORT DLL_EXPORT
+#else
+    #if defined(_WIN32) || defined(__CYGWIN__)
+        #if defined(__GNUC__)
+            #define DLL_EXPORT __attribute__((dllexport))
+            #define DLL_IMPORT __attribute__((dllimport))
+        #else
+            #define DLL_EXPORT __declspec(dllexport)
+            #define DLL_IMPORT __declspec(dllimport)
+        #endif
+    #elif defined(__GNUC__)
+        #define DLL_EXPORT __attribute__((visibility("default")))
+        #define DLL_IMPORT DLL_EXPORT
+    #endif
 #endif
 
 #if defined(RL_BUILD_SHARED)
diff --git a/Engine/Plugins/Animation/RigLogic/Source/RigLogicLib/Public/dna/layers/BehaviorReader.h b/Engine/Plugins/Animation/RigLogic/Source/RigLogicLib/Public/dna/layers/BehaviorReader.h
--- a/Engine/Plugins/Animation/RigLogic/Source/RigLogicLib/Public/dna/layers/BehaviorReader.h
+++ b/Engine/Plugins/Animation/RigLogic/Source/RigLogicLib/Public/dna/layers/BehaviorReader.h
@@ -76,12 +76,10 @@ class DNAAPI BehaviorReader : public virtual DefinitionReader {
         virtual ConstArrayView<float> getPSDValues() const = 0;
         /**
             @brief Number of rows in the entire, uncompressed joint matrix.
-            @see Joints
         */
         virtual std::uint16_t getJointRowCount() const = 0;
         /**
             @brief Number of columns in the entire, uncompressed joint matrix.
-            @see Joints
         */
         virtual std::uint16_t getJointColumnCount() const = 0;
         /**
@@ -91,7 +89,6 @@ class DNAAPI BehaviorReader : public virtual DefinitionReader {
         virtual ConstArrayView<std::uint16_t> getJointVariableAttributeIndices(std::uint16_t lod) const = 0;
         /**
             @brief Number of joint groups present in the entire joint matrix.
-            @see Joints
         */
         virtual std::uint16_t getJointGroupCount() const = 0;
         /**
@@ -108,7 +105,6 @@ class DNAAPI BehaviorReader : public virtual DefinitionReader {
             @warning
                 jointGroupIndex must be less than the value returned by getJointGroupCount.
             @return View over the array of LOD bounds.
-            @see Joints
         */
         virtual ConstArrayView<std::uint16_t> getJointGroupLODs(std::uint16_t jointGroupIndex) const = 0;
         /**
@@ -120,7 +116,6 @@ class DNAAPI BehaviorReader : public virtual DefinitionReader {
             @warning
                 jointGroupIndex must be less than the value returned by getJointGroupCount.
             @return View over the array of column indices.
-            @see Joints
         */
         virtual ConstArrayView<std::uint16_t> getJointGroupInputIndices(std::uint16_t jointGroupIndex) const = 0;
         /**
@@ -132,7 +127,6 @@ class DNAAPI BehaviorReader : public virtual DefinitionReader {
             @warning
                 jointGroupIndex must be less than the value returned by getJointGroupCount.
             @return View over the array of row indices.
-            @see Joints
         */
         virtual ConstArrayView<std::uint16_t> getJointGroupOutputIndices(std::uint16_t jointGroupIndex) const = 0;
         /**
@@ -142,7 +136,6 @@ class DNAAPI BehaviorReader : public virtual DefinitionReader {
             @warning
                 jointGroupIndex must be less than the value returned by getJointGroupCount.
             @return View over the array of values.
-            @see Joints
         */
         virtual ConstArrayView<float> getJointGroupValues(std::uint16_t jointGroupIndex) const = 0;
         /**
@@ -154,7 +147,6 @@ class DNAAPI BehaviorReader : public virtual DefinitionReader {
             @warning
                 jointGroupIndex must be less than the value returned by getJointGroupCount.
             @return View over the array of joint indices.
-            @see Joints
             @see DefinitionReader
         */
         virtual ConstArrayView<std::uint16_t> getJointGroupJointIndices(std::uint16_t jointGroupIndex) const = 0;
@@ -166,19 +158,16 @@ class DNAAPI BehaviorReader : public virtual DefinitionReader {
             @warning
                 These LOD values are not interchangeable with the LOD indices from DefinitionReader::getBlendShapeChannelIndicesForLOD.
             @return View over the array of LOD bounds.
-            @see BlendShapes
         */
         virtual ConstArrayView<std::uint16_t> getBlendShapeChannelLODs() const = 0;
         /**
             @brief Input indices used to index into the input vector.
             @return View over the array of input indices.
-            @see BlendShapes
         */
         virtual ConstArrayView<std::uint16_t> getBlendShapeChannelInputIndices() const = 0;
         /**
             @brief Output indices specify the positions of blend shape channel output values.
             @return View over the array of output indices.
-            @see BlendShapes
         */
         virtual ConstArrayView<std::uint16_t> getBlendShapeChannelOutputIndices() const = 0;
         /**
@@ -188,45 +177,38 @@ class DNAAPI BehaviorReader : public virtual DefinitionReader {
                 value 5 is LOD with lowest details), while the value denotes the number of rows (within the conditional table),
                 belonging to that level.
             @return View over the array of LOD bounds.
-            @see AnimatedMaps
         */
         virtual ConstArrayView<std::uint16_t> getAnimatedMapLODs() const = 0;
         /**
             @brief Input indices used to index into the array of input values.
             @return View over the array of input indices.
-            @see AnimatedMaps
         */
         virtual ConstArrayView<std::uint16_t> getAnimatedMapInputIndices() const = 0;
         /**
             @brief Output indices that specify the computed output value's position.
             @return View over the array of output indices.
-            @see AnimatedMaps
         */
         virtual ConstArrayView<std::uint16_t> getAnimatedMapOutputIndices() const = 0;
         /**
             @brief Filter values(lower-bounds) used to decide whether a particular
                 entry should be evaluated or not.
             @return View over the array of filter values.
-            @see AnimatedMaps
         */
         virtual ConstArrayView<float> getAnimatedMapFromValues() const = 0;
         /**
             @brief Filter values(upper-bounds) used to decide whether a particular
                 entry should be evaluated or not.
             @return View over the array of filter values.
-            @see AnimatedMaps
         */
         virtual ConstArrayView<float> getAnimatedMapToValues() const = 0;
         /**
             @brief Computational values(slope/gradient) used for calculating the output value.
             @return View over the array of computational values.
-            @see AnimatedMaps
         */
         virtual ConstArrayView<float> getAnimatedMapSlopeValues() const = 0;
         /**
             @brief Computational values(vertical intercept) used for calculating the output value.
             @return View over the array of computational values.
-            @see AnimatedMaps
         */
         virtual ConstArrayView<float> getAnimatedMapCutValues() const = 0;
 };
diff --git a/Engine/Plugins/Animation/RigLogic/Source/RigLogicLib/Public/dna/layers/BehaviorWriter.h b/Engine/Plugins/Animation/RigLogic/Source/RigLogicLib/Public/dna/layers/BehaviorWriter.h
--- a/Engine/Plugins/Animation/RigLogic/Source/RigLogicLib/Public/dna/layers/BehaviorWriter.h
+++ b/Engine/Plugins/Animation/RigLogic/Source/RigLogicLib/Public/dna/layers/BehaviorWriter.h
@@ -103,12 +103,10 @@ class DNAAPI BehaviorWriter : public virtual DefinitionWriter {
         virtual void setPSDValues(const float* weights, std::uint16_t count) = 0;
         /**
             @brief Number of rows in the entire, uncompressed joint matrix.
-            @see Joints
         */
         virtual void setJointRowCount(std::uint16_t rowCount) = 0;
         /**
             @brief Number of columns in the entire, uncompressed joint matrix.
-            @see Joints
         */
         virtual void setJointColumnCount(std::uint16_t columnCount) = 0;
         /**
@@ -141,7 +139,6 @@ class DNAAPI BehaviorWriter : public virtual DefinitionWriter {
                 The source address from which the lod bounds are to be copied.
             @param count
                 The number of lod bounds to copy.
-            @see Joints
         */
         virtual void setJointGroupLODs(std::uint16_t jointGroupIndex, const std::uint16_t* lods, std::uint16_t count) = 0;
         /**
@@ -157,7 +154,6 @@ class DNAAPI BehaviorWriter : public virtual DefinitionWriter {
                 The source address from which the column indices are to be copied.
             @param count
                 The number of column indices to copy.
-            @see Joints
         */
         virtual void setJointGroupInputIndices(std::uint16_t jointGroupIndex,
                                                const std::uint16_t* inputIndices,
@@ -175,7 +171,6 @@ class DNAAPI BehaviorWriter : public virtual DefinitionWriter {
                 The source address from which the row indices are to be copied.
             @param count
                 The number of row indices to copy.
-            @see Joints
         */
         virtual void setJointGroupOutputIndices(std::uint16_t jointGroupIndex,
                                                 const std::uint16_t* outputIndices,
@@ -191,7 +186,6 @@ class DNAAPI BehaviorWriter : public virtual DefinitionWriter {
                 The source address from which the values are to be copied.
             @param count
                 The number of values to copy.
-            @see Joints
         */
         virtual void setJointGroupValues(std::uint16_t jointGroupIndex, const float* values, std::uint32_t count) = 0;
         /**
@@ -205,7 +199,6 @@ class DNAAPI BehaviorWriter : public virtual DefinitionWriter {
                 The source address from which the joint indices are to be copied.
             @param count
                 The number of joint indices to copy.
-            @see Joints
         */
         virtual void setJointGroupJointIndices(std::uint16_t jointGroupIndex,
                                                const std::uint16_t* jointIndices,
@@ -222,7 +215,6 @@ class DNAAPI BehaviorWriter : public virtual DefinitionWriter {
             @warning
                 The LOD values set here are not interchangeable with the LOD indices set in DefinitionWriter::setBlendShapeNameIndices
                 and DefinitionWriter::setLODBlendShapeMapping
-            @see BlendShapes
         */
         virtual void setBlendShapeChannelLODs(const std::uint16_t* lods, std::uint16_t count) = 0;
         /**
@@ -231,7 +223,6 @@ class DNAAPI BehaviorWriter : public virtual DefinitionWriter {
                 The source address from which the input indices are to be copied.
             @param count
                 The number of input indices to copy.
-            @see BlendShapes
         */
         virtual void setBlendShapeChannelInputIndices(const std::uint16_t* inputIndices, std::uint16_t count) = 0;
         /**
@@ -240,7 +231,6 @@ class DNAAPI BehaviorWriter : public virtual DefinitionWriter {
                 The source address from which the output indices are to be copied.
             @param count
                 The number of output indices to copy.
-            @see BlendShapes
         */
         virtual void setBlendShapeChannelOutputIndices(const std::uint16_t* outputIndices, std::uint16_t count) = 0;
         /**
@@ -253,7 +243,6 @@ class DNAAPI BehaviorWriter : public virtual DefinitionWriter {
                 The source address from which the lod bounds are to be copied.
             @param count
                 The number of lod bounds to copy.
-            @see AnimatedMaps
         */
         virtual void setAnimatedMapLODs(const std::uint16_t* lods, std::uint16_t count) = 0;
         /**
@@ -262,7 +251,6 @@ class DNAAPI BehaviorWriter : public virtual DefinitionWriter {
                 The source address from which the input indices are to be copied.
             @param count
                 The number of input indices to copy.
-            @see AnimatedMaps
         */
         virtual void setAnimatedMapInputIndices(const std::uint16_t* inputIndices, std::uint16_t count) = 0;
         /**
@@ -271,7 +259,6 @@ class DNAAPI BehaviorWriter : public virtual DefinitionWriter {
                 The source address from which the output indices are to be copied.
             @param count
                 The number of output indices to copy.
-            @see AnimatedMaps
         */
         virtual void setAnimatedMapOutputIndices(const std::uint16_t* outputIndices, std::uint16_t count) = 0;
         /**
@@ -281,7 +268,6 @@ class DNAAPI BehaviorWriter : public virtual DefinitionWriter {
                 The source address from which the filter values are to be copied.
             @param count
                 The number of filter values to copy.
-            @see AnimatedMaps
         */
         virtual void setAnimatedMapFromValues(const float* fromValues, std::uint16_t count) = 0;
         /**
@@ -291,7 +277,6 @@ class DNAAPI BehaviorWriter : public virtual DefinitionWriter {
                 The source address from which the filter values are to be copied.
             @param count
                 The number of filter values to copy.
-            @see AnimatedMaps
         */
         virtual void setAnimatedMapToValues(const float* toValues, std::uint16_t count) = 0;
         /**
@@ -300,7 +285,6 @@ class DNAAPI BehaviorWriter : public virtual DefinitionWriter {
                 The source address from which the computational values are to be copied.
             @param count
                 The number of computational values to copy.
-            @see AnimatedMaps
         */
         virtual void setAnimatedMapSlopeValues(const float* slopeValues, std::uint16_t count) = 0;
         /**
@@ -309,7 +293,6 @@ class DNAAPI BehaviorWriter : public virtual DefinitionWriter {
                 The source address from which the computational values are to be copied.
             @param count
                 The number of computational values to copy.
-            @see AnimatedMaps
         */
         virtual void setAnimatedMapCutValues(const float* cutValues, std::uint16_t count) = 0;
 };
diff --git a/Engine/Plugins/Animation/RigLogic/Source/RigLogicLib/Public/dna/layers/DefinitionReader.h b/Engine/Plugins/Animation/RigLogic/Source/RigLogicLib/Public/dna/layers/DefinitionReader.h
--- a/Engine/Plugins/Animation/RigLogic/Source/RigLogicLib/Public/dna/layers/DefinitionReader.h
+++ b/Engine/Plugins/Animation/RigLogic/Source/RigLogicLib/Public/dna/layers/DefinitionReader.h
@@ -67,7 +67,6 @@ class DNAAPI DefinitionReader : public DescriptorReader {
             @warning
                 The index must be less than the value returned by getJointCount.
             @return View over the joint name string.
-            @see Joints
         */
         virtual StringView getJointName(std::uint16_t index) const = 0;
         /**
@@ -122,7 +121,6 @@ class DNAAPI DefinitionReader : public DescriptorReader {
             @warning
                 The index must be less than the value returned by getBlendShapeChannelCount.
             @return View over the blend shape channel name string.
-            @see BlendShapes
         */
         virtual StringView getBlendShapeChannelName(std::uint16_t index) const = 0;
         /**
@@ -155,7 +153,6 @@ class DNAAPI DefinitionReader : public DescriptorReader {
             @warning
                 The index must be less than the value returned by getAnimatedMapCount.
             @return View over the animated map name string.
-            @see AnimatedMaps
         */
         virtual StringView getAnimatedMapName(std::uint16_t index) const = 0;
         /**
diff --git a/Engine/Plugins/Animation/RigLogic/Source/RigLogicLib/Public/dna/layers/DefinitionWriter.h b/Engine/Plugins/Animation/RigLogic/Source/RigLogicLib/Public/dna/layers/DefinitionWriter.h
--- a/Engine/Plugins/Animation/RigLogic/Source/RigLogicLib/Public/dna/layers/DefinitionWriter.h
+++ b/Engine/Plugins/Animation/RigLogic/Source/RigLogicLib/Public/dna/layers/DefinitionWriter.h
@@ -71,7 +71,6 @@ class DNAAPI DefinitionWriter : public DescriptorWriter {
                 A null-terminated string.
             @note
                 The passed in name is copied, which will involve an additional allocation.
-            @see Joints
         */
         virtual void setJointName(std::uint16_t index, const char* name) = 0;
         /**
@@ -125,7 +124,6 @@ class DNAAPI DefinitionWriter : public DescriptorWriter {
                 A null-terminated string.
             @note
                 The passed in name is copied, which will involve an additional allocation.
-            @see BlendShapes
         */
         virtual void setBlendShapeChannelName(std::uint16_t index, const char* name) = 0;
         /**
@@ -183,7 +181,6 @@ class DNAAPI DefinitionWriter : public DescriptorWriter {
                 A null-terminated string.
             @note
                 The passed in name is copied, which will involve an additional allocation.
-            @see AnimatedMaps
         */
         virtual void setAnimatedMapName(std::uint16_t index, const char* name) = 0;
         /**
diff --git a/Engine/Plugins/Animation/RigLogic/Source/RigLogicLib/Public/dna/layers/DescriptorReader.h b/Engine/Plugins/Animation/RigLogic/Source/RigLogicLib/Public/dna/layers/DescriptorReader.h
--- a/Engine/Plugins/Animation/RigLogic/Source/RigLogicLib/Public/dna/layers/DescriptorReader.h
+++ b/Engine/Plugins/Animation/RigLogic/Source/RigLogicLib/Public/dna/layers/DescriptorReader.h
@@ -33,7 +33,6 @@ class DNAAPI DescriptorReader : public HeaderReader {
             @warning
                 The index must be less than the value returned by getMetaDataCount.
             @return View over the key name string.
-            @see getMedaData
         */
         virtual StringView getMetaDataKey(std::uint32_t index) const = 0;
         /**
diff --git a/Engine/Plugins/Animation/RigLogic/Source/RigLogicLib/Public/dna/layers/GeometryReader.h b/Engine/Plugins/Animation/RigLogic/Source/RigLogicLib/Public/dna/layers/GeometryReader.h
--- a/Engine/Plugins/Animation/RigLogic/Source/RigLogicLib/Public/dna/layers/GeometryReader.h
+++ b/Engine/Plugins/Animation/RigLogic/Source/RigLogicLib/Public/dna/layers/GeometryReader.h
@@ -332,7 +332,7 @@ class DNAAPI GeometryReader : public virtual DefinitionReader {
                 A position in the zero-indexed array of blend shape targets within the specified mesh.
             @warning
                 blendShapeTargetIndex must be less than the value returned by getBlendShapeTargetCount.
-            @see DefinitionReader::getBlendShapeName
+            @see DefinitionReader::getBlendShapeChannelName
         */
         virtual std::uint16_t getBlendShapeChannelIndex(std::uint16_t meshIndex, std::uint16_t blendShapeTargetIndex) const = 0;
         /**
diff --git a/Engine/Plugins/Animation/RigLogic/Source/RigLogicLib/Public/dna/version/Version.h b/Engine/Plugins/Animation/RigLogic/Source/RigLogicLib/Public/dna/version/Version.h
--- a/Engine/Plugins/Animation/RigLogic/Source/RigLogicLib/Public/dna/version/Version.h
+++ b/Engine/Plugins/Animation/RigLogic/Source/RigLogicLib/Public/dna/version/Version.h
@@ -4,5 +4,5 @@
 
 #define DNA_MAJOR_VERSION 9
 #define DNA_MINOR_VERSION 1
-#define DNA_PATCH_VERSION 4
-#define DNA_VERSION_STRING "9.1.4"
+#define DNA_PATCH_VERSION 8
+#define DNA_VERSION_STRING "9.1.8"
diff --git a/Engine/Plugins/Animation/RigLogic/Source/RigLogicLib/Public/pma/Defs.h b/Engine/Plugins/Animation/RigLogic/Source/RigLogicLib/Public/pma/Defs.h
--- a/Engine/Plugins/Animation/RigLogic/Source/RigLogicLib/Public/pma/Defs.h
+++ b/Engine/Plugins/Animation/RigLogic/Source/RigLogicLib/Public/pma/Defs.h
@@ -2,17 +2,28 @@
 
 #pragma once
 
-#if defined(_WIN32) || defined(__CYGWIN__)
-    #if defined(__GNUC__)
-        #define DLL_EXPORT __attribute__((dllexport))
-        #define DLL_IMPORT __attribute__((dllimport))
-    #else
-        #define DLL_EXPORT __declspec(dllexport)
-        #define DLL_IMPORT __declspec(dllimport)
-    #endif
-#elif defined(__GNUC__)
+#if defined(FORCE_DLL_ATTRIBUTE_EXPORT_IMPORT)
+    #define DLL_EXPORT __attribute__((dllexport))
+    #define DLL_IMPORT __attribute__((dllimport))
+#elif defined(FORCE_DLL_DECLSPEC_EXPORT_IMPORT)
+    #define DLL_EXPORT __declspec(dllexport)
+    #define DLL_IMPORT __declspec(dllimport)
+#elif defined(FORCE_DLL_VISIBILITY_EXPORT_IMPORT)
     #define DLL_EXPORT __attribute__((visibility("default")))
     #define DLL_IMPORT DLL_EXPORT
+#else
+    #if defined(_WIN32) || defined(__CYGWIN__)
+        #if defined(__GNUC__)
+            #define DLL_EXPORT __attribute__((dllexport))
+            #define DLL_IMPORT __attribute__((dllimport))
+        #else
+            #define DLL_EXPORT __declspec(dllexport)
+            #define DLL_IMPORT __declspec(dllimport)
+        #endif
+    #elif defined(__GNUC__)
+        #define DLL_EXPORT __attribute__((visibility("default")))
+        #define DLL_IMPORT DLL_EXPORT
+    #endif
 #endif
 
 #if defined(RL_BUILD_SHARED)
diff --git a/Engine/Plugins/Animation/RigLogic/Source/RigLogicLib/Public/pma/version/Version.h b/Engine/Plugins/Animation/RigLogic/Source/RigLogicLib/Public/pma/version/Version.h
--- a/Engine/Plugins/Animation/RigLogic/Source/RigLogicLib/Public/pma/version/Version.h
+++ b/Engine/Plugins/Animation/RigLogic/Source/RigLogicLib/Public/pma/version/Version.h
@@ -4,5 +4,5 @@
 
 #define PMA_MAJOR_VERSION 1
 #define PMA_MINOR_VERSION 3
-#define PMA_PATCH_VERSION 7
-#define PMA_VERSION_STRING "1.3.7"
+#define PMA_PATCH_VERSION 9
+#define PMA_VERSION_STRING "1.3.9"
diff --git a/Engine/Plugins/Animation/RigLogic/Source/RigLogicLib/Public/riglogic/Defs.h b/Engine/Plugins/Animation/RigLogic/Source/RigLogicLib/Public/riglogic/Defs.h
--- a/Engine/Plugins/Animation/RigLogic/Source/RigLogicLib/Public/riglogic/Defs.h
+++ b/Engine/Plugins/Animation/RigLogic/Source/RigLogicLib/Public/riglogic/Defs.h
@@ -2,17 +2,28 @@
 
 #pragma once
 
-#if defined(_WIN32) || defined(__CYGWIN__)
-    #if defined(__GNUC__)
-        #define DLL_EXPORT __attribute__((dllexport))
-        #define DLL_IMPORT __attribute__((dllimport))
-    #else
-        #define DLL_EXPORT __declspec(dllexport)
-        #define DLL_IMPORT __declspec(dllimport)
-    #endif
-#elif defined(__GNUC__)
+#if defined(FORCE_DLL_ATTRIBUTE_EXPORT_IMPORT)
+    #define DLL_EXPORT __attribute__((dllexport))
+    #define DLL_IMPORT __attribute__((dllimport))
+#elif defined(FORCE_DLL_DECLSPEC_EXPORT_IMPORT)
+    #define DLL_EXPORT __declspec(dllexport)
+    #define DLL_IMPORT __declspec(dllimport)
+#elif defined(FORCE_DLL_VISIBILITY_EXPORT_IMPORT)
     #define DLL_EXPORT __attribute__((visibility("default")))
     #define DLL_IMPORT DLL_EXPORT
+#else
+    #if defined(_WIN32) || defined(__CYGWIN__)
+        #if defined(__GNUC__)
+            #define DLL_EXPORT __attribute__((dllexport))
+            #define DLL_IMPORT __attribute__((dllimport))
+        #else
+            #define DLL_EXPORT __declspec(dllexport)
+            #define DLL_IMPORT __declspec(dllimport)
+        #endif
+    #elif defined(__GNUC__)
+        #define DLL_EXPORT __attribute__((visibility("default")))
+        #define DLL_IMPORT DLL_EXPORT
+    #endif
 #endif
 
 #if defined(RL_BUILD_SHARED)
diff --git a/Engine/Plugins/Animation/RigLogic/Source/RigLogicLib/Public/riglogic/riglogic/Configuration.h b/Engine/Plugins/Animation/RigLogic/Source/RigLogicLib/Public/riglogic/riglogic/Configuration.h
--- a/Engine/Plugins/Animation/RigLogic/Source/RigLogicLib/Public/riglogic/riglogic/Configuration.h
+++ b/Engine/Plugins/Animation/RigLogic/Source/RigLogicLib/Public/riglogic/riglogic/Configuration.h
@@ -10,8 +10,10 @@ namespace rl4 {
 enum class CalculationType {
     Scalar,  ///< scalar CPU algorithm
     SSE,  ///< vectorized (SSE) CPU algorithm
-    AVX  ///< vectorized (AVX) CPU algorithm (RigLogic must be built with AVX support,
-         ///< otherwise it falls back to using the Scalar version)
+    AVX,  ///< vectorized (AVX) CPU algorithm (RigLogic must be built with AVX support,
+          ///< otherwise it falls back to using the Scalar version)
+    NEON  ///< vectorized (NEON) CPU algorithm (RigLogic must be built with NEON support,
+          ///< otherwise it falls back to using the Scalar version)
 };
 
 struct Configuration {
diff --git a/Engine/Plugins/Animation/RigLogic/Source/RigLogicLib/Public/riglogic/version/Version.h b/Engine/Plugins/Animation/RigLogic/Source/RigLogicLib/Public/riglogic/version/Version.h
--- a/Engine/Plugins/Animation/RigLogic/Source/RigLogicLib/Public/riglogic/version/Version.h
+++ b/Engine/Plugins/Animation/RigLogic/Source/RigLogicLib/Public/riglogic/version/Version.h
@@ -3,6 +3,6 @@
 #pragma once
 
 #define RL_MAJOR_VERSION 9
-#define RL_MINOR_VERSION 1
-#define RL_PATCH_VERSION 3
-#define RL_VERSION_STRING "9.1.3"
+#define RL_MINOR_VERSION 2
+#define RL_PATCH_VERSION 2
+#define RL_VERSION_STRING "9.2.2"
diff --git a/Engine/Plugins/Animation/RigLogic/Source/RigLogicLib/Public/status/Defs.h b/Engine/Plugins/Animation/RigLogic/Source/RigLogicLib/Public/status/Defs.h
--- a/Engine/Plugins/Animation/RigLogic/Source/RigLogicLib/Public/status/Defs.h
+++ b/Engine/Plugins/Animation/RigLogic/Source/RigLogicLib/Public/status/Defs.h
@@ -2,17 +2,28 @@
 
 #pragma once
 
-#if defined(_WIN32) || defined(__CYGWIN__)
-    #if defined(__GNUC__)
-        #define DLL_EXPORT __attribute__((dllexport))
-        #define DLL_IMPORT __attribute__((dllimport))
-    #else
-        #define DLL_EXPORT __declspec(dllexport)
-        #define DLL_IMPORT __declspec(dllimport)
-    #endif
-#elif defined(__GNUC__)
+#if defined(FORCE_DLL_ATTRIBUTE_EXPORT_IMPORT)
+    #define DLL_EXPORT __attribute__((dllexport))
+    #define DLL_IMPORT __attribute__((dllimport))
+#elif defined(FORCE_DLL_DECLSPEC_EXPORT_IMPORT)
+    #define DLL_EXPORT __declspec(dllexport)
+    #define DLL_IMPORT __declspec(dllimport)
+#elif defined(FORCE_DLL_VISIBILITY_EXPORT_IMPORT)
     #define DLL_EXPORT __attribute__((visibility("default")))
     #define DLL_IMPORT DLL_EXPORT
+#else
+    #if defined(_WIN32) || defined(__CYGWIN__)
+        #if defined(__GNUC__)
+            #define DLL_EXPORT __attribute__((dllexport))
+            #define DLL_IMPORT __attribute__((dllimport))
+        #else
+            #define DLL_EXPORT __declspec(dllexport)
+            #define DLL_IMPORT __declspec(dllimport)
+        #endif
+    #elif defined(__GNUC__)
+        #define DLL_EXPORT __attribute__((visibility("default")))
+        #define DLL_IMPORT DLL_EXPORT
+    #endif
 #endif
 
 #if defined(RL_BUILD_SHARED)
diff --git a/Engine/Plugins/Animation/RigLogic/Source/RigLogicLib/Public/status/version/Version.h b/Engine/Plugins/Animation/RigLogic/Source/RigLogicLib/Public/status/version/Version.h
--- a/Engine/Plugins/Animation/RigLogic/Source/RigLogicLib/Public/status/version/Version.h
+++ b/Engine/Plugins/Animation/RigLogic/Source/RigLogicLib/Public/status/version/Version.h
@@ -4,5 +4,5 @@
 
 #define SC_MAJOR_VERSION 1
 #define SC_MINOR_VERSION 2
-#define SC_PATCH_VERSION 2
-#define SC_VERSION_STRING "1.2.2"
+#define SC_PATCH_VERSION 3
+#define SC_VERSION_STRING "1.2.3"
diff --git a/Engine/Plugins/Animation/RigLogic/Source/RigLogicLib/Public/terse/version/Version.h b/Engine/Plugins/Animation/RigLogic/Source/RigLogicLib/Public/terse/version/Version.h
--- a/Engine/Plugins/Animation/RigLogic/Source/RigLogicLib/Public/terse/version/Version.h
+++ b/Engine/Plugins/Animation/RigLogic/Source/RigLogicLib/Public/terse/version/Version.h
@@ -4,5 +4,5 @@
 
 #define TERSE_MAJOR_VERSION 3
 #define TERSE_MINOR_VERSION 1
-#define TERSE_PATCH_VERSION 6
-#define TERSE_VERSION_STRING "3.1.6"
+#define TERSE_PATCH_VERSION 7
+#define TERSE_VERSION_STRING "3.1.7"
diff --git a/Engine/Plugins/Animation/RigLogic/Source/RigLogicLib/Public/trimd/Fallback.h b/Engine/Plugins/Animation/RigLogic/Source/RigLogicLib/Public/trimd/Fallback.h
--- a/Engine/Plugins/Animation/RigLogic/Source/RigLogicLib/Public/trimd/Fallback.h
+++ b/Engine/Plugins/Animation/RigLogic/Source/RigLogicLib/Public/trimd/Fallback.h
@@ -2,9 +2,6 @@
 
 #pragma once
 
-#include <cstddef>
-#include <utility>
-
 namespace trimd {
 
 namespace fallback {
diff --git a/Engine/Plugins/Animation/RigLogic/Source/RigLogicLib/Public/trimd/NEON.h b/Engine/Plugins/Animation/RigLogic/Source/RigLogicLib/Public/trimd/NEON.h
new file mode 100644
--- /dev/null
+++ b/Engine/Plugins/Animation/RigLogic/Source/RigLogicLib/Public/trimd/NEON.h
@@ -0,0 +1,261 @@
+// Copyright Epic Games, Inc. All Rights Reserved.
+
+#pragma once
+
+// *INDENT-OFF*
+#ifdef TRIMD_ENABLE_NEON
+#include "trimd/Fallback.h"
+
+#include <arm_neon.h>
+
+#include <cstring>
+
+namespace trimd {
+
+namespace neon {
+
+struct F128 {
+    using value_type = float;
+
+    float32x4_t data;
+
+    F128() : data{vdupq_n_f32(0)} {
+    }
+
+    explicit F128(float32x4_t value) : data{value} {
+    }
+
+    explicit F128(float value) : F128{vdupq_n_f32(value)} {
+    }
+
+    F128(float v1, float v2, float v3, float v4) : data{v1, v2, v3, v4} {
+    }
+
+    static F128 fromAlignedSource(const float* source) {
+        return F128{vld1q_f32(source)};
+    }
+
+    static F128 fromUnalignedSource(const float* source) {
+        return F128{vld1q_f32(source)};
+    }
+
+    static F128 loadSingleValue(const float* source) {
+        return F128{vsetq_lane_f32(*source, vdupq_n_f32(0), 0)};
+    }
+
+    #ifdef TRIMD_ENABLE_NEON_FP16
+        static F128 fromAlignedSource(const std::uint16_t* source) {
+            return F128{vcvt_f32_f16(vld1_f16(reinterpret_cast<const float16_t*>(source)))};
+        }
+
+        static F128 fromUnalignedSource(const std::uint16_t* source) {
+            return F128{vcvt_f32_f16(vld1_f16(reinterpret_cast<const float16_t*>(source)))};
+        }
+
+        static F128 loadSingleValue(const std::uint16_t* source) {
+            float16_t value;
+            std::memcpy(&value, source, sizeof(float16_t));
+            return F128{vcvt_f32_f16(vset_lane_f16(value, vdup_n_f16(0), 0))};
+        }
+
+    #endif  // TRIMD_ENABLE_NEON_FP16
+
+    template<typename T>
+    static void prefetchT0(const T*  /*unused*/) {
+    }
+
+    template<typename T>
+    static void prefetchT1(const T*  /*unused*/) {
+    }
+
+    template<typename T>
+    static void prefetchT2(const T*  /*unused*/) {
+    }
+
+    template<typename T>
+    static void prefetchNTA(const T*  /*unused*/) {
+    }
+
+    void alignedLoad(const float* source) {
+        data = vld1q_f32(source);
+    }
+
+    void unalignedLoad(const float* source) {
+        data = vld1q_f32(source);
+    }
+
+    void alignedStore(float* dest) const {
+        vst1q_f32(dest, data);
+    }
+
+    void unalignedStore(float* dest) const {
+        vst1q_f32(dest, data);
+    }
+
+    #ifdef TRIMD_ENABLE_NEON_FP16
+    void alignedLoad(const std::uint16_t* source) {
+        data = vcvt_f32_f16(vld1_f16(reinterpret_cast<const float16_t*>(source)));
+    }
+
+    void unalignedLoad(const std::uint16_t* source) {
+        data = vcvt_f32_f16(vld1_f16(reinterpret_cast<const float16_t*>(source)));
+    }
+
+    void alignedStore(std::uint16_t* dest) const {
+        vst1_f16(reinterpret_cast<float16_t*>(dest), vcvt_f16_f32(data));
+    }
+
+    void unalignedStore(std::uint16_t* dest) const {
+        vst1_f16(reinterpret_cast<float16_t*>(dest), vcvt_f16_f32(data));
+    }
+    #endif  // TRIMD_ENABLE_NEON_FP16
+
+    float sum() const {
+        const float32x2_t tmp = vadd_f32(vget_high_f32(data), vget_low_f32(data));
+        return vget_lane_f32(vpadd_f32(tmp, tmp), 0);
+    }
+
+    F128& operator+=(const F128& rhs) {
+        data = vaddq_f32(data, rhs.data);
+        return *this;
+    }
+
+    F128& operator-=(const F128& rhs) {
+        data = vsubq_f32(data, rhs.data);
+        return *this;
+    }
+
+    F128& operator*=(const F128& rhs) {
+        data = vmulq_f32(data, rhs.data);
+        return *this;
+    }
+
+    F128& operator/=(const F128& rhs) {
+        // reciprocal0 = 1 / rhs (initial estimate)
+        const float32x4_t reciprocal0 = vrecpeq_f32(rhs.data);
+        // Newton-Raphson step to refine the initial reciprocal estimate.
+        // If accuracy is not enough, additional refinement steps may be added (as many as needed)
+        // until desired accuracy is reached (just duplicate the below line).
+        // reciprocal1 = reciprocal0 * (2.0 - (reciprocal0 * rhs))
+        const float32x4_t reciprocal1 = vmulq_f32(reciprocal0, vrecpsq_f32(reciprocal0, rhs.data));
+        // data = data * reciprocal1
+        data = vmulq_f32(data, reciprocal1);
+        return *this;
+    }
+
+    F128& operator&=(const F128& rhs) {
+        data = vreinterpretq_f32_u32(vandq_u32(vreinterpretq_u32_f32(data), vreinterpretq_u32_f32(rhs.data)));
+        return *this;
+    }
+
+    F128& operator|=(const F128& rhs) {
+        data = vreinterpretq_f32_u32(vorrq_u32(vreinterpretq_u32_f32(data), vreinterpretq_u32_f32(rhs.data)));
+        return *this;
+    }
+
+    F128& operator^=(const F128& rhs) {
+        data = vreinterpretq_f32_u32(veorq_u32(vreinterpretq_u32_f32(data), vreinterpretq_u32_f32(rhs.data)));
+        return *this;
+    }
+
+    static constexpr std::size_t size() {
+        return sizeof(decltype(data)) / sizeof(float);
+    }
+
+    static constexpr std::size_t alignment() {
+        return alignof(decltype(data));
+    }
+
+};
+
+inline F128 operator==(const F128& lhs, const F128& rhs) {
+    return F128{vreinterpretq_f32_u32(vceqq_f32(lhs.data, rhs.data))};
+}
+
+inline F128 operator!=(const F128& lhs, const F128& rhs) {
+    return F128{vreinterpretq_f32_u32(vmvnq_u32(vceqq_f32(lhs.data, rhs.data)))};
+}
+
+inline F128 operator<(const F128& lhs, const F128& rhs) {
+    return F128{vreinterpretq_f32_u32(vcltq_f32(lhs.data, rhs.data))};
+}
+
+inline F128 operator<=(const F128& lhs, const F128& rhs) {
+    return F128{vreinterpretq_f32_u32(vcleq_f32(lhs.data, rhs.data))};
+}
+
+inline F128 operator>(const F128& lhs, const F128& rhs) {
+    return F128{vreinterpretq_f32_u32(vcgtq_f32(lhs.data, rhs.data))};
+}
+
+inline F128 operator>=(const F128& lhs, const F128& rhs) {
+    return F128{vreinterpretq_f32_u32(vcgeq_f32(lhs.data, rhs.data))};
+}
+
+inline F128 operator+(const F128& lhs, const F128& rhs) {
+    return F128(lhs) += rhs;
+}
+
+inline F128 operator-(const F128& lhs, const F128& rhs) {
+    return F128(lhs) -= rhs;
+}
+
+inline F128 operator*(const F128& lhs, const F128& rhs) {
+    return F128(lhs) *= rhs;
+}
+
+inline F128 operator/(const F128& lhs, const F128& rhs) {
+    return F128(lhs) /= rhs;
+}
+
+inline F128 operator&(const F128& lhs, const F128& rhs) {
+    return F128(lhs) &= rhs;
+}
+
+inline F128 operator|(const F128& lhs, const F128& rhs) {
+    return F128(lhs) |= rhs;
+}
+
+inline F128 operator^(const F128& lhs, const F128& rhs) {
+    return F128(lhs) ^= rhs;
+}
+
+inline F128 operator~(const F128& rhs) {
+    return F128{vreinterpretq_f32_u32(vmvnq_u32(vreinterpretq_u32_f32(rhs.data)))};
+}
+
+inline void transpose(F128& row0, F128& row1, F128& row2, F128& row3) {
+    // row01 = [row0.x, row1.x, row0.z, row1.z], [row0.y, row1.y, row0.w, row1.w]
+    float32x4x2_t row01 = vtrnq_f32(row0.data, row1.data);
+    // row23 = [row2.x, row3.x, row2.z, row3.z], [row2.y, row3.y, row2.w, row3.w]
+    float32x4x2_t row23 = vtrnq_f32(row2.data, row3.data);
+
+    // row0 = row0.x, row1.x, row2.x, row3.x
+    row0 = F128{vcombine_f32(vget_low_f32(row01.val[0]), vget_low_f32(row23.val[0]))};
+    // row1 = row0.y, row1.y, row2.y, row3.y
+    row1 = F128{vcombine_f32(vget_low_f32(row01.val[1]), vget_low_f32(row23.val[1]))};
+    // row2 = row0.z, row1.z, row2.z, row3.z
+    row2 = F128{vcombine_f32(vget_high_f32(row01.val[0]), vget_high_f32(row23.val[0]))};
+    // row3 = row0.w, row1.w, row2.w, row3.w
+    row3 = F128{vcombine_f32(vget_high_f32(row01.val[1]), vget_high_f32(row23.val[1]))};
+}
+
+inline F128 abs(const F128& rhs) {
+    return F128{vabsq_f32(rhs.data)};
+}
+
+inline F128 andnot(const F128& lhs, const F128& rhs) {
+    return F128{vreinterpretq_f32_u32(vbicq_u32(vreinterpretq_u32_f32(rhs.data), vreinterpretq_u32_f32(lhs.data)))};
+}
+
+using F256 = fallback::T256<F128>;
+using fallback::transpose;
+using fallback::abs;
+using fallback::andnot;
+
+} // namespace neon
+
+} // namespace trimd
+
+#endif  // TRIMD_ENABLE_NEON
+// *INDENT-ON*
diff --git a/Engine/Plugins/Animation/RigLogic/Source/RigLogicLib/Public/trimd/TRiMD.h b/Engine/Plugins/Animation/RigLogic/Source/RigLogicLib/Public/trimd/TRiMD.h
--- a/Engine/Plugins/Animation/RigLogic/Source/RigLogicLib/Public/trimd/TRiMD.h
+++ b/Engine/Plugins/Animation/RigLogic/Source/RigLogicLib/Public/trimd/TRiMD.h
@@ -3,6 +3,7 @@
 #pragma once
 
 #include "trimd/AVX.h"
+#include "trimd/NEON.h"
 #include "trimd/SSE.h"
 #include "trimd/Scalar.h"
 
@@ -15,15 +16,22 @@ namespace trimd {
     using avx::andnot;
 #elif defined(TRIMD_ENABLE_SSE)
     using F256 = sse::F256;
+#elif defined(TRIMD_ENABLE_NEON)
+    using F256 = neon::F256;
 #else
     using F256 = scalar::F256;
 #endif  // TRIMD_ENABLE_AVX
 
-#ifdef TRIMD_ENABLE_SSE
+#if defined(TRIMD_ENABLE_SSE)
     using F128 = sse::F128;
     using sse::abs;
     using sse::transpose;
     using sse::andnot;
+#elif defined(TRIMD_ENABLE_NEON)
+    using F128 = neon::F128;
+    using neon::abs;
+    using neon::transpose;
+    using neon::andnot;
 #else
     using F128 = scalar::F128;
 #endif  // TRIMD_ENABLE_SSE
diff --git a/Engine/Plugins/Animation/RigLogic/Source/RigLogicLib/Public/trimd/version/Version.h b/Engine/Plugins/Animation/RigLogic/Source/RigLogicLib/Public/trimd/version/Version.h
--- a/Engine/Plugins/Animation/RigLogic/Source/RigLogicLib/Public/trimd/version/Version.h
+++ b/Engine/Plugins/Animation/RigLogic/Source/RigLogicLib/Public/trimd/version/Version.h
@@ -3,6 +3,6 @@
 #pragma once
 
 #define TRIMD_MAJOR_VERSION 2
-#define TRIMD_MINOR_VERSION 1
-#define TRIMD_PATCH_VERSION 5
-#define TRIMD_VERSION_STRING "2.1.5"
+#define TRIMD_MINOR_VERSION 2
+#define TRIMD_PATCH_VERSION 0
+#define TRIMD_VERSION_STRING "2.2.0"
diff --git a/Engine/Plugins/Animation/RigLogic/Source/RigLogicLib/Public/trio/Defs.h b/Engine/Plugins/Animation/RigLogic/Source/RigLogicLib/Public/trio/Defs.h
--- a/Engine/Plugins/Animation/RigLogic/Source/RigLogicLib/Public/trio/Defs.h
+++ b/Engine/Plugins/Animation/RigLogic/Source/RigLogicLib/Public/trio/Defs.h
@@ -2,17 +2,28 @@
 
 #pragma once
 
-#if defined(_WIN32) || defined(__CYGWIN__)
-    #if defined(__GNUC__)
-        #define DLL_EXPORT __attribute__((dllexport))
-        #define DLL_IMPORT __attribute__((dllimport))
-    #else
-        #define DLL_EXPORT __declspec(dllexport)
-        #define DLL_IMPORT __declspec(dllimport)
-    #endif
-#elif defined(__GNUC__)
+#if defined(FORCE_DLL_ATTRIBUTE_EXPORT_IMPORT)
+    #define DLL_EXPORT __attribute__((dllexport))
+    #define DLL_IMPORT __attribute__((dllimport))
+#elif defined(FORCE_DLL_DECLSPEC_EXPORT_IMPORT)
+    #define DLL_EXPORT __declspec(dllexport)
+    #define DLL_IMPORT __declspec(dllimport)
+#elif defined(FORCE_DLL_VISIBILITY_EXPORT_IMPORT)
     #define DLL_EXPORT __attribute__((visibility("default")))
     #define DLL_IMPORT DLL_EXPORT
+#else
+    #if defined(_WIN32) || defined(__CYGWIN__)
+        #if defined(__GNUC__)
+            #define DLL_EXPORT __attribute__((dllexport))
+            #define DLL_IMPORT __attribute__((dllimport))
+        #else
+            #define DLL_EXPORT __declspec(dllexport)
+            #define DLL_IMPORT __declspec(dllimport)
+        #endif
+    #elif defined(__GNUC__)
+        #define DLL_EXPORT __attribute__((visibility("default")))
+        #define DLL_IMPORT DLL_EXPORT
+    #endif
 #endif
 
 #if defined(RL_BUILD_SHARED)
diff --git a/Engine/Plugins/Animation/RigLogic/Source/RigLogicLib/Public/trio/version/Version.h b/Engine/Plugins/Animation/RigLogic/Source/RigLogicLib/Public/trio/version/Version.h
--- a/Engine/Plugins/Animation/RigLogic/Source/RigLogicLib/Public/trio/version/Version.h
+++ b/Engine/Plugins/Animation/RigLogic/Source/RigLogicLib/Public/trio/version/Version.h
@@ -4,5 +4,5 @@
 
 #define TRIO_MAJOR_VERSION 4
 #define TRIO_MINOR_VERSION 0
-#define TRIO_PATCH_VERSION 8
-#define TRIO_VERSION_STRING "4.0.8"
+#define TRIO_PATCH_VERSION 11
+#define TRIO_VERSION_STRING "4.0.11"
diff --git a/Engine/Plugins/Animation/RigLogic/Source/RigLogicLibTest/Private/dnatests/Fixturesv21.cpp b/Engine/Plugins/Animation/RigLogic/Source/RigLogicLibTest/Private/dnatests/Fixturesv21.cpp
--- a/Engine/Plugins/Animation/RigLogic/Source/RigLogicLibTest/Private/dnatests/Fixturesv21.cpp
+++ b/Engine/Plugins/Animation/RigLogic/Source/RigLogicLibTest/Private/dnatests/Fixturesv21.cpp
@@ -787,7 +787,7 @@ const unsigned char RawV21::geometry[] = {
     0x00, 0x00, 0x00, 0x00,  // Vertex position: 0
     0x00, 0x00, 0x00, 0x01,  // Vertex position: 1
     0x00, 0x00, 0x00, 0x02,  // Vertex position: 2
-    0x00, 0x06,  // Blend shape index in Definition
+    0x00, 0x02,  // Blend shape index in Definition
     // Mesh-1
     0x00, 0x00, 0x0b, 0x8e,  // Mesh-1 end offset
     0x00, 0x00, 0x00, 0x03,  // Vertex positions X values length
@@ -878,7 +878,7 @@ const unsigned char RawV21::geometry[] = {
     0x00, 0x00, 0x00, 0x00,  // Vertex position: 0
     0x00, 0x00, 0x00, 0x01,  // Vertex position: 1
     0x00, 0x00, 0x00, 0x02,  // Vertex position: 2
-    0x00, 0x06,  // Blend shape index in Definition
+    0x00, 0x02,  // Blend shape index in Definition
     // Mesh-2
     0x00, 0x00, 0x0d, 0x16,  // Mesh-2 offset
     0x00, 0x00, 0x00, 0x03,  // Vertex positions X values length
@@ -969,7 +969,7 @@ const unsigned char RawV21::geometry[] = {
     0x00, 0x00, 0x00, 0x00,  // Vertex position: 0
     0x00, 0x00, 0x00, 0x01,  // Vertex position: 1
     0x00, 0x00, 0x00, 0x02,  // Vertex position: 2
-    0x00, 0x06,  // Blend shape index in Definition
+    0x00, 0x02,  // Blend shape index in Definition
     0x00, 0x00, 0x00, 0x02,  // Blend shape deltas X values length
     0x40, 0x80, 0x00, 0x00,  // 4.0f
     0x40, 0xa0, 0x00, 0x00,  // 5.0f
@@ -982,7 +982,7 @@ const unsigned char RawV21::geometry[] = {
     0x00, 0x00, 0x00, 0x02,  // Vertex position indices length (for each delta)
     0x00, 0x00, 0x00, 0x00,  // Vertex position: 0
     0x00, 0x00, 0x00, 0x02,  // Vertex position: 2
-    0x00, 0x07  // Blend shape index in Definition
+    0x00, 0x03  // Blend shape index in Definition
 };
 
 const unsigned char RawV21::footer[] = {
@@ -1925,26 +1925,26 @@ const pma::Matrix<pma::Matrix<std::uint16_t> > DecodedV21::skinWeightsJointIndic
 const pma::Vector<pma::Matrix<std::uint16_t> > DecodedV21::correctiveBlendShapeIndices = {
     {  // MaxLOD-0 - MaxLOD-1
         {  // Mesh-0
-            6
+            2
         },
         {  // Mesh-1
-            6
+            2
         },
         {  // Mesh-2
-            6, 7
+            2, 3
         }
     },
     {  // MaxLOD-1 - MinLOD-1
         {  // Mesh-0 (Mesh-2 under MaxLOD-0)
-            7
+            2
         }
     },
     {  // MaxLOD-0 - MinLOD-0
         {  // Mesh-0
-            6
+            2
         },
         {  // Mesh-1
-            6
+            2
         }
     }
 };
@@ -1978,9 +1978,10 @@ const pma::Matrix<pma::Matrix<Vector3> > DecodedV21::correctiveBlendShapeDeltas
     },
     {  // MaxLOD-1 - MinLOD-1
         {  // Mesh-0 (Mesh-2 under MaxLOD-0)
-            {  // Blendshape-1
-                {4.0f, 4.0f, 4.0f},
-                {5.0f, 5.0f, 5.0f}
+            {  // Blendshape-0
+                {1.0f, 1.0f, 1.0f},
+                {2.0f, 2.0f, 2.0f},
+                {3.0f, 3.0f, 3.0f}
             }
         }
     },
@@ -2016,7 +2017,7 @@ const pma::Matrix<pma::Matrix<std::uint32_t> > DecodedV21::correctiveBlendShapeV
     },
     {  // MaxLOD-1 - MinLOD-1
         {  // Mesh-0 (Mesh-2 under MaxLOD-0)
-            {0, 2}  // Blendshape-1
+            {0, 1, 2}  // Blendshape-0
         }
     },
     {  // MaxLOD-0 - MinLOD-0
diff --git a/Engine/Plugins/Animation/RigLogic/Source/RigLogicLibTest/Private/dnatests/Fixturesv22.cpp b/Engine/Plugins/Animation/RigLogic/Source/RigLogicLibTest/Private/dnatests/Fixturesv22.cpp
--- a/Engine/Plugins/Animation/RigLogic/Source/RigLogicLibTest/Private/dnatests/Fixturesv22.cpp
+++ b/Engine/Plugins/Animation/RigLogic/Source/RigLogicLibTest/Private/dnatests/Fixturesv22.cpp
@@ -804,7 +804,7 @@ const unsigned char RawV22::geometry[] = {
     0x00, 0x00, 0x00, 0x00,  // Vertex position: 0
     0x00, 0x00, 0x00, 0x01,  // Vertex position: 1
     0x00, 0x00, 0x00, 0x02,  // Vertex position: 2
-    0x00, 0x06,  // Blend shape index in Definition
+    0x00, 0x02,  // Blend shape index in Definition
     // Mesh-1
     0x00, 0x00, 0x01, 0x52,  // Mesh-1 size
     0x00, 0x00, 0x00, 0x03,  // Vertex positions X values length
@@ -895,7 +895,7 @@ const unsigned char RawV22::geometry[] = {
     0x00, 0x00, 0x00, 0x00,  // Vertex position: 0
     0x00, 0x00, 0x00, 0x01,  // Vertex position: 1
     0x00, 0x00, 0x00, 0x02,  // Vertex position: 2
-    0x00, 0x06,  // Blend shape index in Definition
+    0x00, 0x02,  // Blend shape index in Definition
     // Mesh-2
     0x00, 0x00, 0x01, 0x84,  // Mesh-2 size
     0x00, 0x00, 0x00, 0x03,  // Vertex positions X values length
@@ -986,7 +986,7 @@ const unsigned char RawV22::geometry[] = {
     0x00, 0x00, 0x00, 0x00,  // Vertex position: 0
     0x00, 0x00, 0x00, 0x01,  // Vertex position: 1
     0x00, 0x00, 0x00, 0x02,  // Vertex position: 2
-    0x00, 0x06,  // Blend shape index in Definition
+    0x00, 0x02,  // Blend shape index in Definition
     0x00, 0x00, 0x00, 0x02,  // Blend shape deltas X values length
     0x40, 0x80, 0x00, 0x00,  // 4.0f
     0x40, 0xa0, 0x00, 0x00,  // 5.0f
@@ -999,7 +999,7 @@ const unsigned char RawV22::geometry[] = {
     0x00, 0x00, 0x00, 0x02,  // Vertex position indices length (for each delta)
     0x00, 0x00, 0x00, 0x00,  // Vertex position: 0
     0x00, 0x00, 0x00, 0x02,  // Vertex position: 2
-    0x00, 0x07  // Blend shape index in Definition
+    0x00, 0x03  // Blend shape index in Definition
 };
 
 const unsigned char RawV22::unknownLayer1[] = {
@@ -2416,26 +2416,26 @@ const pma::Matrix<pma::Matrix<std::uint16_t> > DecodedV22::skinWeightsJointIndic
 const pma::Vector<pma::Matrix<std::uint16_t> > DecodedV22::correctiveBlendShapeIndices = {
     {  // MaxLOD-0 - MaxLOD-1
         {  // Mesh-0
-            6
+            2
         },
         {  // Mesh-1
-            6
+            2
         },
         {  // Mesh-2
-            6, 7
+            2, 3
         }
     },
     {  // MaxLOD-1 - MinLOD-1
         {  // Mesh-0 (Mesh-2 under MaxLOD-0)
-            7
+            2
         }
     },
     {  // MaxLOD-0 - MinLOD-0
         {  // Mesh-0
-            6
+            2
         },
         {  // Mesh-1
-            6
+            2
         }
     }
 };
@@ -2469,9 +2469,10 @@ const pma::Matrix<pma::Matrix<Vector3> > DecodedV22::correctiveBlendShapeDeltas
     },
     {  // MaxLOD-1 - MinLOD-1
         {  // Mesh-0 (Mesh-2 under MaxLOD-0)
-            {  // Blendshape-1
-                {4.0f, 4.0f, 4.0f},
-                {5.0f, 5.0f, 5.0f}
+            {  // Blendshape-0
+                {1.0f, 1.0f, 1.0f},
+                {2.0f, 2.0f, 2.0f},
+                {3.0f, 3.0f, 3.0f}
             }
         }
     },
@@ -2507,7 +2508,7 @@ const pma::Matrix<pma::Matrix<std::uint32_t> > DecodedV22::correctiveBlendShapeV
     },
     {  // MaxLOD-1 - MinLOD-1
         {  // Mesh-0 (Mesh-2 under MaxLOD-0)
-            {0, 2}  // Blendshape-1
+            {0, 1, 2}  // Blendshape-0
         }
     },
     {  // MaxLOD-0 - MinLOD-0
diff --git a/Engine/Plugins/Animation/RigLogic/Source/RigLogicLibTest/Private/dnatests/Fixturesv23.cpp b/Engine/Plugins/Animation/RigLogic/Source/RigLogicLibTest/Private/dnatests/Fixturesv23.cpp
--- a/Engine/Plugins/Animation/RigLogic/Source/RigLogicLibTest/Private/dnatests/Fixturesv23.cpp
+++ b/Engine/Plugins/Animation/RigLogic/Source/RigLogicLibTest/Private/dnatests/Fixturesv23.cpp
@@ -802,7 +802,7 @@ const unsigned char RawV23::geometry[] = {
     0x00, 0x00, 0x00, 0x00,  // Vertex position: 0
     0x00, 0x00, 0x00, 0x01,  // Vertex position: 1
     0x00, 0x00, 0x00, 0x02,  // Vertex position: 2
-    0x00, 0x06,  // Blend shape index in Definition
+    0x00, 0x02,  // Blend shape index in Definition
     // Mesh-1
     0x00, 0x00, 0x01, 0x52,  // Mesh-1 size
     0x00, 0x00, 0x00, 0x03,  // Vertex positions X values length
@@ -893,7 +893,7 @@ const unsigned char RawV23::geometry[] = {
     0x00, 0x00, 0x00, 0x00,  // Vertex position: 0
     0x00, 0x00, 0x00, 0x01,  // Vertex position: 1
     0x00, 0x00, 0x00, 0x02,  // Vertex position: 2
-    0x00, 0x06,  // Blend shape index in Definition
+    0x00, 0x02,  // Blend shape index in Definition
     // Mesh-2
     0x00, 0x00, 0x01, 0x84,  // Mesh-2 size
     0x00, 0x00, 0x00, 0x03,  // Vertex positions X values length
@@ -984,7 +984,7 @@ const unsigned char RawV23::geometry[] = {
     0x00, 0x00, 0x00, 0x00,  // Vertex position: 0
     0x00, 0x00, 0x00, 0x01,  // Vertex position: 1
     0x00, 0x00, 0x00, 0x02,  // Vertex position: 2
-    0x00, 0x06,  // Blend shape index in Definition
+    0x00, 0x02,  // Blend shape index in Definition
     0x00, 0x00, 0x00, 0x02,  // Blend shape deltas X values length
     0x40, 0x80, 0x00, 0x00,  // 4.0f
     0x40, 0xa0, 0x00, 0x00,  // 5.0f
@@ -997,7 +997,7 @@ const unsigned char RawV23::geometry[] = {
     0x00, 0x00, 0x00, 0x02,  // Vertex position indices length (for each delta)
     0x00, 0x00, 0x00, 0x00,  // Vertex position: 0
     0x00, 0x00, 0x00, 0x02,  // Vertex position: 2
-    0x00, 0x07  // Blend shape index in Definition
+    0x00, 0x03  // Blend shape index in Definition
 };
 
 const unsigned char RawV23::machineLearnedBehavior[] = {
@@ -2172,26 +2172,26 @@ const pma::Matrix<pma::Matrix<std::uint16_t> > DecodedV23::skinWeightsJointIndic
 const pma::Vector<pma::Matrix<std::uint16_t> > DecodedV23::correctiveBlendShapeIndices = {
     {  // MaxLOD-0 - MaxLOD-1
         {  // Mesh-0
-            6
+            2
         },
         {  // Mesh-1
-            6
+            2
         },
         {  // Mesh-2
-            6, 7
+            2, 3
         }
     },
     {  // MaxLOD-1 - MinLOD-1
         {  // Mesh-0 (Mesh-2 under MaxLOD-0)
-            7
+            2
         }
     },
     {  // MaxLOD-0 - MinLOD-0
         {  // Mesh-0
-            6
+            2
         },
         {  // Mesh-1
-            6
+            2
         }
     }
 };
@@ -2225,9 +2225,10 @@ const pma::Matrix<pma::Matrix<Vector3> > DecodedV23::correctiveBlendShapeDeltas
     },
     {  // MaxLOD-1 - MinLOD-1
         {  // Mesh-0 (Mesh-2 under MaxLOD-0)
-            {  // Blendshape-1
-                {4.0f, 4.0f, 4.0f},
-                {5.0f, 5.0f, 5.0f}
+            {  // Blendshape-0
+                {1.0f, 1.0f, 1.0f},
+                {2.0f, 2.0f, 2.0f},
+                {3.0f, 3.0f, 3.0f}
             }
         }
     },
@@ -2263,7 +2264,7 @@ const pma::Matrix<pma::Matrix<std::uint32_t> > DecodedV23::correctiveBlendShapeV
     },
     {  // MaxLOD-1 - MinLOD-1
         {  // Mesh-0 (Mesh-2 under MaxLOD-0)
-            {0, 2}  // Blendshape-1
+            {0, 1, 2}  // Blendshape-0
         }
     },
     {  // MaxLOD-0 - MinLOD-0
diff --git a/Engine/Plugins/Animation/RigLogic/Source/RigLogicLibTest/Private/rltests/Defs.h b/Engine/Plugins/Animation/RigLogic/Source/RigLogicLibTest/Private/rltests/Defs.h
--- a/Engine/Plugins/Animation/RigLogic/Source/RigLogicLibTest/Private/rltests/Defs.h
+++ b/Engine/Plugins/Animation/RigLogic/Source/RigLogicLibTest/Private/rltests/Defs.h
@@ -4,7 +4,7 @@
 
 #ifdef _MSC_VER
     #pragma warning(push)
-    #pragma warning(disable : 4061 4121 4365 4987 4668 4996 6330)
+    #pragma warning(disable : 4061 4121 4324 4365 4987 4668 4996 6330)
 #endif
 #ifdef __clang__
     #pragma clang diagnostic push
diff --git a/Engine/Plugins/Animation/RigLogic/Source/RigLogicLibTest/Private/rltests/joints/bpcm/Helpers.h b/Engine/Plugins/Animation/RigLogic/Source/RigLogicLibTest/Private/rltests/joints/bpcm/Helpers.h
--- a/Engine/Plugins/Animation/RigLogic/Source/RigLogicLibTest/Private/rltests/joints/bpcm/Helpers.h
+++ b/Engine/Plugins/Animation/RigLogic/Source/RigLogicLibTest/Private/rltests/joints/bpcm/Helpers.h
@@ -2,6 +2,8 @@
 
 #pragma once
 
+#include "riglogic/riglogic/Configuration.h"
+
 #include <cstddef>
 #include <cstdint>
 
@@ -14,3 +16,21 @@ struct OutputScope {
 struct StrategyTestParams {
     std::uint16_t lod;
 };
+
+template<std::uint16_t LOD>
+struct TStrategyTestParams {
+
+    static constexpr std::uint16_t lod() {
+        return LOD;
+    }
+
+};
+
+template<rl4::CalculationType CalcType>
+struct TCalculationType {
+
+    static constexpr rl4::CalculationType get() {
+        return CalcType;
+    }
+
+};
diff --git a/Engine/Plugins/Animation/RigLogic/Source/RigLogicLibTest/Private/rltests/joints/bpcm/TestBlock4CalculationStrategy.cpp b/Engine/Plugins/Animation/RigLogic/Source/RigLogicLibTest/Private/rltests/joints/bpcm/TestBlock4CalculationStrategy.cpp
new file mode 100644
--- /dev/null
+++ b/Engine/Plugins/Animation/RigLogic/Source/RigLogicLibTest/Private/rltests/joints/bpcm/TestBlock4CalculationStrategy.cpp
@@ -0,0 +1,158 @@
+// Copyright Epic Games, Inc. All Rights Reserved.
+
+#include "riglogic/system/simd/Detect.h"
+
+#include "rltests/Defs.h"
+#include "rltests/controls/ControlFixtures.h"
+#include "rltests/joints/bpcm/BPCMFixturesBlock4.h"
+#include "rltests/joints/bpcm/Helpers.h"
+
+#include "riglogic/TypeDefs.h"
+#include "riglogic/joints/bpcm/Evaluator.h"
+#include "riglogic/joints/bpcm/strategies/Block4.h"
+#include "riglogic/system/simd/SIMD.h"
+
+namespace {
+
+template<typename TTestTypes>
+class Block4JointCalculationStrategyTest : public ::testing::TestWithParam<StrategyTestParams> {
+    protected:
+        void SetUp() override {
+            Block4JointCalculationStrategyTest::SetUpImpl();
+        }
+
+        template<typename TestTypes = TTestTypes>
+        typename std::enable_if<std::tuple_size<TestTypes>::value != 0ul, void>::type SetUpImpl() {
+            using T = typename std::tuple_element<0, TestTypes>::type;
+            using TFVec = typename std::tuple_element<1, TestTypes>::type;
+            using TStrategyTestParams = typename std::tuple_element<2, TestTypes>::type;
+            params.lod = TStrategyTestParams::lod();
+
+            using CalculationStrategyBase = rl4::bpcm::JointCalculationStrategy<T>;
+            using CalculationStrategy = rl4::bpcm::block4::Block4JointCalculationStrategy<T, TFVec>;
+            strategy = pma::UniqueInstance<CalculationStrategy, CalculationStrategyBase>::with(&memRes).create();
+        }
+
+        template<typename TestTypes = TTestTypes>
+        typename std::enable_if<std::tuple_size<TestTypes>::value == 0ul, void>::type SetUpImpl() {
+        }
+
+        template<typename TArray>
+        void execute(const rl4::bpcm::Evaluator<StorageValueType>& joints, const TArray& expected, OutputScope scope) {
+            if (strategy == nullptr) {
+                return;
+            }
+
+            auto outputInstance = joints.createInstance(&memRes);
+            auto outputBuffer = outputInstance->getOutputBuffer();
+
+            auto inputInstanceFactory =
+                ControlsFactory::getInstanceFactory(0, static_cast<std::uint16_t>(block4::input::values.size()), 0, 0);
+            auto inputInstance = inputInstanceFactory(&memRes);
+            auto inputBuffer = inputInstance->getInputBuffer();
+            std::copy(block4::input::values.begin(), block4::input::values.end(), inputBuffer.begin());
+
+            rl4::ConstArrayView<float> expectedView{expected[scope.lod].data() + scope.offset, scope.size};
+            rl4::ConstArrayView<float> outputView{outputBuffer.data() + scope.offset, scope.size};
+            joints.calculate(inputInstance.get(), outputInstance.get(), scope.lod);
+            ASSERT_EQ(outputView, expectedView);
+        }
+
+    protected:
+        pma::AlignedMemoryResource memRes;
+        block4::OptimizedStorage<StorageValueType>::StrategyPtr strategy;
+        StrategyTestParams params;
+
+};
+
+}  // namespace
+
+using Block4JointCalculationTypeList = ::testing::Types<
+#if defined(RL_BUILD_WITH_AVX) || defined(RL_BUILD_WITH_SSE)
+        std::tuple<StorageValueType, trimd::sse::F128, TStrategyTestParams<0> >,
+        std::tuple<StorageValueType, trimd::sse::F128, TStrategyTestParams<1> >,
+        std::tuple<StorageValueType, trimd::sse::F128, TStrategyTestParams<2> >,
+        std::tuple<StorageValueType, trimd::sse::F128, TStrategyTestParams<3> >,
+#endif  // RL_BUILD_WITH_AVX || RL_BUILD_WITH_SSE
+#if defined(RL_BUILD_WITH_NEON)
+        std::tuple<StorageValueType, trimd::neon::F128, TStrategyTestParams<0> >,
+        std::tuple<StorageValueType, trimd::neon::F128, TStrategyTestParams<1> >,
+        std::tuple<StorageValueType, trimd::neon::F128, TStrategyTestParams<2> >,
+        std::tuple<StorageValueType, trimd::neon::F128, TStrategyTestParams<3> >,
+#endif  // RL_BUILD_WITH_NEON
+#if !defined(RL_USE_HALF_FLOATS)
+        std::tuple<StorageValueType, trimd::scalar::F128, TStrategyTestParams<0> >,
+        std::tuple<StorageValueType, trimd::scalar::F128, TStrategyTestParams<1> >,
+        std::tuple<StorageValueType, trimd::scalar::F128, TStrategyTestParams<2> >,
+        std::tuple<StorageValueType, trimd::scalar::F128, TStrategyTestParams<3> >,
+#endif  // RL_USE_HALF_FLOATS
+    std::tuple<>
+    >;
+
+TYPED_TEST_SUITE(Block4JointCalculationStrategyTest, Block4JointCalculationTypeList, );
+
+TYPED_TEST(Block4JointCalculationStrategyTest, Block4Padded) {
+    const OutputScope scope{this->params.lod, 0ul, 1ul};
+    auto joints = block4::OptimizedStorage<StorageValueType>::create(std::move(this->strategy), 0u, &this->memRes);
+    this->execute(joints, block4::output::valuesPerLOD, scope);
+}
+
+TYPED_TEST(Block4JointCalculationStrategyTest, Block4Exact) {
+    const OutputScope scope{this->params.lod, 1ul, 4ul};
+    auto joints = block4::OptimizedStorage<StorageValueType>::create(std::move(this->strategy), 1u, &this->memRes);
+    this->execute(joints, block4::output::valuesPerLOD, scope);
+}
+
+TYPED_TEST(Block4JointCalculationStrategyTest, Block8Padded) {
+    const OutputScope scope{this->params.lod, 5ul, 7ul};
+    auto joints = block4::OptimizedStorage<StorageValueType>::create(std::move(this->strategy), 2u, &this->memRes);
+    this->execute(joints, block4::output::valuesPerLOD, scope);
+}
+
+TYPED_TEST(Block4JointCalculationStrategyTest, Block8Exact) {
+    const OutputScope scope{this->params.lod, 12ul, 8ul};
+    auto joints = block4::OptimizedStorage<StorageValueType>::create(std::move(this->strategy), 3u, &this->memRes);
+    this->execute(joints, block4::output::valuesPerLOD, scope);
+}
+
+TYPED_TEST(Block4JointCalculationStrategyTest, Block12Padded) {
+    const OutputScope scope{this->params.lod, 20ul, 9ul};
+    auto joints = block4::OptimizedStorage<StorageValueType>::create(std::move(this->strategy), 4u, &this->memRes);
+    this->execute(joints, block4::output::valuesPerLOD, scope);
+}
+
+TYPED_TEST(Block4JointCalculationStrategyTest, Block12Exact) {
+    const OutputScope scope{this->params.lod, 29ul, 12ul};
+    auto joints = block4::OptimizedStorage<StorageValueType>::create(std::move(this->strategy), 5u, &this->memRes);
+    this->execute(joints, block4::output::valuesPerLOD, scope);
+}
+
+TYPED_TEST(Block4JointCalculationStrategyTest, Block16Padded) {
+    const OutputScope scope{this->params.lod, 41ul, 14ul};
+    auto joints = block4::OptimizedStorage<StorageValueType>::create(std::move(this->strategy), 6u, &this->memRes);
+    this->execute(joints, block4::output::valuesPerLOD, scope);
+}
+
+TYPED_TEST(Block4JointCalculationStrategyTest, Block16Exact) {
+    const OutputScope scope{this->params.lod, 55ul, 16ul};
+    auto joints = block4::OptimizedStorage<StorageValueType>::create(std::move(this->strategy), 7u, &this->memRes);
+    this->execute(joints, block4::output::valuesPerLOD, scope);
+}
+
+TYPED_TEST(Block4JointCalculationStrategyTest, MultipleBlocks) {
+    const OutputScope scope{this->params.lod, 0ul, block4::output::valuesPerLOD.front().size()};
+    auto joints = block4::OptimizedStorage<StorageValueType>::create(std::move(this->strategy), &this->memRes);
+    this->execute(joints, block4::output::valuesPerLOD, scope);
+}
+
+TYPED_TEST(Block4JointCalculationStrategyTest, InputRegionA) {
+    const OutputScope scope{this->params.lod, 71ul, 2ul};
+    auto joints = block4::OptimizedStorage<StorageValueType>::create(std::move(this->strategy), 8u, &this->memRes);
+    this->execute(joints, block4::output::valuesPerLOD, scope);
+}
+
+TYPED_TEST(Block4JointCalculationStrategyTest, InputRegionB) {
+    const OutputScope scope{this->params.lod, 73ul, 2ul};
+    auto joints = block4::OptimizedStorage<StorageValueType>::create(std::move(this->strategy), 9u, &this->memRes);
+    this->execute(joints, block4::output::valuesPerLOD, scope);
+}
diff --git a/Engine/Plugins/Animation/RigLogic/Source/RigLogicLibTest/Private/rltests/joints/bpcm/TestBlock4StorageBuilder.cpp b/Engine/Plugins/Animation/RigLogic/Source/RigLogicLibTest/Private/rltests/joints/bpcm/TestBlock4StorageBuilder.cpp
new file mode 100644
--- /dev/null
+++ b/Engine/Plugins/Animation/RigLogic/Source/RigLogicLibTest/Private/rltests/joints/bpcm/TestBlock4StorageBuilder.cpp
@@ -0,0 +1,75 @@
+// Copyright Epic Games, Inc. All Rights Reserved.
+
+#include "riglogic/system/simd/Detect.h"
+
+#include "rltests/Defs.h"
+#include "rltests/joints/bpcm/Assertions.h"
+#include "rltests/joints/bpcm/BPCMFixturesBlock4.h"
+#include "rltests/joints/bpcm/Helpers.h"
+
+#include "riglogic/TypeDefs.h"
+#include "riglogic/joints/JointsBuilder.h"
+#include "riglogic/joints/bpcm/Evaluator.h"
+#include "riglogic/joints/bpcm/strategies/Block4.h"
+#include "riglogic/riglogic/RigLogic.h"
+#include "riglogic/system/simd/SIMD.h"
+
+namespace {
+
+template<typename TTestTypes>
+class Block4JointStorageBuilderTest : public ::testing::Test {
+    protected:
+        template<typename TestTypes = TTestTypes>
+        typename std::enable_if<std::tuple_size<TestTypes>::value == 0ul, void>::type buildStorage() {
+        }
+
+        template<typename TestTypes = TTestTypes>
+        typename std::enable_if<std::tuple_size<TestTypes>::value != 0ul, void>::type buildStorage() {
+            using TValue = typename std::tuple_element<0, TestTypes>::type;
+            using TFVec = typename std::tuple_element<1, TestTypes>::type;
+            using TCalculationType = typename std::tuple_element<2, TestTypes>::type;
+
+            rl4::Configuration config{};
+            config.calculationType = TCalculationType::get();
+            auto builder = rl4::JointsBuilder::create(config, &memRes);
+            builder->computeStorageRequirements(&reader);
+            builder->allocateStorage(&reader);
+            builder->fillStorage(&reader);
+            auto joints = builder->build();
+            auto jointsImpl = static_cast<rl4::bpcm::Evaluator<TValue>*>(joints.get());
+
+            auto strategy = pma::UniqueInstance<rl4::bpcm::block4::Block4JointCalculationStrategy<TValue, TFVec>,
+                                                rl4::bpcm::JointCalculationStrategy<TValue> >::with(&memRes).create();
+            auto expected = block4::OptimizedStorage<TValue>::create(std::move(strategy), &memRes);
+
+            rl4::bpcm::Evaluator<TValue>::Accessor::assertRawDataEqual(*jointsImpl, expected);
+            rl4::bpcm::Evaluator<TValue>::Accessor::assertJointGroupsEqual(*jointsImpl, expected);
+            rl4::bpcm::Evaluator<TValue>::Accessor::assertLODsEqual(*jointsImpl, expected);
+        }
+
+    protected:
+        pma::AlignedMemoryResource memRes;
+        block4::CanonicalReader reader;
+
+};
+
+}  // namespace
+
+using Block4StorageValueTypeList = ::testing::Types<
+#if defined(RL_BUILD_WITH_SSE)
+        std::tuple<StorageValueType, trimd::sse::F128, TCalculationType<rl4::CalculationType::SSE> >,
+#endif  // RL_BUILD_WITH_AVX || RL_BUILD_WITH_SSE
+#if defined(RL_BUILD_WITH_NEON)
+        std::tuple<StorageValueType, trimd::neon::F128, TCalculationType<rl4::CalculationType::NEON> >,
+#endif  // RL_BUILD_WITH_NEON
+#if !defined(RL_USE_HALF_FLOATS)
+        std::tuple<StorageValueType, trimd::scalar::F128, TCalculationType<rl4::CalculationType::Scalar> >,
+#endif  // RL_USE_HALF_FLOATS
+    std::tuple<>
+    >;
+
+TYPED_TEST_SUITE(Block4JointStorageBuilderTest, Block4StorageValueTypeList, );
+
+TYPED_TEST(Block4JointStorageBuilderTest, LayoutOptimization) {
+    this->buildStorage();
+}
diff --git a/Engine/Plugins/Animation/RigLogic/Source/RigLogicLibTest/Private/rltests/joints/bpcm/TestBlock8CalculationStrategy.cpp b/Engine/Plugins/Animation/RigLogic/Source/RigLogicLibTest/Private/rltests/joints/bpcm/TestBlock8CalculationStrategy.cpp
new file mode 100644
--- /dev/null
+++ b/Engine/Plugins/Animation/RigLogic/Source/RigLogicLibTest/Private/rltests/joints/bpcm/TestBlock8CalculationStrategy.cpp
@@ -0,0 +1,152 @@
+// Copyright Epic Games, Inc. All Rights Reserved.
+
+#include "riglogic/system/simd/Detect.h"
+
+#include "rltests/Defs.h"
+#include "rltests/controls/ControlFixtures.h"
+#include "rltests/joints/bpcm/BPCMFixturesBlock8.h"
+#include "rltests/joints/bpcm/Helpers.h"
+
+#include "riglogic/TypeDefs.h"
+#include "riglogic/joints/bpcm/Evaluator.h"
+#include "riglogic/joints/bpcm/strategies/Block8.h"
+#include "riglogic/system/simd/SIMD.h"
+
+namespace {
+
+template<typename TTestTypes>
+class Block8JointCalculationStrategyTest : public ::testing::Test {
+    protected:
+        void SetUp() override {
+            Block8JointCalculationStrategyTest::SetUpImpl();
+        }
+
+        template<typename TestTypes = TTestTypes>
+        typename std::enable_if<std::tuple_size<TestTypes>::value != 0ul, void>::type SetUpImpl() {
+            using T = typename std::tuple_element<0, TestTypes>::type;
+            using TFVec = typename std::tuple_element<1, TestTypes>::type;
+            using TStrategyTestParams = typename std::tuple_element<2, TestTypes>::type;
+            params.lod = TStrategyTestParams::lod();
+
+            using CalculationStrategyBase = rl4::bpcm::JointCalculationStrategy<T>;
+            using CalculationStrategy = rl4::bpcm::block8::Block8JointCalculationStrategy<T, TFVec>;
+            strategy = pma::UniqueInstance<CalculationStrategy, CalculationStrategyBase>::with(&memRes).create();
+        }
+
+        template<typename TestTypes = TTestTypes>
+        typename std::enable_if<std::tuple_size<TestTypes>::value == 0ul, void>::type SetUpImpl() {
+        }
+
+        template<typename TArray>
+        void execute(const rl4::bpcm::Evaluator<StorageValueType>& joints, const TArray& expected, OutputScope scope) {
+            if (strategy == nullptr) {
+                return;
+            }
+
+            auto outputInstance = joints.createInstance(&memRes);
+            auto outputBuffer = outputInstance->getOutputBuffer();
+
+            auto inputInstanceFactory =
+                ControlsFactory::getInstanceFactory(0, static_cast<std::uint16_t>(block8::input::values.size()), 0, 0);
+            auto inputInstance = inputInstanceFactory(&memRes);
+            auto inputBuffer = inputInstance->getInputBuffer();
+            std::copy(block8::input::values.begin(), block8::input::values.end(), inputBuffer.begin());
+
+            rl4::ConstArrayView<float> expectedView{expected[scope.lod].data() + scope.offset, scope.size};
+            rl4::ConstArrayView<float> outputView{outputBuffer.data() + scope.offset, scope.size};
+            joints.calculate(inputInstance.get(), outputInstance.get(), scope.lod);
+            ASSERT_EQ(outputView, expectedView);
+        }
+
+    protected:
+        pma::AlignedMemoryResource memRes;
+        block8::OptimizedStorage<StorageValueType>::StrategyPtr strategy;
+        StrategyTestParams params;
+
+};
+
+}  // namespace
+
+using Block8JointCalculationTypeList = ::testing::Types<
+#if defined(RL_BUILD_WITH_AVX)
+        std::tuple<StorageValueType, trimd::avx::F256, TStrategyTestParams<0> >,
+        std::tuple<StorageValueType, trimd::avx::F256, TStrategyTestParams<1> >,
+        std::tuple<StorageValueType, trimd::avx::F256, TStrategyTestParams<2> >,
+        std::tuple<StorageValueType, trimd::avx::F256, TStrategyTestParams<3> >,
+#endif  // RL_BUILD_WITH_AVX
+#if defined(RL_BUILD_WITH_AVX) || defined(RL_BUILD_WITH_SSE)
+        std::tuple<StorageValueType, trimd::sse::F256, TStrategyTestParams<0> >,
+        std::tuple<StorageValueType, trimd::sse::F256, TStrategyTestParams<1> >,
+        std::tuple<StorageValueType, trimd::sse::F256, TStrategyTestParams<2> >,
+        std::tuple<StorageValueType, trimd::sse::F256, TStrategyTestParams<3> >,
+#endif  // RL_BUILD_WITH_AVX || RL_BUILD_WITH_SSE
+#if defined(RL_BUILD_WITH_NEON)
+        std::tuple<StorageValueType, trimd::neon::F256, TStrategyTestParams<0> >,
+        std::tuple<StorageValueType, trimd::neon::F256, TStrategyTestParams<1> >,
+        std::tuple<StorageValueType, trimd::neon::F256, TStrategyTestParams<2> >,
+        std::tuple<StorageValueType, trimd::neon::F256, TStrategyTestParams<3> >,
+#endif  // RL_BUILD_WITH_NEON
+#if !defined(RL_USE_HALF_FLOATS)
+        std::tuple<StorageValueType, trimd::scalar::F256, TStrategyTestParams<0> >,
+        std::tuple<StorageValueType, trimd::scalar::F256, TStrategyTestParams<1> >,
+        std::tuple<StorageValueType, trimd::scalar::F256, TStrategyTestParams<2> >,
+        std::tuple<StorageValueType, trimd::scalar::F256, TStrategyTestParams<3> >,
+#endif  // RL_USE_HALF_FLOATS
+    std::tuple<>
+    >;
+
+TYPED_TEST_SUITE(Block8JointCalculationStrategyTest, Block8JointCalculationTypeList, );
+
+TYPED_TEST(Block8JointCalculationStrategyTest, Block8Padded) {
+    const OutputScope scope{this->params.lod, 0ul, 5ul};
+    auto joints = block8::OptimizedStorage<StorageValueType>::create(std::move(this->strategy), 0u, &this->memRes);
+    this->execute(joints, block8::output::valuesPerLOD, scope);
+}
+
+TYPED_TEST(Block8JointCalculationStrategyTest, Block8Exact) {
+    const OutputScope scope{this->params.lod, 5ul, 8ul};
+    auto joints = block8::OptimizedStorage<StorageValueType>::create(std::move(this->strategy), 1u, &this->memRes);
+    this->execute(joints, block8::output::valuesPerLOD, scope);
+}
+
+TYPED_TEST(Block8JointCalculationStrategyTest, Block16Padded) {
+    const OutputScope scope{this->params.lod, 13ul, 9ul};
+    auto joints = block8::OptimizedStorage<StorageValueType>::create(std::move(this->strategy), 2u, &this->memRes);
+    this->execute(joints, block8::output::valuesPerLOD, scope);
+}
+
+TYPED_TEST(Block8JointCalculationStrategyTest, Block16Exact) {
+    const OutputScope scope{this->params.lod, 22ul, 16ul};
+    auto joints = block8::OptimizedStorage<StorageValueType>::create(std::move(this->strategy), 3u, &this->memRes);
+    this->execute(joints, block8::output::valuesPerLOD, scope);
+}
+
+TYPED_TEST(Block8JointCalculationStrategyTest, Block24Padded) {
+    const OutputScope scope{this->params.lod, 38ul, 17ul};
+    auto joints = block8::OptimizedStorage<StorageValueType>::create(std::move(this->strategy), 4u, &this->memRes);
+    this->execute(joints, block8::output::valuesPerLOD, scope);
+}
+
+TYPED_TEST(Block8JointCalculationStrategyTest, Block24Exact) {
+    const OutputScope scope{this->params.lod, 55ul, 24ul};
+    auto joints = block8::OptimizedStorage<StorageValueType>::create(std::move(this->strategy), 5u, &this->memRes);
+    this->execute(joints, block8::output::valuesPerLOD, scope);
+}
+
+TYPED_TEST(Block8JointCalculationStrategyTest, Block32Padded) {
+    const OutputScope scope{this->params.lod, 79ul, 25ul};
+    auto joints = block8::OptimizedStorage<StorageValueType>::create(std::move(this->strategy), 6u, &this->memRes);
+    this->execute(joints, block8::output::valuesPerLOD, scope);
+}
+
+TYPED_TEST(Block8JointCalculationStrategyTest, Block32Exact) {
+    const OutputScope scope{this->params.lod, 104ul, 32ul};
+    auto joints = block8::OptimizedStorage<StorageValueType>::create(std::move(this->strategy), 7u, &this->memRes);
+    this->execute(joints, block8::output::valuesPerLOD, scope);
+}
+
+TYPED_TEST(Block8JointCalculationStrategyTest, MultipleBlocks) {
+    const OutputScope scope{this->params.lod, 0ul, block8::output::valuesPerLOD.front().size()};
+    auto joints = block8::OptimizedStorage<StorageValueType>::create(std::move(this->strategy), &this->memRes);
+    this->execute(joints, block8::output::valuesPerLOD, scope);
+}
diff --git a/Engine/Plugins/Animation/RigLogic/Source/RigLogicLibTest/Private/rltests/joints/bpcm/TestBlock8StorageBuilder.cpp b/Engine/Plugins/Animation/RigLogic/Source/RigLogicLibTest/Private/rltests/joints/bpcm/TestBlock8StorageBuilder.cpp
new file mode 100644
--- /dev/null
+++ b/Engine/Plugins/Animation/RigLogic/Source/RigLogicLibTest/Private/rltests/joints/bpcm/TestBlock8StorageBuilder.cpp
@@ -0,0 +1,70 @@
+// Copyright Epic Games, Inc. All Rights Reserved.
+
+#include "riglogic/system/simd/Detect.h"
+
+#include "rltests/Defs.h"
+#include "rltests/joints/bpcm/Assertions.h"
+#include "rltests/joints/bpcm/BPCMFixturesBlock8.h"
+#include "rltests/joints/bpcm/Helpers.h"
+
+#include "riglogic/TypeDefs.h"
+#include "riglogic/joints/JointsBuilder.h"
+#include "riglogic/joints/bpcm/Evaluator.h"
+#include "riglogic/joints/bpcm/strategies/Block8.h"
+#include "riglogic/riglogic/RigLogic.h"
+#include "riglogic/system/simd/SIMD.h"
+
+namespace {
+
+template<typename TTestTypes>
+class Block8JointStorageBuilderTest : public ::testing::Test {
+    protected:
+        template<typename TestTypes = TTestTypes>
+        typename std::enable_if<std::tuple_size<TestTypes>::value == 0ul, void>::type buildStorage() {
+        }
+
+        template<typename TestTypes = TTestTypes>
+        typename std::enable_if<std::tuple_size<TestTypes>::value != 0ul, void>::type buildStorage() {
+            using TValue = typename std::tuple_element<0, TestTypes>::type;
+            using TFVec = typename std::tuple_element<1, TestTypes>::type;
+            using TCalculationType = typename std::tuple_element<2, TestTypes>::type;
+
+            rl4::Configuration config{};
+            config.calculationType = TCalculationType::get();
+            auto builder = rl4::JointsBuilder::create(config, &memRes);
+            builder->computeStorageRequirements(&reader);
+            builder->allocateStorage(&reader);
+            builder->fillStorage(&reader);
+            auto joints = builder->build();
+            auto jointsImpl = static_cast<rl4::bpcm::Evaluator<TValue>*>(joints.get());
+
+            auto strategy = pma::UniqueInstance<rl4::bpcm::block8::Block8JointCalculationStrategy<TValue, TFVec>,
+                                                rl4::bpcm::JointCalculationStrategy<TValue> >::with(&memRes).create();
+            auto expected = block8::OptimizedStorage<TValue>::create(std::move(strategy), &memRes);
+
+            rl4::bpcm::Evaluator<TValue>::Accessor::assertRawDataEqual(*jointsImpl, expected);
+            rl4::bpcm::Evaluator<TValue>::Accessor::assertJointGroupsEqual(*jointsImpl, expected);
+            rl4::bpcm::Evaluator<TValue>::Accessor::assertLODsEqual(*jointsImpl, expected);
+        }
+
+    protected:
+        pma::AlignedMemoryResource memRes;
+        block8::CanonicalReader reader;
+
+};
+
+}  // namespace
+
+// Block-8 storage optimizer will execute only if RigLogic is built with AVX support, and AVX is chosen as calculation
+// type. In all other cases- Block-4 storage optimizer will run.
+using Block8StorageValueTypeList = ::testing::Types<
+#if defined(RL_BUILD_WITH_AVX)
+        std::tuple<StorageValueType, trimd::avx::F256, TCalculationType<rl4::CalculationType::AVX> >,
+#endif  // RL_BUILD_WITH_AVX
+    std::tuple<>
+    >;
+TYPED_TEST_SUITE(Block8JointStorageBuilderTest, Block8StorageValueTypeList, );
+
+TYPED_TEST(Block8JointStorageBuilderTest, LayoutOptimization) {
+    this->buildStorage();
+}
diff --git a/Engine/Plugins/Animation/RigLogic/Source/RigLogicLibTest/Private/rltests/joints/bpcm/TestCalculationStrategyAVX.cpp b/Engine/Plugins/Animation/RigLogic/Source/RigLogicLibTest/Private/rltests/joints/bpcm/TestCalculationStrategyAVX.cpp
deleted file mode 100644
diff --git a/Engine/Plugins/Animation/RigLogic/Source/RigLogicLibTest/Private/rltests/joints/bpcm/TestCalculationStrategySSE.cpp b/Engine/Plugins/Animation/RigLogic/Source/RigLogicLibTest/Private/rltests/joints/bpcm/TestCalculationStrategySSE.cpp
deleted file mode 100644
diff --git a/Engine/Plugins/Animation/RigLogic/Source/RigLogicLibTest/Private/rltests/joints/bpcm/TestCalculationStrategyScalar.cpp b/Engine/Plugins/Animation/RigLogic/Source/RigLogicLibTest/Private/rltests/joints/bpcm/TestCalculationStrategyScalar.cpp
deleted file mode 100644
diff --git a/Engine/Plugins/Animation/RigLogic/Source/RigLogicLibTest/Private/rltests/joints/bpcm/TestStorageBuilderAVX.cpp b/Engine/Plugins/Animation/RigLogic/Source/RigLogicLibTest/Private/rltests/joints/bpcm/TestStorageBuilderAVX.cpp
deleted file mode 100644
diff --git a/Engine/Plugins/Animation/RigLogic/Source/RigLogicLibTest/Private/rltests/joints/bpcm/TestStorageBuilderSSE.cpp b/Engine/Plugins/Animation/RigLogic/Source/RigLogicLibTest/Private/rltests/joints/bpcm/TestStorageBuilderSSE.cpp
deleted file mode 100644
diff --git a/Engine/Plugins/Animation/RigLogic/Source/RigLogicLibTest/Private/rltests/joints/bpcm/TestStorageBuilderScalar.cpp b/Engine/Plugins/Animation/RigLogic/Source/RigLogicLibTest/Private/rltests/joints/bpcm/TestStorageBuilderScalar.cpp
deleted file mode 100644
diff --git a/Engine/Plugins/Animation/RigLogic/Source/RigLogicLibTest/Private/rltests/ml/cpu/TestInference.cpp b/Engine/Plugins/Animation/RigLogic/Source/RigLogicLibTest/Private/rltests/ml/cpu/TestInference.cpp
--- a/Engine/Plugins/Animation/RigLogic/Source/RigLogicLibTest/Private/rltests/ml/cpu/TestInference.cpp
+++ b/Engine/Plugins/Animation/RigLogic/Source/RigLogicLibTest/Private/rltests/ml/cpu/TestInference.cpp
@@ -11,6 +11,7 @@
 #include "riglogic/TypeDefs.h"
 #include "riglogic/ml/cpu/Factory.h"
 #include "riglogic/system/simd/Detect.h"
+#include "riglogic/system/simd/SIMD.h"
 
 #include <tuple>
 
diff --git a/Engine/Plugins/Animation/RigLogic/Source/RigLogicLibTest/Private/trimdtests/TestT128.cpp b/Engine/Plugins/Animation/RigLogic/Source/RigLogicLibTest/Private/trimdtests/TestT128.cpp
--- a/Engine/Plugins/Animation/RigLogic/Source/RigLogicLibTest/Private/trimdtests/TestT128.cpp
+++ b/Engine/Plugins/Animation/RigLogic/Source/RigLogicLibTest/Private/trimdtests/TestT128.cpp
@@ -4,17 +4,30 @@
 
 #include "trimd/TRiMD.h"
 
+using T128Types = ::testing::Types<
+    trimd::scalar::F128
 #ifdef TRIMD_ENABLE_SSE
-    using T128Types = ::testing::Types<trimd::scalar::F128, trimd::sse::F128>;
+        , trimd::sse::F128
+#endif  // TRIMD_ENABLE_SSE
+#ifdef TRIMD_ENABLE_NEON
+        , trimd::neon::F128
+#endif  // TRIMD_ENABLE_NEON
+    >;
 
+#ifdef TRIMD_ENABLE_SSE
     bool equal(const trimd::sse::F128& lhs, const trimd::sse::F128& rhs) {
         return (std::memcmp(&lhs.data, &rhs.data, sizeof(lhs.data)) == 0);
     }
 
-#else
-    using T128Types = ::testing::Types<trimd::scalar::F128>;
 #endif  // TRIMD_ENABLE_SSE
 
+#ifdef TRIMD_ENABLE_NEON
+    bool equal(const trimd::neon::F128& lhs, const trimd::neon::F128& rhs) {
+        return (std::memcmp(&lhs.data, &rhs.data, sizeof(lhs.data)) == 0);
+    }
+
+#endif  // TRIMD_ENABLE_NEON
+
 template<typename TF128>
 static TF128 frombits(uint32_t bits0, uint32_t bits1, uint32_t bits2, uint32_t bits3) {
     return TF128{
@@ -415,7 +428,7 @@ TEST(T128Test, AndNotScalar) {
     }
 
     #ifdef TRIMD_ENABLE_F16C
-        TEST(T128Test, LoadAlignedHalfFloats) {
+        TEST(T128Test, LoadAlignedHalfFloatsSSE) {
             alignas(trimd::sse::F128::alignment()) const std::uint16_t halfFloats[] = {15360, 16384, 16896, 17408};
             trimd::sse::F128 expected{1.0f, 2.0f, 3.0f, 4.0f};
 
@@ -427,7 +440,7 @@ TEST(T128Test, AndNotScalar) {
             ASSERT_TRUE(equal(v2, expected));
         }
 
-        TEST(T128Test, LoadUnalignedHalfFloats) {
+        TEST(T128Test, LoadUnalignedHalfFloatsSSE) {
             const std::uint16_t halfFloats[] = {15360, 16384, 16896, 17408};
             trimd::sse::F128 expected{1.0f, 2.0f, 3.0f, 4.0f};
 
@@ -439,7 +452,7 @@ TEST(T128Test, AndNotScalar) {
             ASSERT_TRUE(equal(v2, expected));
         }
 
-        TEST(T128Test, StoreAlignedHalfFloats) {
+        TEST(T128Test, StoreAlignedHalfFloatsSSE) {
             const std::uint16_t expected[] = {15360, 16384, 16896, 17408};
 
             alignas(trimd::sse::F128::alignment()) std::uint16_t halfFloats[4ul] = {};
@@ -449,7 +462,7 @@ TEST(T128Test, AndNotScalar) {
             ASSERT_ELEMENTS_EQ(halfFloats, expected, 4ul);
         }
 
-        TEST(T128Test, StoreUnalignedHalfFloats) {
+        TEST(T128Test, StoreUnalignedHalfFloatsSSE) {
             const std::uint16_t expected[] = {15360, 16384, 16896, 17408};
 
             std::uint16_t halfFloats[4ul] = {};
@@ -460,3 +473,89 @@ TEST(T128Test, AndNotScalar) {
         }
     #endif  // TRIMD_ENABLE_F16C
 #endif  // TRIMD_ENABLE_SSE
+
+#ifdef TRIMD_ENABLE_NEON
+    TEST(T128Test, TransposeSquareNEON) {
+        trimd::neon::F128 v1{1.0f, 2.0f, 3.0f, 4.0f};
+        trimd::neon::F128 v2{1.0f, 2.0f, 3.0f, 4.0f};
+        trimd::neon::F128 v3{1.0f, 2.0f, 3.0f, 4.0f};
+        trimd::neon::F128 v4{1.0f, 2.0f, 3.0f, 4.0f};
+
+        trimd::neon::transpose(v1, v2, v3, v4);
+
+        trimd::neon::F128 e1{1.0f, 1.0f, 1.0f, 1.0f};
+        trimd::neon::F128 e2{2.0f, 2.0f, 2.0f, 2.0f};
+        trimd::neon::F128 e3{3.0f, 3.0f, 3.0f, 3.0f};
+        trimd::neon::F128 e4{4.0f, 4.0f, 4.0f, 4.0f};
+
+        ASSERT_TRUE(equal(v1, e1));
+        ASSERT_TRUE(equal(v2, e2));
+        ASSERT_TRUE(equal(v3, e3));
+        ASSERT_TRUE(equal(v4, e4));
+    }
+
+    TEST(T128Test, AbsNEON) {
+        trimd::neon::F128 v{-1.0f, 2.0f, -3.0f, 0.0f};
+        v = trimd::neon::abs(v);
+        trimd::neon::F128 e{1.0f, 2.0f, 3.0f, 0.0f};
+        ASSERT_TRUE(equal(v, e));
+    }
+
+    TEST(T128Test, AndNotNEON) {
+        trimd::neon::F128 v{1.0f, 2.0f, 3.0f, 4.0f};
+        trimd::neon::F128 mask1 = frombits<trimd::neon::F128>(0x00000000u, 0x00000000u, 0x00000000u, 0x00000000u);
+        trimd::neon::F128 mask2 = frombits<trimd::neon::F128>(0xFFFFFFFFu, 0xFFFFFFFFu, 0xFFFFFFFFu, 0xFFFFFFFFu);
+        trimd::neon::F128 result1 = trimd::neon::andnot(mask1, v);
+        trimd::neon::F128 result2 = trimd::neon::andnot(mask2, v);
+        trimd::neon::F128 e1{1.0f, 2.0f, 3.0f, 4.0f};
+        trimd::neon::F128 e2{0.0f, 0.0f, 0.0f, 0.0f};
+        ASSERT_TRUE(equal(result1, e1));
+        ASSERT_TRUE(equal(result2, e2));
+    }
+
+    #ifdef TRIMD_ENABLE_NEON_FP16
+        TEST(T128Test, LoadAlignedHalfFloatsNEON) {
+            alignas(trimd::neon::F128::alignment()) const std::uint16_t halfFloats[] = {15360, 16384, 16896, 17408};
+            trimd::neon::F128 expected{1.0f, 2.0f, 3.0f, 4.0f};
+
+            auto v = trimd::neon::F128::fromAlignedSource(halfFloats);
+            ASSERT_TRUE(equal(v, expected));
+
+            trimd::neon::F128 v2;
+            v2.alignedLoad(halfFloats);
+            ASSERT_TRUE(equal(v2, expected));
+        }
+
+        TEST(T128Test, LoadUnalignedHalfFloatsNEON) {
+            const std::uint16_t halfFloats[] = {15360, 16384, 16896, 17408};
+            trimd::neon::F128 expected{1.0f, 2.0f, 3.0f, 4.0f};
+
+            auto v = trimd::neon::F128::fromUnalignedSource(halfFloats);
+            ASSERT_TRUE(equal(v, expected));
+
+            trimd::neon::F128 v2;
+            v2.unalignedLoad(halfFloats);
+            ASSERT_TRUE(equal(v2, expected));
+        }
+
+        TEST(T128Test, StoreAlignedHalfFloatsNEON) {
+            const std::uint16_t expected[] = {15360, 16384, 16896, 17408};
+
+            alignas(trimd::neon::F128::alignment()) std::uint16_t halfFloats[4ul] = {};
+            const trimd::neon::F128 v{1.0f, 2.0f, 3.0f, 4.0f};
+            v.alignedStore(halfFloats);
+
+            ASSERT_ELEMENTS_EQ(halfFloats, expected, 4ul);
+        }
+
+        TEST(T128Test, StoreUnalignedHalfFloatsNEON) {
+            const std::uint16_t expected[] = {15360, 16384, 16896, 17408};
+
+            std::uint16_t halfFloats[4ul] = {};
+            const trimd::neon::F128 v{1.0f, 2.0f, 3.0f, 4.0f};
+            v.unalignedStore(halfFloats);
+
+            ASSERT_ELEMENTS_EQ(halfFloats, expected, 4ul);
+        }
+    #endif  // TRIMD_ENABLE_NEON_FP16
+#endif  // TRIMD_ENABLE_NEON
