Fix VS2022 17.8 preview compile warnings stdext::checked_array_iterator -> std::span

#rnx

diff --git a/Engine/Plugins/Animation/RigLogic/Source/RigLogicLib/Private/riglogic/riglogic/RigInstanceImpl.cpp b/Engine/Plugins/Animation/RigLogic/Source/RigLogicLib/Private/riglogic/riglogic/RigInstanceImpl.cpp
--- a/Engine/Plugins/Animation/RigLogic/Source/RigLogicLib/Private/riglogic/riglogic/RigInstanceImpl.cpp
+++ b/Engine/Plugins/Animation/RigLogic/Source/RigLogicLib/Private/riglogic/riglogic/RigInstanceImpl.cpp
@@ -20,6 +20,11 @@
 #include <cassert>
 #include <cstdint>
 #include <numeric>
+#ifdef _MSC_VER
+#if (_MSC_VER >= 1900)
+#include <span>
+#endif
+#endif
 #ifdef _MSC_VER
     #pragma warning(pop)
 #endif
@@ -91,7 +96,7 @@ void RigInstanceImpl::setGUIControlValues(const float* values) {
     #if (defined(_MSC_VER) && _MSC_VER >= 1900)
         std::copy(values,
                   values + guiControlCount,
-                  stdext::make_checked_array_iterator(guiControlBuffer.data(), guiControlBuffer.size()));
+                  std::span{guiControlBuffer.data(), guiControlBuffer.size()}.begin());
     #else
         std::copy(values, values + guiControlCount, guiControlBuffer.begin());
     #endif
@@ -123,7 +128,7 @@ void RigInstanceImpl::setRawControlValues(const float* values) {
     #if (defined(_MSC_VER) && _MSC_VER >= 1900)
         std::transform(values,
                        values + rawControlCount,
-                       stdext::make_checked_array_iterator(inputBuffer.data(), inputBuffer.size()),
+                       std::span{inputBuffer.data(), inputBuffer.size()}.begin(),
                        clamper);
     #else
         std::transform(values, values + rawControlCount, inputBuffer.begin(), clamper);
diff --git a/Engine/Plugins/Animation/RigLogic/Source/RigLogicLib/Public/terse/types/DynArray.h b/Engine/Plugins/Animation/RigLogic/Source/RigLogicLib/Public/terse/types/DynArray.h
--- a/Engine/Plugins/Animation/RigLogic/Source/RigLogicLib/Public/terse/types/DynArray.h
+++ b/Engine/Plugins/Animation/RigLogic/Source/RigLogicLib/Public/terse/types/DynArray.h
@@ -14,6 +14,9 @@
 #include <iterator>
 #include <memory>
 #include <type_traits>
+#if !defined(__clang__)
+#include <span>
+#endif
 #ifdef _MSC_VER
     #pragma warning(pop)
 #endif
@@ -70,7 +73,7 @@ class DynArray {
 
             #if defined(_MSC_VER) && !defined(__clang__)
                 if (size() != 0ul) {
-                    std::copy(start, end, stdext::checked_array_iterator<value_type*>(data(), size()));
+                    std::copy(start, end, std::span{data(), size()}.begin());
                 }
             #else
                 std::copy(start, end, data());
@@ -207,7 +210,7 @@ class DynArray {
             resize_uninitialized(static_cast<std::size_t>(std::distance(start, end)));
             #if defined(_MSC_VER) && !defined(__clang__)
                 if (size() != 0ul) {
-                    std::copy(start, end, stdext::checked_array_iterator<value_type*>(data(), size()));
+                    std::copy(start, end, std::span{data(), size()}.begin());
                 }
             #else
                 std::copy(start, end, data());
