Added a function to UDNAAsset that allows it to be copied at runtime.

This is needed for systems that dynamically generate meshes at runtime, such as Mutable.

diff --git a/Engine/Plugins/Animation/RigLogic/Source/RigLogicModule/Private/DNAAsset.cpp b/Engine/Plugins/Animation/RigLogic/Source/RigLogicModule/Private/DNAAsset.cpp
--- a/Engine/Plugins/Animation/RigLogic/Source/RigLogicModule/Private/DNAAsset.cpp
+++ b/Engine/Plugins/Animation/RigLogic/Source/RigLogicModule/Private/DNAAsset.cpp
@@ -125,6 +125,44 @@ void UDNAAsset::SetGeometryReader(TSharedPtr<IDNAReader> SourceDNAReader)
 #endif // #if WITH_EDITORONLY_DATA
 }
 
+void UDNAAsset::InitializeForRuntimeFrom(UDNAAsset* Other)
+{
+	FWriteScopeLock DNAScopeLock{DNAUpdateLock};
+	FWriteScopeLock ContextScopeLock{RigRuntimeContextUpdateLock};
+	FWriteScopeLock MappingScopeLock{DNAIndexMappingUpdateLock};
+
+	// Store a reference to the other asset's runtime context, using the accessor function to 
+	// ensure that the necessary lock is taken.
+	//
+	// The runtime context is immutable, i.e. not modified after initialization, so it's safe to
+	// share across UDNAAssets.
+	RigRuntimeContext = Other->GetRigRuntimeContext();
+
+	// BehaviorReader is used at runtime by GetDNAIndexMapping, so it needs to be populated here.
+	//
+	// The reference is taken from RigRuntimeContext to ensure it's consistent with the runtime
+	// context, as the UDNAAsset's BehaviorReader can be changed to point to a new one before the
+	// runtime context is updated.
+	//
+	// As with the runtime context itself, the BehaviorReader is immutable and safe to share.
+	BehaviorReader = RigRuntimeContext->BehaviorReader;
+
+	// Ensure bKeepDNAAfterInitialization reflects the current state of the BehaviorReader, i.e.
+	// if bKeepDNAAfterInitialization is true, BehaviorReader will contain DNA and vice versa.
+	bKeepDNAAfterInitialization = Other->bKeepDNAAfterInitialization;
+
+	// This map is populated on demand, so doesn't need to be copied.
+	DNAIndexMappingContainer.Empty(1);
+
+	// Clear any fields not needed at runtime to avoid any old data causing confusion
+#if WITH_EDITORONLY_DATA
+	AssetImportData = nullptr;
+#endif
+
+	DnaFileName.Empty();
+	GeometryReader = nullptr;
+}
+
 void UDNAAsset::InvalidateRigRuntimeContext()
 {
 	FWriteScopeLock ContextScopeLock{RigRuntimeContextUpdateLock};
@@ -150,9 +188,12 @@ void UDNAAsset::InitializeRigRuntimeContext()
 			RigRuntimeContext = NewContext;
 		}
 #if !WITH_EDITOR
-		BehaviorReader->Unload(EDNADataLayer::Behavior);
-		BehaviorReader->Unload(EDNADataLayer::Geometry);
-		BehaviorReader->Unload(EDNADataLayer::MachineLearnedBehavior);
+		if (!bKeepDNAAfterInitialization)
+		{
+			BehaviorReader->Unload(EDNADataLayer::Behavior);
+			BehaviorReader->Unload(EDNADataLayer::Geometry);
+			BehaviorReader->Unload(EDNADataLayer::MachineLearnedBehavior);
+		}
 #endif  // !WITH_EDITOR
 	}
 }
diff --git a/Engine/Plugins/Animation/RigLogic/Source/RigLogicModule/Public/DNAAsset.h b/Engine/Plugins/Animation/RigLogic/Source/RigLogicModule/Public/DNAAsset.h
--- a/Engine/Plugins/Animation/RigLogic/Source/RigLogicModule/Public/DNAAsset.h
+++ b/Engine/Plugins/Animation/RigLogic/Source/RigLogicModule/Public/DNAAsset.h
@@ -54,6 +54,15 @@ public:
 	UPROPERTY(VisibleAnywhere, AssetRegistrySearchable, Category = ImportSettings)
 	FString DnaFileName;
 
+	/** In non-editor builds, the DNA source data will be unloaded to save memory after the runtime
+	  * data has been initialized from it.
+	  * 
+	  * Set this property to true to keep the DNA in memory, e.g. if you need to modify it at
+	  * runtime. For most use cases, this shouldn't be needed.
+	 **/
+	UPROPERTY(EditAnywhere, Category = RuntimeSettings)
+	bool bKeepDNAAfterInitialization;
+
 	bool Init(const FString& Filename);
 	void Serialize(FArchive& Ar) override;
 
@@ -63,6 +72,20 @@ public:
 	void SetBehaviorReader(TSharedPtr<IDNAReader> SourceDNAReader);
 	void SetGeometryReader(TSharedPtr<IDNAReader> SourceDNAReader);
 
+	/** Initialize this object for use at runtime from another instance that has already been
+	  * initialized. 
+	  * 
+	  * Overwrites all member variables. Only data needed for runtime evaluation will be copied.
+	  * 
+	  * Performs a shallow copy, so the runtime data is shared between the two instances and the
+	  * memory cost of the copied UDNAAsset is very low.
+	  *
+	  * Note that the reference to the shared runtime data will be dropped if the source DNA is
+	  * modified, so the two instances are effectively independent and can safely be modified or
+	  * deleted without affecting the other.
+	 **/
+	void InitializeForRuntimeFrom(UDNAAsset* Other);
+
 private:
 	friend struct FAnimNode_RigLogic;
 	friend struct FRigUnit_RigLogic;
