Added MetaData map to DNA User Asset

diff --git a/Engine/Plugins/Animation/RigLogic/Source/RigLogicModule/Public/DNAAsset.h b/Engine/Plugins/Animation/RigLogic/Source/RigLogicModule/Public/DNAAsset.h
--- a/Engine/Plugins/Animation/RigLogic/Source/RigLogicModule/Public/DNAAsset.h
+++ b/Engine/Plugins/Animation/RigLogic/Source/RigLogicModule/Public/DNAAsset.h
@@ -32,7 +32,7 @@ enum class EDNADataLayer: uint8;
   * Currently, the design-time part still loads the geometry, as it is needed for the skeletal mesh update; once SkeletalMeshDNAReader is
   * fully implemented, it will be able to read the geometry directly from the SkeletalMesh and won't load it into this asset 
   **/
-UCLASS(NotBlueprintable, hidecategories = (Object))
+UCLASS(BlueprintType, meta = (DisplayName = "MetaHuman DNA data"))
 class RIGLOGICMODULE_API UDNAAsset : public UAssetUserData
 {
 	GENERATED_BODY()
@@ -63,6 +63,12 @@ public:
 	UPROPERTY(EditAnywhere, Category = RuntimeSettings)
 	bool bKeepDNAAfterInitialization;
 
+	/** Collection of runtime metadata related to a specific character.
+	  * The value field is a FString and requires casting for a derived types.
+	 **/
+	UPROPERTY(EditAnywhere, BlueprintReadWrite, Category = RuntimeSettings)
+	TMap<FString, FString> MetaData;
+
 	bool Init(const FString& Filename);
 	void Serialize(FArchive& Ar) override;
 
