Fix link error when merging modules

#rnx

diff --git a/Engine/Plugins/Animation/RigLogic/Source/RigLogicLibTest/RigLogicLibTest.Build.cs b/Engine/Plugins/Animation/RigLogic/Source/RigLogicLibTest/RigLogicLibTest.Build.cs
--- a/Engine/Plugins/Animation/RigLogic/Source/RigLogicLibTest/RigLogicLibTest.Build.cs
+++ b/Engine/Plugins/Animation/RigLogic/Source/RigLogicLibTest/RigLogicLibTest.Build.cs
@@ -20,7 +20,7 @@ namespace UnrealBuildTool.Rules
 
 			string RigLogicLibPath = Path.GetFullPath(Path.Combine(ModuleDirectory, "../RigLogicLib"));
 
-			if (Target.LinkType == TargetLinkType.Monolithic)
+			if (Target.LinkType == TargetLinkType.Monolithic || Target.bMergeModules)
 			{
 				PublicDependencyModuleNames.Add("RigLogicLib");
 				PrivateIncludePaths.Add(Path.Combine(RigLogicLibPath, "Private"));
