Refactored DNA import and fixed reimport for SkeletalMesh assets

diff --git a/Engine/Plugins/Animation/RigLogic/Source/RigLogicEditor/Private/DNAAssetImportFactory.cpp b/Engine/Plugins/Animation/RigLogic/Source/RigLogicEditor/Private/DNAAssetImportFactory.cpp
--- a/Engine/Plugins/Animation/RigLogic/Source/RigLogicEditor/Private/DNAAssetImportFactory.cpp
+++ b/Engine/Plugins/Animation/RigLogic/Source/RigLogicEditor/Private/DNAAssetImportFactory.cpp
@@ -68,6 +68,7 @@ EReimportResult::Type UDNAAssetImportFactory::Reimport(UObject* Obj, int32 Sourc
 {
 	// This is done here in oreder to enable importing of the dna file with the same name as 
 	// skeletal mesh. This is treated as an normal import but engine is recognizing it as an skeletal mesh reimport.
+	// Also this is called for regular Import since DNA import and reimport don't differ that much
 
 	UPackage* Pkg = Obj->GetPackage();
 	FString Name = FPaths::GetBaseFilename(PreferredReimportPath);
@@ -137,6 +138,18 @@ UObject* UDNAAssetImportFactory::FactoryCreateFile(UClass* InClass, UObject* InP
 	ImportUI->bIsReimport = false;
 	ImportUI->ReimportMesh = nullptr;
 
+	USkeletalMesh* SkeletalMeshForDNAImport = nullptr;
+	ForEachObjectWithOuter(InParent, [&SkeletalMeshForDNAImport](UObject* Object)
+	{
+		if (USkeletalMesh* SkeletalMesh = Cast<USkeletalMesh>(Object))
+		{
+			SkeletalMeshForDNAImport = SkeletalMesh;
+			return false;
+		}
+		return true;
+	});
+	ImportUI->SkeletalMesh = SkeletalMeshForDNAImport;
+
 	// Show the import dialog only when not in a "yes to all" state or when automating import
 	bool bIsAutomated = IsAutomatedImport();
 	bool bShowImportDialog = !bIsAutomated;
diff --git a/Engine/Plugins/Animation/RigLogic/Source/RigLogicEditor/Private/DNAImporter.cpp b/Engine/Plugins/Animation/RigLogic/Source/RigLogicEditor/Private/DNAImporter.cpp
--- a/Engine/Plugins/Animation/RigLogic/Source/RigLogicEditor/Private/DNAImporter.cpp
+++ b/Engine/Plugins/Animation/RigLogic/Source/RigLogicEditor/Private/DNAImporter.cpp
@@ -30,6 +30,7 @@ FDNAAssetImportOptions* GetImportOptions(FDNAImporter * DNAImporter, UDNAAssetIm
 	if (bShowOptionDialog)
 	{
 		FDNAAssetImportOptions* ImportOptions = DNAImporter->GetImportOptions();
+		ImportOptions->SkeletalMesh = ImportUI->SkeletalMesh;
 
 		// if SkeletalMesh was set by outside, please make sure copy back to UI
 		if (ImportOptions->SkeletalMesh)
@@ -56,7 +57,7 @@ FDNAAssetImportOptions* GetImportOptions(FDNAImporter * DNAImporter, UDNAAssetIm
 			}
 		}
 
-		DNAImporter->SetDNAFileName(*FPaths::GetCleanFilename(FullPath));
+		DNAImporter->SetDNAFileName(*FPaths::GetCleanFilename(InFilename));
 
 		//This option must always be the same value has the skeletalmesh one.
 		
diff --git a/Engine/Plugins/Animation/RigLogic/Source/RigLogicEditor/Private/RigLogicEditor.cpp b/Engine/Plugins/Animation/RigLogic/Source/RigLogicEditor/Private/RigLogicEditor.cpp
--- a/Engine/Plugins/Animation/RigLogic/Source/RigLogicEditor/Private/RigLogicEditor.cpp
+++ b/Engine/Plugins/Animation/RigLogic/Source/RigLogicEditor/Private/RigLogicEditor.cpp
@@ -5,6 +5,7 @@
 #include "Logging/LogMacros.h"
 #include "DNAImporter.h"
 #include "RigUnit_RigLogic.h"
+#include "EditorFramework/AssetImportData.h"
 #include "IAssetTypeActions.h"
 #include "AssetToolsModule.h"
 #include "Toolkits/AssetEditorToolkit.h"
@@ -96,16 +97,15 @@ void FRigLogicEditor::ExecuteDNAImport(UObject* Mesh)
 
 	FDNAImporter* DNAImporter = FDNAImporter::GetInstance();
 	const TArray<FString> Filenames = { DNAImporter->PromptForDNAImportFile() };
-	UFactory* Factory = Cast<UFactory>(UDNAAssetImportFactory::StaticClass()->GetDefaultObject());
-	FString Path = FPaths::GetPath(Mesh->GetPathName());
+	UDNAAssetImportFactory* Factory = Cast<UDNAAssetImportFactory>(UDNAAssetImportFactory::StaticClass()->GetDefaultObject());
 
-	if ((Mesh->GetFName()).ToString().Compare(Filenames[0]))
-	{
-		FReimportManager::Instance()->Reimport(Mesh, false, true, Filenames[0]);
-	}
-	else
+	//Reimport will do the same thing as import we just won't get problems when having the same DNA name as SkeletalMesh, new DNA gets initialized anyway
+	bool bSuccess = FReimportManager::Instance()->Reimport(Mesh, false, true, Filenames[0], Factory);
+
+	if (!bSuccess)
 	{
-		AssetToolsModule.Get().ImportAssets(Filenames, Path, Factory, true, nullptr, true);
+		const FText Message = LOCTEXT("DNA_ReimportFailedMessage", "Reimporting of DNA failed");
+		UE_LOG(LogRigLogicEditor, Error, TEXT("%s"), *Message.ToString());
 	}
 }
 
@@ -114,29 +114,33 @@ void FRigLogicEditor::ExecuteDNAReimport(class UObject* Mesh)
 	FAssetToolsModule& AssetToolsModule = FModuleManager::Get().LoadModuleChecked<FAssetToolsModule>("AssetTools");
 	USkeletalMesh* SkelMesh = Cast<USkeletalMesh>(Mesh);
 	const TArray<UAssetUserData*>* AssetData = SkelMesh->GetAssetUserDataArray();
+	UDNAAsset* DNAAsset = nullptr;
 
-	if (!AssetData->IsEmpty())
+	for (UAssetUserData* AssetDataElement : *AssetData)
 	{
-		const UDNAAsset* DNAAsset = Cast<UDNAAsset>((*AssetData)[0]);
-		if (DNAAsset)
+		if (Cast<UDNAAsset>(AssetDataElement))
 		{
-			const TArray<FString> Filenames = { DNAAsset->DnaFileName };
-			UFactory* Factory = Cast<UFactory>(UDNAAssetImportFactory::StaticClass()->GetDefaultObject());
-			FString Path = FPaths::GetPath(Mesh->GetPathName());
+			DNAAsset = Cast<UDNAAsset>(AssetDataElement);
+		}
+	}
 
-			if ((Mesh->GetFName()).ToString().Compare(Filenames[0]))
-			{
-				FReimportManager::Instance()->Reimport(Mesh, false, true, Filenames[0]);
-			}
-			else
-			{
-				AssetToolsModule.Get().ImportAssets(Filenames, Path, Factory, true, nullptr, true);
-			}
+	if (!AssetData->IsEmpty() && DNAAsset && DNAAsset->AssetImportData && !DNAAsset->AssetImportData->GetFirstFilename().IsEmpty())
+	{
+		const TArray<FString> Filenames = { DNAAsset->AssetImportData->GetFirstFilename() };
+		UFactory* Factory = Cast<UFactory>(UDNAAssetImportFactory::StaticClass()->GetDefaultObject());
+
+		bool bSuccess = FReimportManager::Instance()->Reimport(Mesh, false, true, Filenames[0]);
+
+		if(!bSuccess)
+		{
+			const FText Message = LOCTEXT("DNA_ReimportFailedMessage", "Reimporting of DNA failed");
+			UE_LOG(LogRigLogicEditor, Error, TEXT("%s"), *Message.ToString());
 		}
 	}
 	else
 	{
-		FNotificationInfo Info(LOCTEXT("DNA_ReimportErrorMessage", "There is no DNA file attached to do Reimport"));
+		const FText Message = LOCTEXT("DNA_ReimportErrorMessage", "There is no DNA file attached to do Reimport");
+		FNotificationInfo Info(Message);
 		Info.ExpireDuration = 3.0f;
 		Info.bUseLargeFont = false;
 		TSharedPtr<SNotificationItem> Notification = FSlateNotificationManager::Get().AddNotification(Info);
@@ -144,6 +148,8 @@ void FRigLogicEditor::ExecuteDNAReimport(class UObject* Mesh)
 		{
 			Notification->SetCompletionState(SNotificationItem::CS_Fail);
 		}
+
+		UE_LOG(LogRigLogicEditor, Error, TEXT("%s"), *Message.ToString());
 	}
 }
 
@@ -161,9 +167,9 @@ void FRigLogicEditor::GetAssetRegistryTagsForDNA(FAssetRegistryTagsContext Conte
 			if (UserData)
 			{
 				const UDNAAsset* DNAAsset = Cast<UDNAAsset>(UserData);
-				if (DNAAsset)
+				if (DNAAsset && DNAAsset->AssetImportData)
 				{
-					DNAname = DNAAsset->DnaFileName;
+					DNAname = DNAAsset->AssetImportData->GetFirstFilename();
 					FPaths::NormalizeFilename(DNAname);
 				}
 			}
diff --git a/Engine/Plugins/Animation/RigLogic/Source/RigLogicModule/Private/DNAAsset.cpp b/Engine/Plugins/Animation/RigLogic/Source/RigLogicModule/Private/DNAAsset.cpp
--- a/Engine/Plugins/Animation/RigLogic/Source/RigLogicModule/Private/DNAAsset.cpp
+++ b/Engine/Plugins/Animation/RigLogic/Source/RigLogicModule/Private/DNAAsset.cpp
@@ -11,9 +11,6 @@
 #include "RigLogicMemoryStream.h"
 #include "SharedRigRuntimeContext.h"
 
-#if WITH_EDITORONLY_DATA
-    #include "EditorFramework/AssetImportData.h"
-#endif
 #include "Engine/AssetUserData.h"
 #include "Serialization/BufferArchive.h"
 #include "Serialization/MemoryReader.h"
@@ -24,6 +21,7 @@
 #include "Animation/Skeleton.h"
 #include "Components/SkeletalMeshComponent.h"
 #include "Engine/SkeletalMesh.h"
+#include "EditorFramework/AssetImportData.h"
 
 #include "riglogic/RigLogic.h"
 
@@ -257,8 +255,13 @@ bool UDNAAsset::Init(const FString& DNAFilename)
 	{
 		UE_LOG(LogDNAAsset, Warning, TEXT("%s"), ANSI_TO_TCHAR(rl4::Status::get().message));
 	}
-
-	DnaFileName = DNAFilename; //memorize for re-import
+#if WITH_EDITORONLY_DATA
+	AssetImportData = NewObject<UAssetImportData>(this, TEXT("AssetImportData"));
+	TArray<FAssetImportInfo::FSourceFile> SourceFiles = { FAssetImportInfo::FSourceFile(DNAFilename) };
+	AssetImportData->SetSourceFiles(MoveTemp(SourceFiles));
+#endif
+	//This is done just for search through Asset Registry
+	DnaFileName = FPaths::GetCleanFilename(DNAFilename);
 	
 	if (!FPaths::FileExists(DNAFilename))
 	{
diff --git a/Engine/Plugins/Animation/RigLogic/Source/RigLogicModule/Public/DNAAsset.h b/Engine/Plugins/Animation/RigLogic/Source/RigLogicModule/Public/DNAAsset.h
--- a/Engine/Plugins/Animation/RigLogic/Source/RigLogicModule/Public/DNAAsset.h
+++ b/Engine/Plugins/Animation/RigLogic/Source/RigLogicModule/Public/DNAAsset.h
@@ -51,7 +51,7 @@ public:
 	TSharedPtr<IDNAReader> GetGeometryReader();
 #endif
 
-	UPROPERTY(VisibleAnywhere, AssetRegistrySearchable, Category = ImportSettings)
+	UPROPERTY(AssetRegistrySearchable)
 	FString DnaFileName;
 
 	/** In non-editor builds, the DNA source data will be unloaded to save memory after the runtime
