[SkeletalMesh/GroomBinding]
 - Make sure any async compilation dependent on the skeletal meshes are finished before applying changes to their render data directly
 - This is only temporary until these methods are replaced by something that preserves render data immutability

diff --git a/Engine/Plugins/Animation/RigLogic/Source/RigLogicModule/Private/SkelMeshDNAUtils.cpp b/Engine/Plugins/Animation/RigLogic/Source/RigLogicModule/Private/SkelMeshDNAUtils.cpp
--- a/Engine/Plugins/Animation/RigLogic/Source/RigLogicModule/Private/SkelMeshDNAUtils.cpp
+++ b/Engine/Plugins/Animation/RigLogic/Source/RigLogicModule/Private/SkelMeshDNAUtils.cpp
@@ -12,6 +12,7 @@
 #include "Misc/FileHelper.h"
 #include "UObject/ConstructorHelpers.h"
 #include "Engine/SkeletalMesh.h"
+#include "AssetCompilingManager.h"
 #include "ComponentReregisterContext.h"
 #include "Rendering/SkeletalMeshModel.h"
 #include "Rendering/SkeletalMeshLODModel.h"
@@ -360,6 +361,9 @@ void USkelMeshDNAUtils::RebuildRenderData(USkeletalMesh* InSkelMesh)
 {
 	FPlatformTime::InitTiming();
 
+	// Before touching the render data, just make sure any asset depending on it are finished building.
+	FAssetCompilingManager::Get().FinishCompilationForObjects({ InSkelMesh });
+
 	double StartTime = FPlatformTime::Seconds();
 	{
 		InSkelMesh->FlushRenderState();
@@ -416,6 +420,9 @@ void USkelMeshDNAUtils::RebuildRenderData_VertexPosition(USkeletalMesh* InSkelMe
 	}
 
 	{
+		// Before touching the render data, just make sure any asset depending on it are finished building.
+		FAssetCompilingManager::Get().FinishCompilationForObjects({ InSkelMesh });
+
 		FSkeletalMeshModel* MeshModel = InSkelMesh->GetImportedModel();
 		FSkeletalMeshRenderData* RenderData = InSkelMesh->GetResourceForRendering();
 
