Fix gathering of tests for RigLogic

diff --git a/Engine/Plugins/Animation/RigLogic/Source/RigLogicLibTest/Private/RigLogicLibTestSuite.cpp b/Engine/Plugins/Animation/RigLogic/Source/RigLogicLibTest/Private/RigLogicLibTestSuite.cpp
--- a/Engine/Plugins/Animation/RigLogic/Source/RigLogicLibTest/Private/RigLogicLibTestSuite.cpp
+++ b/Engine/Plugins/Animation/RigLogic/Source/RigLogicLibTest/Private/RigLogicLibTestSuite.cpp
@@ -1,55 +1,70 @@
 // Copyright Epic Games, Inc. All Rights Reserved.
 
 #if WITH_DEV_AUTOMATION_TESTS
-#pragma warning(disable:4121)
-#pragma warning(disable:4668)
-#pragma warning(disable:4996)
-#pragma warning(disable:6330)
-#include "gtest/gtest.h"
-#pragma warning(default:4121)
-#pragma warning(default:4668)
-#pragma warning(default:4996)
-#pragma warning(default:6330)
 
-#include "CoreMinimal.h"
+#include "RigLogicLibtest.h"
+#include "gtest/gtest.h"
 #include "Misc/AutomationTest.h"
 
-class FTestCollector : public ::testing::EmptyTestEventListener
+class FRigLogicLibTestPrinter
+: public ::testing::EmptyTestEventListener
 {
-public:
-    explicit FTestCollector(TArray<FString>* Names) : TestNames{Names}
-    {
-    }
-
-private:
-    void OnTestStart(const ::testing::TestInfo& TestInfo) override
-    {
-        FString TestCaseName = FString::Printf(TEXT("%s.%s"), *FString(TestInfo.test_case_name()), *FString(TestInfo.name()));
-        TestNames->Add(TestCaseName);
-    }
-
-private:
-    TArray<FString>* TestNames;
+	virtual void OnTestStart(const ::testing::TestInfo& InTestInfo) override
+	{
+		UE_LOG(LogRigLogicLibTest, Verbose, TEXT("Test %s.%s Starting"), *FString(InTestInfo.test_suite_name()), *FString(InTestInfo.name()));
+	}
+
+	virtual void OnTestPartResult(const ::testing::TestPartResult& InTestPartResult) override
+	{
+		if (InTestPartResult.failed())
+		{
+			UE_LOG(LogRigLogicLibTest, Error, TEXT("FAILED in %s:%d\n%s"), *FString(InTestPartResult.file_name()), InTestPartResult.line_number(), *FString(InTestPartResult.summary()));
+		}
+		else
+		{
+			UE_LOG(LogRigLogicLibTest, Verbose, TEXT("Succeeded in %s:%d\n%s"), *FString(InTestPartResult.file_name()), InTestPartResult.line_number(), *FString(InTestPartResult.summary()));
+		}
+	}
+
+	virtual void OnTestEnd(const ::testing::TestInfo& InTestInfo) override
+	{
+		UE_LOG(LogRigLogicLibTest, Verbose, TEXT("Test %s.%s Ending"), *FString(InTestInfo.test_suite_name()), *FString(InTestInfo.name()));
+	}
 };
 
 IMPLEMENT_COMPLEX_AUTOMATION_TEST(FRigLogicLibTestSuite, "RigLogicLib", EAutomationTestFlags::EditorContext | EAutomationTestFlags::EngineFilter)
 
 void FRigLogicLibTestSuite::GetTests(TArray<FString>& OutBeautifiedNames, TArray<FString>& OutTestCommands) const
 {
-    FTestCollector Collector{&OutBeautifiedNames};
-    ::testing::InitGoogleTest();
-    ::testing::TestEventListeners& Listeners = ::testing::UnitTest::GetInstance()->listeners();
-    Listeners.Append(&Collector);
-    auto _ = RUN_ALL_TESTS();
-    OutTestCommands = OutBeautifiedNames;
-    Listeners.Release(&Collector);
+	::testing::InitGoogleTest();
+
+	const ::testing::UnitTest* Instance = ::testing::UnitTest::GetInstance();
+	for (int32 TestSuiteIndex = 0; TestSuiteIndex < Instance->total_test_suite_count(); ++TestSuiteIndex)
+	{
+		const ::testing::TestSuite* TestSuite = Instance->GetTestSuite(TestSuiteIndex);
+		const FString TestCaseName = FString::Format(TEXT("{0}"), { TestSuite->name() });
+		OutBeautifiedNames.Add(TestCaseName);
+	}
+
+	OutTestCommands = OutBeautifiedNames;	
 }
 
 bool FRigLogicLibTestSuite::RunTest(const FString& Parameters)
 {
-    TArray<ANSICHAR> TestFilter{TCHAR_TO_ANSI(*Parameters), Parameters.Len() + 1};
-    ::testing::GTEST_FLAG(filter) = TestFilter.GetData();
-    return (RUN_ALL_TESTS() == 0);
+	::testing::TestEventListeners& Listeners = ::testing::UnitTest::GetInstance()->listeners();
+
+	FRigLogicLibTestPrinter TestPrinter;
+	Listeners.Append(&TestPrinter);
+
+	const FString Filter = FString::Format(TEXT("{0}*"), { Parameters });
+	const TArray<ANSICHAR> TestFilter{ TCHAR_TO_ANSI(*Filter), Filter.Len() + 1 };
+	::testing::GTEST_FLAG(filter) = TestFilter.GetData();
+
+	TestTrue("RigLogicLib Tests", RUN_ALL_TESTS() == 0);
+
+	Listeners.Release(&TestPrinter);
+
+	return true;
 }
 
-#endif  // WITH_DEV_AUTOMATION_TESTS
+#endif  // WITH_DEV_AUTOMATION_TESTS
\ No newline at end of file
