Prerequisites and small improvements in RigLogic for AnimNext

diff --git a/Engine/Plugins/Animation/RigLogic/Source/RigLogicModule/Private/AnimNode_RigLogic.cpp b/Engine/Plugins/Animation/RigLogic/Source/RigLogicModule/Private/AnimNode_RigLogic.cpp
--- a/Engine/Plugins/Animation/RigLogic/Source/RigLogicModule/Private/AnimNode_RigLogic.cpp
+++ b/Engine/Plugins/Animation/RigLogic/Source/RigLogicModule/Private/AnimNode_RigLogic.cpp
@@ -91,7 +91,7 @@ void FAnimNode_RigLogic::CacheBones_AnyThread(const FAnimationCacheBonesContext&
 		RigInstance = new FRigInstance(LocalRigRuntimeContext->RigLogic.Get());
 	}
 
-	LocalDNAIndexMapping = DNAAsset->GetDNAIndexMapping(Skeleton, SkeletalMesh, SkeletalMeshComponent);
+	LocalDNAIndexMapping = DNAAsset->GetDNAIndexMapping(Skeleton, SkeletalMesh);
 	// CacheBones is called on LOD switches as well, in which case compact pose bone indices must be remapped
 	const FBoneContainer& RequiredBones = Context.AnimInstanceProxy->GetRequiredBones();
 	if (RequiredBones.IsValid())
diff --git a/Engine/Plugins/Animation/RigLogic/Source/RigLogicModule/Private/DNAAsset.cpp b/Engine/Plugins/Animation/RigLogic/Source/RigLogicModule/Private/DNAAsset.cpp
--- a/Engine/Plugins/Animation/RigLogic/Source/RigLogicModule/Private/DNAAsset.cpp
+++ b/Engine/Plugins/Animation/RigLogic/Source/RigLogicModule/Private/DNAAsset.cpp
@@ -205,8 +205,7 @@ TSharedPtr<FSharedRigRuntimeContext> UDNAAsset::GetRigRuntimeContext()
 }
 
 TSharedPtr<FDNAIndexMapping> UDNAAsset::GetDNAIndexMapping(const USkeleton* Skeleton,
-														   const USkeletalMesh* SkeletalMesh,
-														   const USkeletalMeshComponent* SkeletalMeshComponent)
+														   const USkeletalMesh* SkeletalMesh)
 {
 	LLM_SCOPE_BYNAME(TEXT("Animation/RigLogic"));
 
@@ -238,7 +237,7 @@ TSharedPtr<FDNAIndexMapping> UDNAAsset::GetDNAIndexMapping(const USkeleton* Skel
 		DNAIndexMapping->SkeletonGuid = SkeletonGuid;
 		DNAIndexMapping->MapControlCurves(BehaviorReader.Get(), Skeleton);
 		DNAIndexMapping->MapNeuralNetworkMaskCurves(BehaviorReader.Get(), Skeleton);
-		DNAIndexMapping->MapJoints(BehaviorReader.Get(), SkeletalMeshComponent);
+		DNAIndexMapping->MapJoints(BehaviorReader.Get(), SkeletalMesh);
 		DNAIndexMapping->MapMorphTargets(BehaviorReader.Get(), Skeleton, SkeletalMesh);
 		DNAIndexMapping->MapMaskMultipliers(BehaviorReader.Get(), Skeleton);
 		DNAIndexMappingContainer.Add(SkeletalMesh, DNAIndexMapping);
diff --git a/Engine/Plugins/Animation/RigLogic/Source/RigLogicModule/Private/DNAIndexMapping.cpp b/Engine/Plugins/Animation/RigLogic/Source/RigLogicModule/Private/DNAIndexMapping.cpp
--- a/Engine/Plugins/Animation/RigLogic/Source/RigLogicModule/Private/DNAIndexMapping.cpp
+++ b/Engine/Plugins/Animation/RigLogic/Source/RigLogicModule/Private/DNAIndexMapping.cpp
@@ -6,7 +6,6 @@
 #include "DNAReader.h"
 #include "Animation/Skeleton.h"
 #include "Animation/SmartName.h"
-#include "Components/SkeletalMeshComponent.h"
 #include "Engine/SkeletalMesh.h"
 #include "Animation/AnimCurveTypes.h"
 
@@ -78,20 +77,22 @@ void FDNAIndexMapping::MapNeuralNetworkMaskCurves(const IDNAReader* DNAReader, c
 	}
 }
 
-
-void FDNAIndexMapping::MapJoints(const IDNAReader* DNAReader, const USkeletalMeshComponent* SkeletalMeshComponent)
+void FDNAIndexMapping::MapJoints(const IDNAReader* DNAReader, const USkeletalMesh* SkeletalMesh)
 {
 	LLM_SCOPE_BYNAME(TEXT("Animation/RigLogic"));
 
+	const FReferenceSkeleton& RefSkeleton = SkeletalMesh->GetRefSkeleton();
 	const uint16 JointCount = DNAReader->GetJointCount();
 	JointsMapDNAIndicesToMeshPoseBoneIndices.Reset(JointCount);
+
 	for (uint16 JointIndex = 0; JointIndex < JointCount; ++JointIndex)
 	{
 		const FString JointName = DNAReader->GetJointName(JointIndex);
 		const FName BoneName = FName(*JointName);
-		const int32 BoneIndex = SkeletalMeshComponent->GetBoneIndex(BoneName);
+		const int32 BoneIndex = RefSkeleton.FindBoneIndex(BoneName);
+
 		// BoneIndex may be INDEX_NONE, but it's handled properly by the Evaluate method
-		JointsMapDNAIndicesToMeshPoseBoneIndices.Add(FMeshPoseBoneIndex{BoneIndex});
+		JointsMapDNAIndicesToMeshPoseBoneIndices.Add(FMeshPoseBoneIndex{ BoneIndex });
 	}
 }
 
diff --git a/Engine/Plugins/Animation/RigLogic/Source/RigLogicModule/Private/DNAIndexMapping.h b/Engine/Plugins/Animation/RigLogic/Source/RigLogicModule/Private/DNAIndexMapping.h
--- a/Engine/Plugins/Animation/RigLogic/Source/RigLogicModule/Private/DNAIndexMapping.h
+++ b/Engine/Plugins/Animation/RigLogic/Source/RigLogicModule/Private/DNAIndexMapping.h
@@ -12,7 +12,6 @@
 class IDNAReader;
 class USkeleton;
 class USkeletalMesh;
-class USkeletalMeshComponent;
 
 struct FDNAIndexMapping
 {
@@ -35,7 +34,7 @@ struct FDNAIndexMapping
 
 	void MapControlCurves(const IDNAReader* DNAReader, const USkeleton* Skeleton);
 	void MapNeuralNetworkMaskCurves(const IDNAReader* DNAReader, const USkeleton* Skeleton);
-	void MapJoints(const IDNAReader* DNAReader, const USkeletalMeshComponent* SkeletalMeshComponent);
+	void MapJoints(const IDNAReader* DNAReader, const USkeletalMesh* SkeletalMesh);
 	void MapMorphTargets(const IDNAReader* DNAReader, const USkeleton* Skeleton, const USkeletalMesh* SkeletalMesh);
 	void MapMaskMultipliers(const IDNAReader* DNAReader, const USkeleton* Skeleton);
 
diff --git a/Engine/Plugins/Animation/RigLogic/Source/RigLogicModule/Public/DNAAsset.h b/Engine/Plugins/Animation/RigLogic/Source/RigLogicModule/Public/DNAAsset.h
--- a/Engine/Plugins/Animation/RigLogic/Source/RigLogicModule/Public/DNAAsset.h
+++ b/Engine/Plugins/Animation/RigLogic/Source/RigLogicModule/Public/DNAAsset.h
@@ -92,16 +92,16 @@ public:
 	 **/
 	void InitializeForRuntimeFrom(UDNAAsset* Other);
 
+	TSharedPtr<FSharedRigRuntimeContext> GetRigRuntimeContext();
+	TSharedPtr<FDNAIndexMapping> GetDNAIndexMapping(const USkeleton* Skeleton,
+		const USkeletalMesh* SkeletalMesh);
+
 private:
 	friend struct FAnimNode_RigLogic;
 	friend struct FRigUnit_RigLogic;
 
-	TSharedPtr<FSharedRigRuntimeContext> GetRigRuntimeContext();
 	void InvalidateRigRuntimeContext();
 	void InitializeRigRuntimeContext();
-	TSharedPtr<FDNAIndexMapping> GetDNAIndexMapping(const USkeleton* Skeleton,
-													const USkeletalMesh* SkeletalMesh,
-													const USkeletalMeshComponent* SkeletalMeshComponent);
 
 private:
 	// Synchronize DNA updates
