Moving DNAIndexMapping and SharedRigRuntimeContext to public folder

diff --git a/Engine/Plugins/Animation/RigLogic/Source/RigLogicModule/Private/DNAIndexMapping.h b/Engine/Plugins/Animation/RigLogic/Source/RigLogicModule/Private/DNAIndexMapping.h
deleted file mode 100644
diff --git a/Engine/Plugins/Animation/RigLogic/Source/RigLogicModule/Private/DNAIndexMappingDeprecated.h b/Engine/Plugins/Animation/RigLogic/Source/RigLogicModule/Private/DNAIndexMappingDeprecated.h
new file mode 100644
--- /dev/null
+++ b/Engine/Plugins/Animation/RigLogic/Source/RigLogicModule/Private/DNAIndexMappingDeprecated.h
@@ -0,0 +1,15 @@
+// Copyright Epic Games, Inc. All Rights Reserved.
+
+#pragma once
+
+#include "DNAIndexMapping.h"
+
+#include "DNAIndexMappingDeprecated.generated.h"
+
+// Used and needed for previously serialized skeletal mesh assets
+// that contain the DNA index mapping as part of their user data.
+UCLASS(NotBlueprintable, hidecategories = (Object), deprecated)
+class UDEPRECATED_DNAIndexMapping : public UAssetUserData
+{
+	GENERATED_BODY()
+};
\ No newline at end of file
diff --git a/Engine/Plugins/Animation/RigLogic/Source/RigLogicModule/Private/SharedRigRuntimeContext.h b/Engine/Plugins/Animation/RigLogic/Source/RigLogicModule/Private/SharedRigRuntimeContext.h
deleted file mode 100644
diff --git a/Engine/Plugins/Animation/RigLogic/Source/RigLogicModule/Public/DNAIndexMapping.h b/Engine/Plugins/Animation/RigLogic/Source/RigLogicModule/Public/DNAIndexMapping.h
new file mode 100644
--- /dev/null
+++ b/Engine/Plugins/Animation/RigLogic/Source/RigLogicModule/Public/DNAIndexMapping.h
@@ -0,0 +1,38 @@
+// Copyright Epic Games, Inc. All Rights Reserved.
+
+#pragma once
+
+#include "CoreMinimal.h"
+#include "BoneIndices.h"
+#include "Engine/AssetUserData.h"
+#include "Animation/AnimCurveTypes.h"
+
+class IDNAReader;
+class USkeleton;
+class USkeletalMesh;
+
+struct FDNAIndexMapping
+{
+	template <typename T>
+	struct TArrayWrapper
+	{
+		TArray<T> Values;
+	};
+	
+	using FCachedIndexedCurve = TBaseBlendedCurve<FDefaultAllocator, UE::Anim::FCurveElementIndexed>; 
+
+	FGuid SkeletonGuid;
+
+	// all the control attributes that we will need to extract, alongside their control index
+	FCachedIndexedCurve ControlAttributeCurves;
+	FCachedIndexedCurve NeuralNetworkMaskCurves;
+	TArray<FMeshPoseBoneIndex> JointsMapDNAIndicesToMeshPoseBoneIndices;
+	TArray<FCachedIndexedCurve> MorphTargetCurvesPerLOD;
+	TArray<FCachedIndexedCurve> MaskMultiplierCurvesPerLOD;
+
+	void MapControlCurves(const IDNAReader* DNAReader, const USkeleton* Skeleton);
+	void MapNeuralNetworkMaskCurves(const IDNAReader* DNAReader, const USkeleton* Skeleton);
+	void MapJoints(const IDNAReader* DNAReader, const USkeletalMesh* SkeletalMesh);
+	void MapMorphTargets(const IDNAReader* DNAReader, const USkeleton* Skeleton, const USkeletalMesh* SkeletalMesh);
+	void MapMaskMultipliers(const IDNAReader* DNAReader, const USkeleton* Skeleton);
+};
\ No newline at end of file
diff --git a/Engine/Plugins/Animation/RigLogic/Source/RigLogicModule/Public/SharedRigRuntimeContext.h b/Engine/Plugins/Animation/RigLogic/Source/RigLogicModule/Public/SharedRigRuntimeContext.h
new file mode 100644
--- /dev/null
+++ b/Engine/Plugins/Animation/RigLogic/Source/RigLogicModule/Public/SharedRigRuntimeContext.h
@@ -0,0 +1,29 @@
+// Copyright Epic Games, Inc. All Rights Reserved.
+
+#pragma once
+
+#include "CoreMinimal.h"
+#include "DNAReader.h"
+#include "RigLogic.h"
+
+struct FSharedRigRuntimeContext
+{
+	template <typename T>
+	struct TNestedArray
+	{
+		TArray<T> Values;
+	};
+
+	// Redundant but stored here as well so the runtime context can be updated / queried atomically
+	TSharedPtr<IDNAReader> BehaviorReader;
+
+	/** RigLogic itself is stateless, and is designed to be shared between
+	  * multiple rig instances based on the same DNA.
+	**/
+	TSharedPtr<FRigLogic> RigLogic;
+
+	/** Cached joint indices that need to be updated for each LOD **/
+	TArray<TNestedArray<uint16>> VariableJointIndicesPerLOD;
+
+	void CacheVariableJointIndices();
+};
