Update to new RigLogic implementing body rig support

diff --git a/Engine/Plugins/Animation/RigLogic/RigLogic.uplugin b/Engine/Plugins/Animation/RigLogic/RigLogic.uplugin
--- a/Engine/Plugins/Animation/RigLogic/RigLogic.uplugin
+++ b/Engine/Plugins/Animation/RigLogic/RigLogic.uplugin
@@ -1,9 +1,9 @@
 {
 	"FileVersion": 3,
-	"Version": 9.2,
-	"VersionName": "9.2.2",
-	"FriendlyName": "RigLogic Plugin v9.2.2",
-	"Description": "3Lateral RigLogic Plugin for Facial Animation v9.2.2",
+	"Version": 10.1,
+	"VersionName": "10.1.2",
+	"FriendlyName": "RigLogic Plugin v10.1.2",
+	"Description": "3Lateral RigLogic Plugin for Facial Animation v10.1.2",
 	"Category": "Animation",
 	"CreatedBy": "3Lateral",
 	"CreatedByURL": "https://www.3lateral.com/",
diff --git a/Engine/Plugins/Animation/RigLogic/Source/RigLogicLib/Private/dna/DNA.h b/Engine/Plugins/Animation/RigLogic/Source/RigLogicLib/Private/dna/DNA.h
--- a/Engine/Plugins/Animation/RigLogic/Source/RigLogicLib/Private/dna/DNA.h
+++ b/Engine/Plugins/Animation/RigLogic/Source/RigLogicLib/Private/dna/DNA.h
@@ -98,10 +98,11 @@ enum FileVersion : std::uint64_t {
     v21 = rev(2, 1),
     v22 = rev(2, 2),
     v23 = rev(2, 3),
-    latest = v23
+    v24 = rev(2, 4),
+    latest = v24
 };
 
-using Dispatch = Dispatcher<FileVersion, FileVersion::v21, FileVersion::v22, FileVersion::v23>;
+using Dispatch = Dispatcher<FileVersion, FileVersion::v21, FileVersion::v22, FileVersion::v23, FileVersion::v24>;
 
 template<typename A, FileVersion B>
 using IfHigher = typename std::enable_if<(A::value() > B), void>::type;
@@ -1297,6 +1298,222 @@ struct RawMachineLearnedBehavior {
 
 };
 
+struct RawRBFPose {
+    String<char> name;
+    float scale;
+
+    explicit RawRBFPose(MemoryResource* memRes) :
+        name{memRes},
+        scale{1.0f} {
+    }
+
+    template<class Archive>
+    void serialize(Archive& archive) {
+        archive.label("name");
+        archive(name);
+        archive.label("scale");
+        archive(scale);
+    }
+
+};
+
+struct RawRBFSolver {
+    terse::ArchiveSize<std::uint32_t, std::uint32_t> size;
+    terse::Anchor<std::uint32_t> baseMarker;
+    String<char> name;
+    DynArray<std::uint16_t> rawControlIndices;
+    DynArray<std::uint16_t> poseIndices;
+    DynArray<float> rawControlValues;
+    float radius;
+    float weightThreshold;
+    std::uint16_t solverType;
+    std::uint16_t automaticRadius;
+    std::uint16_t distanceMethod;
+    std::uint16_t normalizeMethod;
+    std::uint16_t functionType;
+    std::uint16_t twistAxis;
+    terse::ArchiveSize<std::uint32_t, std::uint32_t>::Proxy sizeMarker;
+
+    explicit RawRBFSolver(MemoryResource* memRes) :
+        size{},
+        baseMarker{},
+        name{memRes},
+        rawControlIndices{memRes},
+        poseIndices{memRes},
+        rawControlValues{memRes},
+        radius{},
+        weightThreshold{},
+        solverType{},
+        automaticRadius{},
+        distanceMethod{},
+        normalizeMethod{},
+        functionType{},
+        twistAxis{},
+        sizeMarker{size, baseMarker} {
+    }
+
+    template<class Archive>
+    void serialize(Archive& archive) {
+        archive.label("size");
+        archive(size);
+        archive(baseMarker);
+        archive.label("name");
+        archive(name);
+        archive.label("rawControlIndices");
+        archive(rawControlIndices);
+        archive.label("poseIndices");
+        archive(poseIndices);
+        archive.label("rawControlValues");
+        archive(rawControlValues);
+        archive.label("radius");
+        archive(radius);
+        archive.label("weightThreshold");
+        archive(weightThreshold);
+        archive.label("solverType");
+        archive(solverType);
+        archive.label("automaticRadius");
+        archive(automaticRadius);
+        archive.label("distanceMethod");
+        archive(distanceMethod);
+        archive.label("normalizeMethod");
+        archive(normalizeMethod);
+        archive.label("functionType");
+        archive(functionType);
+        archive.label("twistAxis");
+        archive(twistAxis);
+        archive(sizeMarker);
+    }
+
+};
+
+struct RawRBFBehavior {
+    RawLODMapping lodSolverMapping;
+    Vector<RawRBFSolver> solvers;
+    Vector<RawRBFPose> poses;
+
+    explicit RawRBFBehavior(MemoryResource* memRes) :
+        lodSolverMapping{memRes},
+        solvers{memRes},
+        poses{memRes} {
+    }
+
+    template<class Archive>
+    void serialize(Archive& archive) {
+        archive.label("lodSolverMapping");
+        archive(lodSolverMapping);
+        archive.label("solvers");
+        archive(solvers);
+        archive.label("poses");
+        archive(poses);
+    }
+
+};
+
+struct RawJointRepresentation {
+    std::uint16_t translation;
+    std::uint16_t rotation;
+    std::uint16_t scale;
+
+    template<class Archive>
+    void serialize(Archive& archive) {
+        archive.label("translation");
+        archive(translation);
+        archive.label("rotation");
+        archive(rotation);
+        archive.label("scale");
+        archive(scale);
+    }
+
+};
+
+struct RawJointBehaviorMetadata {
+    Vector<RawJointRepresentation> jointRepresentations;
+
+    explicit RawJointBehaviorMetadata(MemoryResource* memRes) :
+        jointRepresentations{memRes} {
+    }
+
+    template<class Archive>
+    void serialize(Archive& archive) {
+        archive.label("jointRepresentations");
+        archive(jointRepresentations);
+    }
+
+};
+
+struct RawTwist {
+    DynArray<float> twistBlendWeights;
+    DynArray<std::uint16_t> twistOutputJointIndices;
+    DynArray<std::uint16_t> twistInputControlIndices;
+    std::uint16_t twistAxis;
+
+    explicit RawTwist(MemoryResource* memRes) :
+        twistBlendWeights{memRes},
+        twistOutputJointIndices{memRes},
+        twistInputControlIndices{memRes},
+        twistAxis{} {
+    }
+
+    template<class Archive>
+    void serialize(Archive& archive) {
+        archive.label("twistBlendWeights");
+        archive(twistBlendWeights);
+        archive.label("twistOutputJointIndices");
+        archive(twistOutputJointIndices);
+        archive.label("twistInputControlIndices");
+        archive(twistInputControlIndices);
+        archive.label("twistAxis");
+        archive(twistAxis);
+    }
+
+};
+
+struct RawSwing {
+    DynArray<float> swingBlendWeights;
+    DynArray<std::uint16_t> swingOutputJointIndices;
+    DynArray<std::uint16_t> swingInputControlIndices;
+    std::uint16_t twistAxis;
+
+    explicit RawSwing(MemoryResource* memRes) :
+        swingBlendWeights{memRes},
+        swingOutputJointIndices{memRes},
+        swingInputControlIndices{memRes},
+        twistAxis{} {
+    }
+
+    template<class Archive>
+    void serialize(Archive& archive) {
+        archive.label("swingBlendWeights");
+        archive(swingBlendWeights);
+        archive.label("swingOutputJointIndices");
+        archive(swingOutputJointIndices);
+        archive.label("swingInputControlIndices");
+        archive(swingInputControlIndices);
+        archive.label("twistAxis");
+        archive(twistAxis);
+    }
+
+};
+
+struct RawTwistSwingBehavior {
+    Vector<RawTwist> twists;
+    Vector<RawSwing> swings;
+
+    explicit RawTwistSwingBehavior(MemoryResource* memRes) :
+        twists{memRes},
+        swings{memRes} {
+    }
+
+    template<class Archive>
+    void serialize(Archive& archive) {
+        archive.label("twists");
+        archive(twists);
+        archive.label("swings");
+        archive(swings);
+    }
+
+};
+
 struct DNA {
     MemoryResource* memRes;
     Signature<3> signature{{'D', 'N', 'A'}};
@@ -1307,13 +1524,20 @@ struct DNA {
     Layer<sid4("bhvr"), rev(1, 1), RawBehavior, FileVersion::v22> behavior;
     Layer<sid4("geom"), rev(1, 1), RawGeometry, FileVersion::v22> geometry;
     Layer<sid4("mlbh"), rev(1, 0), RawMachineLearnedBehavior, FileVersion::v23> machineLearnedBehavior;
+    Layer<sid4("rbfb"), rev(1, 0), RawRBFBehavior, FileVersion::v24> rbfBehavior;
+    Layer<sid4("jbmd"), rev(1, 0), RawJointBehaviorMetadata, FileVersion::v24> jointBehaviorMetadata;
+    Layer<sid4("twsw"), rev(1, 0), RawTwistSwingBehavior, FileVersion::v24> twistSwingBehavior;
+
 
     using Layers =
         LayerContainer<decltype(descriptor),
                        decltype(definition),
                        decltype(behavior),
                        decltype(geometry),
-                       decltype(machineLearnedBehavior)>;
+                       decltype(machineLearnedBehavior),
+                       decltype(rbfBehavior),
+                       decltype(jointBehaviorMetadata),
+                       decltype(twistSwingBehavior)>;
     Layers layers;
 
     DNA(UnknownLayerPolicy unknownPolicy, UpgradeFormatPolicy upgradePolicy, MemoryResource* memRes_) :
@@ -1324,10 +1548,15 @@ struct DNA {
         behavior{memRes},
         geometry{memRes},
         machineLearnedBehavior{memRes},
-        layers{unknownPolicy, upgradePolicy, memRes, &descriptor, &definition, &behavior, &geometry, &machineLearnedBehavior} {
+        rbfBehavior{memRes},
+        jointBehaviorMetadata{memRes},
+        twistSwingBehavior{memRes},
+        layers{unknownPolicy, upgradePolicy, memRes, &descriptor, &definition, &behavior, &geometry,
+               &machineLearnedBehavior, &rbfBehavior, &jointBehaviorMetadata, &twistSwingBehavior} {
     }
 
-    explicit DNA(MemoryResource* memRes_) : DNA{UnknownLayerPolicy::Preserve, UpgradeFormatPolicy::Allowed, memRes_} {
+    explicit DNA(MemoryResource* memRes_) : DNA{UnknownLayerPolicy::Preserve, UpgradeFormatPolicy::Allowed,
+                                                memRes_} {
     }
 
     template<class Archive>
@@ -1434,6 +1663,18 @@ struct DNA {
         static_cast<RawMachineLearnedBehavior&>(machineLearnedBehavior) = RawMachineLearnedBehavior{memRes};
     }
 
+    void unloadRBFBehavior() {
+        static_cast<RawRBFBehavior&>(rbfBehavior) = RawRBFBehavior{memRes};
+    }
+
+    void unloadJointBehaviorMetadata() {
+        static_cast<RawJointBehaviorMetadata&>(jointBehaviorMetadata) = RawJointBehaviorMetadata{memRes};
+    }
+
+    void unloadTwistSwingBehavior() {
+        static_cast<RawTwistSwingBehavior&>(twistSwingBehavior) = RawTwistSwingBehavior{memRes};
+    }
+
 };
 
 }  // namespace dna
diff --git a/Engine/Plugins/Animation/RigLogic/Source/RigLogicLib/Private/dna/DataLayerBitmask.h b/Engine/Plugins/Animation/RigLogic/Source/RigLogicLib/Private/dna/DataLayerBitmask.h
--- a/Engine/Plugins/Animation/RigLogic/Source/RigLogicLib/Private/dna/DataLayerBitmask.h
+++ b/Engine/Plugins/Animation/RigLogic/Source/RigLogicLib/Private/dna/DataLayerBitmask.h
@@ -13,7 +13,10 @@ enum class DataLayerBitmask {
     Behavior = 4,
     GeometryBlendShapesOnly = 8,
     GeometryRest = 16,
-    MachineLearnedBehavior = 32
+    MachineLearnedBehavior = 32,
+    RBFBehavior = 64,
+    JointBehaviorMetadata = 128,
+    TwistSwingBehavior = 256
 };
 
 inline DataLayerBitmask computeDataLayerBitmask(DataLayer layer) {
@@ -24,6 +27,9 @@ inline DataLayerBitmask computeDataLayerBitmask(DataLayer layer) {
         result |= DataLayerBitmask::GeometryBlendShapesOnly;
         result |= DataLayerBitmask::GeometryRest;
         result |= DataLayerBitmask::MachineLearnedBehavior;
+        result |= DataLayerBitmask::RBFBehavior;
+        result |= DataLayerBitmask::JointBehaviorMetadata;
+        result |= DataLayerBitmask::TwistSwingBehavior;
         return result;
     }
 
@@ -48,6 +54,18 @@ inline DataLayerBitmask computeDataLayerBitmask(DataLayer layer) {
         result |= DataLayerBitmask::MachineLearnedBehavior;
     }
 
+    if (contains(layer, DataLayer::RBFBehavior)) {
+        result |= DataLayerBitmask::RBFBehavior;
+    }
+
+    if (contains(layer, DataLayer::JointBehaviorMetadata)) {
+        result |= DataLayerBitmask::JointBehaviorMetadata;
+    }
+
+    if (contains(layer, DataLayer::TwistSwingBehavior)) {
+        result |= DataLayerBitmask::TwistSwingBehavior;
+    }
+
     return result;
 }
 
diff --git a/Engine/Plugins/Animation/RigLogic/Source/RigLogicLib/Private/dna/DenormalizedData.h b/Engine/Plugins/Animation/RigLogic/Source/RigLogicLib/Private/dna/DenormalizedData.h
--- a/Engine/Plugins/Animation/RigLogic/Source/RigLogicLib/Private/dna/DenormalizedData.h
+++ b/Engine/Plugins/Animation/RigLogic/Source/RigLogicLib/Private/dna/DenormalizedData.h
@@ -4,11 +4,20 @@
 
 #include "dna/LODMapping.h"
 #include "dna/layers/DefinitionReader.h"
-#include "dna/types/Aliases.h"
 #include "dna/utils/Extd.h"
 
+#ifdef _MSC_VER
+    #pragma warning(push)
+    #pragma warning(disable : 4365 4987)
+#endif
 #include <cassert>
 #include <cstdint>
+#include <cmath>
+#include <limits>
+#ifdef _MSC_VER
+    #pragma warning(pop)
+#endif
+
 
 namespace dna {
 
@@ -16,43 +25,75 @@ template<class Reader>
 struct DenormalizedData {
     LODMapping jointVariableAttributeIndices;
     LODMapping meshBlendShapeMappingIndices;
+    Matrix<std::uint16_t> rbfPoseJointOutputIndices;
+    Matrix<float> rbfPoseJointOutputValues;
+    Matrix<std::uint16_t> rbfBlendShapeChannelOutputIndices;
+    Matrix<std::uint16_t> rbfAnimatedMapOutputIndices;
 
     explicit DenormalizedData(MemoryResource* memRes) :
         jointVariableAttributeIndices{memRes},
-        meshBlendShapeMappingIndices{memRes} {
-    }
-
-    void populate(const Reader* source) {
-        populateJointVariableAttributeIndices(source, jointVariableAttributeIndices);
-        populateMeshBlendShapeMappingIndices(source, meshBlendShapeMappingIndices);
+        meshBlendShapeMappingIndices{memRes},
+        rbfPoseJointOutputIndices{memRes},
+        rbfPoseJointOutputValues{memRes},
+        rbfBlendShapeChannelOutputIndices{memRes},
+        rbfAnimatedMapOutputIndices{memRes} {
     }
 
-    private:
-        void populateJointVariableAttributeIndices(const Reader* source, LODMapping& destination) {
+    public:
+        void populateJointVariableAttributeIndices(const Reader* source) {
             // Prepare storage for all available LODs
             const auto lodCount = source->getLODCount();
-            destination.setLODCount(lodCount);
+            jointVariableAttributeIndices.setLODCount(lodCount);
+            auto memRes = rbfPoseJointOutputIndices.get_allocator().getMemoryResource();
+            UnorderedSet<std::uint16_t> twswAttributeIndicesSet{memRes};
+
+            for (std::uint16_t i = {}; i < source->getTwistCount(); ++i) {
+                auto jointIndices = source->getTwistOutputJointIndices(i);
+                for (const auto jointIndex : jointIndices) {
+                    twswAttributeIndicesSet.emplace(static_cast<std::uint16_t>(jointIndex * 9u + 3u));
+                    twswAttributeIndicesSet.emplace(static_cast<std::uint16_t>(jointIndex * 9u + 4u));
+                    twswAttributeIndicesSet.emplace(static_cast<std::uint16_t>(jointIndex * 9u + 5u));
+                }
+            }
+            for (std::uint16_t i = {}; i < source->getSwingCount(); ++i) {
+                auto jointIndices = source->getSwingOutputJointIndices(i);
+                for (const auto jointIndex : jointIndices) {
+                    twswAttributeIndicesSet.emplace(static_cast<std::uint16_t>(jointIndex * 9u + 3u));
+                    twswAttributeIndicesSet.emplace(static_cast<std::uint16_t>(jointIndex * 9u + 4u));
+                    twswAttributeIndicesSet.emplace(static_cast<std::uint16_t>(jointIndex * 9u + 5u));
+                }
+            }
+
+            Vector<UnorderedSet<std::uint16_t> > distinctIndicesPerLOD{lodCount, twswAttributeIndicesSet, memRes};
             // Concatenate all output indices for each LOD
-            for (std::uint16_t i = 0u; i < source->getJointGroupCount(); ++i) {
+            for (std::uint16_t i = {}; i < source->getJointGroupCount(); ++i) {
                 const auto outputIndices = source->getJointGroupOutputIndices(i);
                 const auto lodSizes = source->getJointGroupLODs(i);
                 assert(lodSizes.size() == lodCount);
-                for (std::uint16_t lod = 0u; lod < lodCount; ++lod) {
+                for (std::uint16_t lod = {}; lod < lodCount; ++lod) {
                     // In this case, each LOD has a distinct set of indices, so the LOD and Index parameters
                     // are the same for all LODs
-                    destination.addIndices(lod, outputIndices.data(), lodSizes[lod]);
-                    destination.associateLODWithIndices(lod, lod);
+                    for (std::uint16_t oi = {}; oi < lodSizes[lod]; ++oi) {
+                        distinctIndicesPerLOD[lod].emplace(outputIndices[oi]);
+                    }
                 }
             }
+
+            for (std::uint16_t lod = {}; lod < lodCount; ++lod) {
+                jointVariableAttributeIndices.addIndices(lod, distinctIndicesPerLOD[lod].begin(),
+                                                         distinctIndicesPerLOD[lod].end());
+                jointVariableAttributeIndices.associateLODWithIndices(lod, lod);
+            }
+            jointVariableAttributeIndices.sortIndices();
         }
 
-        void populateMeshBlendShapeMappingIndices(const Reader* source, LODMapping& destination) {
+        void populateMeshBlendShapeMappingIndices(const Reader* source) {
             // Prepare storage for all available LODs
             const auto lodCount = source->getLODCount();
-            destination.setLODCount(lodCount);
+            meshBlendShapeMappingIndices.setLODCount(lodCount);
             // Include only those mapping indices which are present in the already filtered
             // mesh and blendshape LOD mapping
-            for (std::uint16_t lod = 0u; lod < lodCount; ++lod) {
+            for (std::uint16_t lod = {}; lod < lodCount; ++lod) {
                 const auto meshIndices = source->getMeshIndicesForLOD(lod);
                 const auto blendShapeIndices = source->getBlendShapeChannelIndicesForLOD(lod);
 
@@ -62,19 +103,149 @@ struct DenormalizedData {
                         return (meshNeeded && blendShapeNeeded);
                     };
 
-                for (std::uint16_t i = 0u; i < source->getMeshBlendShapeChannelMappingCount(); ++i) {
+                for (std::uint16_t i = {}; i < source->getMeshBlendShapeChannelMappingCount(); ++i) {
                     const auto mapping = source->getMeshBlendShapeChannelMapping(i);
                     if (isMappingNeeded(mapping)) {
                         // In this case, each LOD has a distinct set of indices, so the LOD and Index parameters
                         // are the same for all LODs
-                        destination.addIndices(lod, &i, static_cast<std::uint16_t>(1));
+                        meshBlendShapeMappingIndices.addIndices(lod, &i, static_cast<std::uint16_t>(1));
                     }
                 }
 
-                destination.associateLODWithIndices(lod, lod);
+                meshBlendShapeMappingIndices.associateLODWithIndices(lod, lod);
             }
         }
 
+        void populateRBFPoseJointOutputIndices(const Reader* source) {
+            const std::uint16_t poseCount = source->getRBFPoseCount();
+            rbfPoseJointOutputIndices.resize(poseCount);
+
+            const auto indexOf = [](ConstArrayView<std::uint16_t> container, std::uint16_t val) {
+                    const auto valIt = std::find(container.begin(), container.end(), val);
+                    return static_cast<size_t>(std::distance(container.begin(), valIt));
+                };
+            const std::uint16_t rawControlCount = source->getRawControlCount();
+            const std::uint16_t psdControlCount = source->getPSDCount();
+            const std::uint16_t mlControlCount = source->getMLControlCount();
+            auto rbfControlOffset = rawControlCount + psdControlCount + mlControlCount;
+
+            for (std::uint16_t gi = {}; gi < source->getJointGroupCount(); ++gi) {
+                const auto inputIndices = source->getJointGroupInputIndices(gi);
+                const auto columnCount = inputIndices.size();
+                for (const auto& inputIndex : inputIndices) {
+                    if (inputIndex >= rbfControlOffset) {
+                        const auto poseIndex = static_cast<std::uint16_t>(inputIndex - rbfControlOffset);
+                        const auto outputIndices = source->getJointGroupOutputIndices(gi);
+                        rbfPoseJointOutputIndices[poseIndex].reserve(outputIndices.size());
+                        const auto columnIndex = indexOf(inputIndices, inputIndex);
+                        const auto values = source->getJointGroupValues(gi);
+                        for (std::uint16_t oi = 0; oi < outputIndices.size(); ++oi) {
+                            if (std::abs(values[oi * columnCount + columnIndex]) > 0.0f) {
+                                rbfPoseJointOutputIndices[poseIndex].push_back(outputIndices[oi]);
+                            }
+                        }
+                    }
+                }
+            }
+
+        }
+
+        void populateRBFPoseJointOutputValues(const Reader* source) {
+            const std::uint16_t poseCount = source->getRBFPoseCount();
+            rbfPoseJointOutputValues.resize(poseCount);
+            const auto findJointGroupIndex = [source](std::uint16_t outputIndex) {
+                    const auto groupCount = source->getJointGroupCount();
+                    const auto jointIndex = static_cast<std::uint16_t>(outputIndex / 9);
+                    for (std::uint16_t gi = {}; gi < groupCount; ++gi) {
+                        const auto jointIndices = source->getJointGroupJointIndices(gi);
+                        if (extd::contains(jointIndices, jointIndex)) {
+                            return gi;
+                        }
+                    }
+                    return std::numeric_limits<std::uint16_t>::max();
+                };
+
+            const auto indexOf = [](ConstArrayView<std::uint16_t> container, std::uint16_t val) {
+                    const auto valIt = std::find(container.begin(), container.end(), val);
+                    return static_cast<size_t>(std::distance(container.begin(), valIt));
+                };
+            const std::uint16_t rawControlCount = source->getRawControlCount();
+            const std::uint16_t psdControlCount = source->getPSDCount();
+            const std::uint16_t mlControlCount = source->getMLControlCount();
+            auto rbfControlOffset = rawControlCount + psdControlCount + mlControlCount;
+            for (std::uint16_t pi = {}; pi < poseCount; ++pi) {
+                const auto inputIndex = static_cast<std::uint16_t>(rbfControlOffset + pi);
+                const auto jointOutputIndices = source->getRBFPoseJointOutputIndices(pi);
+                const auto jointOutputIndicesCount = jointOutputIndices.size();
+                rbfPoseJointOutputValues[pi].resize(jointOutputIndices.size(), 0.0f);
+                for (std::uint16_t oi = {}; oi < jointOutputIndicesCount; ++oi) {
+                    const auto outputIndex = jointOutputIndices[oi];
+                    const auto jointGroupIndex = findJointGroupIndex(outputIndex);
+
+                    const auto jointGroupOutputIndices = source->getJointGroupOutputIndices(jointGroupIndex);
+                    const auto rowCount = jointGroupOutputIndices.size();
+
+                    const auto jointGroupInputIndices = source->getJointGroupInputIndices(jointGroupIndex);
+                    const auto columnCount = jointGroupInputIndices.size();
+
+                    const auto columnIndex = indexOf(jointGroupInputIndices, inputIndex);
+                    const auto rowIndex = indexOf(jointGroupOutputIndices, outputIndex);
+
+                    if ((columnIndex < columnCount) && (rowIndex < rowCount)) {
+                        const auto values = source->getJointGroupValues(jointGroupIndex);
+                        rbfPoseJointOutputValues[pi][oi] = values[rowIndex * columnCount + columnIndex];
+                    }
+                }
+            }
+        }
+
+        void populateRBFBlendShapeOutputIndices(const Reader* source) {
+            const std::uint16_t solverCount = source->getRBFSolverCount();
+            rbfBlendShapeChannelOutputIndices.resize(solverCount);
+            const auto bscInputIndices = source->getBlendShapeChannelInputIndices();
+            const auto bscOutputIndices = source->getBlendShapeChannelOutputIndices();
+            const auto bscCount = bscInputIndices.size();
+            const std::uint16_t rawControlCount = source->getRawControlCount();
+            const std::uint16_t psdControlCount = source->getPSDCount();
+            const std::uint16_t mlControlCount = source->getMLControlCount();
+            auto rbfControlOffset = rawControlCount + psdControlCount + mlControlCount;
+
+            const auto poseCount = source->getRBFPoseCount();
+            rbfBlendShapeChannelOutputIndices.resize(poseCount);
+            for (std::uint16_t pi = {}; pi < poseCount; ++pi) {
+                const auto inputIndex = rbfControlOffset + pi;
+                for (std::uint16_t bi = {}; bi < bscCount; ++bi) {
+                    if (bscInputIndices[bi] == inputIndex) {
+                        rbfBlendShapeChannelOutputIndices[pi].push_back(bscOutputIndices[bi]);
+                    }
+                }
+            }
+
+        }
+
+        void populateRBFAnimatedMapOutputIndices(const Reader* source) {
+            const std::uint16_t solverCount = source->getRBFSolverCount();
+            rbfAnimatedMapOutputIndices.resize(solverCount);
+            const auto amInputIndices = source->getAnimatedMapInputIndices();
+            const auto amOutputIndices = source->getAnimatedMapOutputIndices();
+            const auto amCount = source->getAnimatedMapCount();
+            const std::uint16_t rawControlCount = source->getRawControlCount();
+            const std::uint16_t psdControlCount = source->getPSDCount();
+            const std::uint16_t mlControlCount = source->getMLControlCount();
+            auto rbfControlOffset = rawControlCount + psdControlCount + mlControlCount;
+
+            const auto poseCount = source->getRBFPoseCount();
+            for (std::uint16_t pi = {}; pi < poseCount; ++pi) {
+                const auto inputIndex = rbfControlOffset + pi;
+                for (std::uint16_t ai = {}; ai < amCount; ++ai) {
+                    if (amInputIndices[ai] == inputIndex) {
+                        rbfAnimatedMapOutputIndices[pi].push_back(amOutputIndices[ai]);
+                    }
+                }
+            }
+
+        }
+
 };
 
 }  // namespace dna
diff --git a/Engine/Plugins/Animation/RigLogic/Source/RigLogicLib/Private/dna/LODMapping.cpp b/Engine/Plugins/Animation/RigLogic/Source/RigLogicLib/Private/dna/LODMapping.cpp
--- a/Engine/Plugins/Animation/RigLogic/Source/RigLogicLib/Private/dna/LODMapping.cpp
+++ b/Engine/Plugins/Animation/RigLogic/Source/RigLogicLib/Private/dna/LODMapping.cpp
@@ -92,7 +92,7 @@ void LODMapping::addIndices(std::uint16_t index, const std::uint16_t* source, st
     if (index >= indices.size()) {
         indices.resize(index + 1ul);
     }
-    indices[index].reserve(count);
+    indices[index].reserve(indices[index].size() + count);
     indices[index].insert(indices[index].end(), source, source + count);
 }
 
@@ -116,6 +116,12 @@ void LODMapping::filterIndices(std::function<bool(std::uint16_t)> filterer) {
     }
 }
 
+void LODMapping::sortIndices() {
+    for (auto& row : indices) {
+        std::sort(row.begin(), row.end());
+    }
+}
+
 void LODMapping::associateLODWithIndices(std::uint16_t lod, std::uint16_t index) {
     if (lod >= lods.size()) {
         lods.resize(lod + 1ul);
diff --git a/Engine/Plugins/Animation/RigLogic/Source/RigLogicLib/Private/dna/LODMapping.h b/Engine/Plugins/Animation/RigLogic/Source/RigLogicLib/Private/dna/LODMapping.h
--- a/Engine/Plugins/Animation/RigLogic/Source/RigLogicLib/Private/dna/LODMapping.h
+++ b/Engine/Plugins/Animation/RigLogic/Source/RigLogicLib/Private/dna/LODMapping.h
@@ -10,6 +10,7 @@
 #endif
 #include <cstdint>
 #include <functional>
+#include <iterator>
 #ifdef _MSC_VER
     #pragma warning(pop)
 #endif
@@ -32,9 +33,20 @@ class LODMapping {
         std::uint16_t getIndexListCount() const;
         void clearIndices(std::uint16_t index);
         void addIndices(std::uint16_t index, const std::uint16_t* source, std::uint16_t count);
+
+        template<class TIterator>
+        void addIndices(std::uint16_t index, TIterator begin, TIterator end) {
+            if (index >= indices.size()) {
+                indices.resize(index + 1ul);
+            }
+            indices[index].reserve(indices[index].size() + static_cast<std::size_t>(std::distance(begin, end)));
+            indices[index].insert(indices[index].end(), begin, end);
+        }
+
         void associateLODWithIndices(std::uint16_t lod, std::uint16_t index);
         void mapIndices(std::function<std::uint16_t(std::uint16_t)> mapper);
         void filterIndices(std::function<bool(std::uint16_t)> filterer);
+        void sortIndices();
         UnorderedSet<std::uint16_t> getCombinedDistinctIndices(MemoryResource* memRes) const;
 
     private:
diff --git a/Engine/Plugins/Animation/RigLogic/Source/RigLogicLib/Private/dna/Reader.cpp b/Engine/Plugins/Animation/RigLogic/Source/RigLogicLib/Private/dna/Reader.cpp
--- a/Engine/Plugins/Animation/RigLogic/Source/RigLogicLib/Private/dna/Reader.cpp
+++ b/Engine/Plugins/Animation/RigLogic/Source/RigLogicLib/Private/dna/Reader.cpp
@@ -10,6 +10,9 @@ DefinitionReader::~DefinitionReader() = default;
 BehaviorReader::~BehaviorReader() = default;
 GeometryReader::~GeometryReader() = default;
 MachineLearnedBehaviorReader::~MachineLearnedBehaviorReader() = default;
+RBFBehaviorReader::~RBFBehaviorReader() = default;
+TwistSwingBehaviorReader::~TwistSwingBehaviorReader() = default;
+JointBehaviorMetadataReader::~JointBehaviorMetadataReader() = default;
 Reader::~Reader() = default;
 
 }  // namespace dna
diff --git a/Engine/Plugins/Animation/RigLogic/Source/RigLogicLib/Private/dna/ReaderImpl.h b/Engine/Plugins/Animation/RigLogic/Source/RigLogicLib/Private/dna/ReaderImpl.h
--- a/Engine/Plugins/Animation/RigLogic/Source/RigLogicLib/Private/dna/ReaderImpl.h
+++ b/Engine/Plugins/Animation/RigLogic/Source/RigLogicLib/Private/dna/ReaderImpl.h
@@ -4,8 +4,7 @@
 
 #include "dna/BaseImpl.h"
 #include "dna/DenormalizedData.h"
-#include "dna/layers/Geometry.h"
-#include "dna/layers/MachineLearnedBehavior.h"
+#include "dna/Reader.h"
 #include "dna/TypeDefs.h"
 
 #ifdef _MSC_VER
@@ -172,6 +171,47 @@ class ReaderImpl : public TReaderBase, public virtual BaseImpl {
         ConstArrayView<float> getNeuralNetworkLayerBiases(std::uint16_t netIndex, std::uint16_t layerIndex) const override;
         ConstArrayView<float> getNeuralNetworkLayerWeights(std::uint16_t netIndex, std::uint16_t layerIndex) const override;
 
+        // RBFBehaviorReader methods
+        std::uint16_t getRBFPoseCount() const override;
+        StringView getRBFPoseName(std::uint16_t index) const override;
+        ConstArrayView<std::uint16_t> getRBFPoseJointOutputIndices(std::uint16_t poseIndex) const override;
+        ConstArrayView<std::uint16_t> getRBFPoseBlendShapeChannelOutputIndices(std::uint16_t poseIndex) const override;
+        ConstArrayView<std::uint16_t> getRBFPoseAnimatedMapOutputIndices(std::uint16_t poseIndex) const override;
+        ConstArrayView<float> getRBFPoseJointOutputValues(std::uint16_t poseIndex) const override;
+        float getRBFPoseScale(std::uint16_t poseIndex) const override;
+        std::uint16_t getRBFSolverCount() const override;
+        std::uint16_t getRBFSolverIndexListCount() const override;
+        ConstArrayView<std::uint16_t> getRBFSolverIndicesForLOD(std::uint16_t lod) const override;
+        StringView getRBFSolverName(std::uint16_t index) const override;
+        ConstArrayView<std::uint16_t> getRBFSolverRawControlIndices(std::uint16_t solverIndex) const override;
+        ConstArrayView<std::uint16_t> getRBFSolverPoseIndices(std::uint16_t solverIndex) const override;
+        ConstArrayView<float> getRBFSolverRawControlValues(std::uint16_t solverIndex) const override;
+        RBFSolverType getRBFSolverType(std::uint16_t solverIndex) const override;
+        float getRBFSolverRadius(std::uint16_t solverIndex) const override;
+        AutomaticRadius getRBFSolverAutomaticRadius(std::uint16_t solverIndex) const override;
+        float getRBFSolverWeightThreshold(std::uint16_t solverIndex) const override;
+        RBFDistanceMethod getRBFSolverDistanceMethod(std::uint16_t solverIndex) const override;
+        RBFNormalizeMethod getRBFSolverNormalizeMethod(std::uint16_t solverIndex) const override;
+        RBFFunctionType getRBFSolverFunctionType(std::uint16_t solverIndex) const override;
+        TwistAxis getRBFSolverTwistAxis(std::uint16_t solverIndex) const override;
+
+        // JointBehaviorMetadataReader methods
+        TranslationRepresentation getJointTranslationRepresentation(std::uint16_t jointIndex) const override;
+        RotationRepresentation getJointRotationRepresentation(std::uint16_t jointIndex) const override;
+        ScaleRepresentation getJointScaleRepresentation(std::uint16_t jointIndex) const override;
+
+        // TwistSwingBehaviorReader methods
+        std::uint16_t getTwistCount() const override;
+        TwistAxis getTwistSetupTwistAxis(std::uint16_t twistIndex) const override;
+        ConstArrayView<std::uint16_t> getTwistInputControlIndices(std::uint16_t twistIndex) const override;
+        ConstArrayView<std::uint16_t> getTwistOutputJointIndices(std::uint16_t twistIndex) const override;
+        ConstArrayView<float> getTwistBlendWeights(std::uint16_t twistIndex) const override;
+        std::uint16_t getSwingCount() const override;
+        TwistAxis getSwingSetupTwistAxis(std::uint16_t swingIndex) const override;
+        ConstArrayView<std::uint16_t> getSwingInputControlIndices(std::uint16_t swingIndex) const override;
+        ConstArrayView<std::uint16_t> getSwingOutputJointIndices(std::uint16_t swingIndex) const override;
+        ConstArrayView<float> getSwingBlendWeights(std::uint16_t swingIndex) const override;
+
     protected:
         mutable DenormalizedData<TReaderBase> cache;
 
@@ -435,7 +475,7 @@ template<class TReaderBase>
 inline ConstArrayView<std::uint16_t> ReaderImpl<TReaderBase>::getMeshBlendShapeChannelMappingIndicesForLOD(std::uint16_t lod)
 const {
     if (cache.meshBlendShapeMappingIndices.getLODCount() == static_cast<std::uint16_t>(0)) {
-        cache.populate(this);
+        cache.populateMeshBlendShapeMappingIndices(this);
     }
     return cache.meshBlendShapeMappingIndices.getIndices(lod);
 }
@@ -566,7 +606,7 @@ inline std::uint16_t ReaderImpl<TReaderBase>::getJointColumnCount() const {
 template<class TReaderBase>
 inline ConstArrayView<std::uint16_t> ReaderImpl<TReaderBase>::getJointVariableAttributeIndices(std::uint16_t lod) const {
     if (cache.jointVariableAttributeIndices.getLODCount() == static_cast<std::uint16_t>(0)) {
-        cache.populate(this);
+        cache.populateJointVariableAttributeIndices(this);
     }
     return cache.jointVariableAttributeIndices.getIndices(lod);
 }
@@ -1021,7 +1061,7 @@ inline std::uint16_t ReaderImpl<TReaderBase>::getNeuralNetworkIndexListCount() c
 }
 
 template<class TReaderBase>
-ConstArrayView<std::uint16_t> ReaderImpl<TReaderBase>::getNeuralNetworkIndicesForLOD(std::uint16_t lod) const {
+inline ConstArrayView<std::uint16_t> ReaderImpl<TReaderBase>::getNeuralNetworkIndicesForLOD(std::uint16_t lod) const {
     return dna.machineLearnedBehavior.lodNeuralNetworkMapping.getIndices(lod);
 }
 
@@ -1043,8 +1083,9 @@ inline StringView ReaderImpl<TReaderBase>::getMeshRegionName(std::uint16_t meshI
 }
 
 template<class TReaderBase>
-ConstArrayView<std::uint16_t> ReaderImpl<TReaderBase>::getNeuralNetworkIndicesForMeshRegion(std::uint16_t meshIndex,
-                                                                                            std::uint16_t regionIndex) const {
+inline ConstArrayView<std::uint16_t> ReaderImpl<TReaderBase>::getNeuralNetworkIndicesForMeshRegion(std::uint16_t meshIndex,
+                                                                                                   std::uint16_t regionIndex)
+const {
     if (meshIndex < dna.machineLearnedBehavior.neuralNetworkToMeshRegion.indices.size()) {
         const auto& mesh = dna.machineLearnedBehavior.neuralNetworkToMeshRegion.indices[meshIndex];
         if (regionIndex < mesh.size()) {
@@ -1123,6 +1164,289 @@ inline ConstArrayView<float> ReaderImpl<TReaderBase>::getNeuralNetworkLayerWeigh
     return {};
 }
 
+template<class TReaderBase>
+inline std::uint16_t ReaderImpl<TReaderBase>::getRBFPoseCount() const {
+    return static_cast<std::uint16_t>(dna.rbfBehavior.poses.size());
+}
+
+template<class TReaderBase>
+inline StringView ReaderImpl<TReaderBase>::getRBFPoseName(std::uint16_t poseIndex) const {
+    const auto& poses = dna.rbfBehavior.poses;
+    if (poseIndex < poses.size()) {
+        return poses[poseIndex].name;
+    }
+    return {};
+}
+
+template<class TReaderBase>
+inline ConstArrayView<std::uint16_t> ReaderImpl<TReaderBase>::getRBFPoseJointOutputIndices(std::uint16_t poseIndex) const {
+    const auto& poses = dna.rbfBehavior.poses;
+    if (poseIndex < poses.size()) {
+        if (cache.rbfPoseJointOutputIndices.size() == 0) {
+            cache.populateRBFPoseJointOutputIndices(this);
+        }
+        return cache.rbfPoseJointOutputIndices[poseIndex];
+    }
+    return {};
+}
+
+template<class TReaderBase>
+inline ConstArrayView<std::uint16_t> ReaderImpl<TReaderBase>::getRBFPoseBlendShapeChannelOutputIndices(std::uint16_t poseIndex)
+const {
+    const auto& poses = dna.rbfBehavior.poses;
+    if (poseIndex < poses.size()) {
+        if (cache.rbfBlendShapeChannelOutputIndices.size() == 0) {
+            cache.populateRBFBlendShapeOutputIndices(this);
+        }
+        return cache.rbfBlendShapeChannelOutputIndices[poseIndex];
+    }
+    return {};
+}
+
+template<class TReaderBase>
+inline ConstArrayView<std::uint16_t> ReaderImpl<TReaderBase>::getRBFPoseAnimatedMapOutputIndices(std::uint16_t poseIndex)
+const {
+    const auto& poses = dna.rbfBehavior.poses;
+    if (poseIndex < poses.size()) {
+        if (cache.rbfAnimatedMapOutputIndices.size() == 0) {
+            cache.populateRBFAnimatedMapOutputIndices(this);
+        }
+        return cache.rbfAnimatedMapOutputIndices[poseIndex];
+    }
+    return {};
+}
+
+template<class TReaderBase>
+inline ConstArrayView<float> ReaderImpl<TReaderBase>::getRBFPoseJointOutputValues(std::uint16_t poseIndex) const {
+    const auto& poses = dna.rbfBehavior.poses;
+    if (poseIndex < poses.size()) {
+        if (cache.rbfPoseJointOutputValues.size() == 0) {
+            cache.populateRBFPoseJointOutputValues(this);
+        }
+        return cache.rbfPoseJointOutputValues[poseIndex];
+    }
+    return {};
+}
+
+template<class TReaderBase>
+inline float ReaderImpl<TReaderBase>::getRBFPoseScale(std::uint16_t poseIndex) const {
+    const auto& poses = dna.rbfBehavior.poses;
+    if (poseIndex < poses.size()) {
+        return poses[poseIndex].scale;
+    }
+    return 0.0f;
+}
+
+template<class TReaderBase>
+inline std::uint16_t ReaderImpl<TReaderBase>::getRBFSolverCount() const {
+    return static_cast<std::uint16_t>(dna.rbfBehavior.solvers.size());
+}
+
+template<class TReaderBase>
+inline std::uint16_t ReaderImpl<TReaderBase>::getRBFSolverIndexListCount() const {
+    return dna.rbfBehavior.lodSolverMapping.getIndexListCount();
+}
+
+template<class TReaderBase>
+inline ConstArrayView<std::uint16_t> ReaderImpl<TReaderBase>::getRBFSolverIndicesForLOD(std::uint16_t lod) const {
+    return dna.rbfBehavior.lodSolverMapping.getIndices(lod);
+}
+
+template<class TReaderBase>
+inline StringView ReaderImpl<TReaderBase>::getRBFSolverName(std::uint16_t solverIndex) const {
+    if (dna.rbfBehavior.solvers.size() > solverIndex) {
+        return dna.rbfBehavior.solvers[solverIndex].name;
+    }
+    return {};
+}
+
+template<class TReaderBase>
+inline ConstArrayView<std::uint16_t> ReaderImpl<TReaderBase>::getRBFSolverRawControlIndices(std::uint16_t solverIndex)
+const {
+    if (dna.rbfBehavior.solvers.size() > solverIndex) {
+        return dna.rbfBehavior.solvers[solverIndex].rawControlIndices;
+    }
+    return {};
+}
+
+template<class TReaderBase>
+inline ConstArrayView<std::uint16_t> ReaderImpl<TReaderBase>::getRBFSolverPoseIndices(std::uint16_t solverIndex) const {
+    if (dna.rbfBehavior.solvers.size() > solverIndex) {
+        return dna.rbfBehavior.solvers[solverIndex].poseIndices;
+    }
+    return {};
+}
+
+template<class TReaderBase>
+inline ConstArrayView<float> ReaderImpl<TReaderBase>::getRBFSolverRawControlValues(std::uint16_t solverIndex) const {
+    if (dna.rbfBehavior.solvers.size() > solverIndex) {
+        return dna.rbfBehavior.solvers[solverIndex].rawControlValues;
+    }
+    return {};
+}
+
+template<class TReaderBase>
+inline RBFSolverType ReaderImpl<TReaderBase>::getRBFSolverType(std::uint16_t solverIndex) const {
+    if (dna.rbfBehavior.solvers.size() > solverIndex) {
+        return static_cast<RBFSolverType>(dna.rbfBehavior.solvers[solverIndex].solverType);
+    }
+    return {};
+}
+
+template<class TReaderBase>
+inline float ReaderImpl<TReaderBase>::getRBFSolverRadius(std::uint16_t solverIndex) const {
+    if (dna.rbfBehavior.solvers.size() > solverIndex) {
+        return dna.rbfBehavior.solvers[solverIndex].radius;
+    }
+    return 0.0f;
+}
+
+template<class TReaderBase>
+inline AutomaticRadius ReaderImpl<TReaderBase>::getRBFSolverAutomaticRadius(std::uint16_t solverIndex) const {
+    if (dna.rbfBehavior.solvers.size() > solverIndex) {
+        return static_cast<AutomaticRadius>(dna.rbfBehavior.solvers[solverIndex].automaticRadius);
+    }
+    return {};
+}
+
+template<class TReaderBase>
+inline float ReaderImpl<TReaderBase>::getRBFSolverWeightThreshold(std::uint16_t solverIndex) const {
+    if (dna.rbfBehavior.solvers.size() > solverIndex) {
+        return dna.rbfBehavior.solvers[solverIndex].weightThreshold;
+    }
+    return 0.0f;
+}
+
+template<class TReaderBase>
+inline RBFDistanceMethod ReaderImpl<TReaderBase>::getRBFSolverDistanceMethod(std::uint16_t solverIndex) const {
+    if (dna.rbfBehavior.solvers.size() > solverIndex) {
+        return static_cast<RBFDistanceMethod>(dna.rbfBehavior.solvers[solverIndex].distanceMethod);
+    }
+    return {};
+}
+
+template<class TReaderBase>
+inline RBFNormalizeMethod ReaderImpl<TReaderBase>::getRBFSolverNormalizeMethod(std::uint16_t solverIndex) const {
+    if (dna.rbfBehavior.solvers.size() > solverIndex) {
+        return static_cast<RBFNormalizeMethod>(dna.rbfBehavior.solvers[solverIndex].normalizeMethod);
+    }
+    return {};
+}
+
+template<class TReaderBase>
+inline RBFFunctionType ReaderImpl<TReaderBase>::getRBFSolverFunctionType(std::uint16_t solverIndex) const {
+    if (dna.rbfBehavior.solvers.size() > solverIndex) {
+        return static_cast<RBFFunctionType>(dna.rbfBehavior.solvers[solverIndex].functionType);
+    }
+    return {};
+}
+
+template<class TReaderBase>
+inline TwistAxis ReaderImpl<TReaderBase>::getRBFSolverTwistAxis(std::uint16_t solverIndex) const {
+    if (dna.rbfBehavior.solvers.size() > solverIndex) {
+        return static_cast<TwistAxis>(dna.rbfBehavior.solvers[solverIndex].twistAxis);
+    }
+    return {};
+}
+
+template<class TReaderBase>
+inline TranslationRepresentation ReaderImpl<TReaderBase>::getJointTranslationRepresentation(std::uint16_t jointIndex) const {
+    if (dna.jointBehaviorMetadata.jointRepresentations.size() > jointIndex) {
+        return static_cast<TranslationRepresentation>(dna.jointBehaviorMetadata.jointRepresentations[jointIndex].translation);
+    }
+    return {};
+}
+
+template<class TReaderBase>
+inline RotationRepresentation ReaderImpl<TReaderBase>::getJointRotationRepresentation(std::uint16_t jointIndex) const {
+    if (dna.jointBehaviorMetadata.jointRepresentations.size() > jointIndex) {
+        return static_cast<RotationRepresentation>(dna.jointBehaviorMetadata.jointRepresentations[jointIndex].rotation);
+    }
+    return {};
+}
+
+template<class TReaderBase>
+inline ScaleRepresentation ReaderImpl<TReaderBase>::getJointScaleRepresentation(std::uint16_t jointIndex) const {
+    if (dna.jointBehaviorMetadata.jointRepresentations.size() > jointIndex) {
+        return static_cast<ScaleRepresentation>(dna.jointBehaviorMetadata.jointRepresentations[jointIndex].scale);
+    }
+    return {};
+}
+
+template<class TReaderBase>
+inline std::uint16_t ReaderImpl<TReaderBase>::getTwistCount() const {
+    return static_cast<std::uint16_t>(dna.twistSwingBehavior.twists.size());
+}
+
+template<class TReaderBase>
+inline TwistAxis ReaderImpl<TReaderBase>::getTwistSetupTwistAxis(std::uint16_t twistIndex) const {
+    if (dna.twistSwingBehavior.twists.size() > twistIndex) {
+        return static_cast<TwistAxis>(dna.twistSwingBehavior.twists[twistIndex].twistAxis);
+    }
+    return {};
+}
+
+template<class TReaderBase>
+inline ConstArrayView<std::uint16_t> ReaderImpl<TReaderBase>::getTwistInputControlIndices(std::uint16_t twistIndex) const {
+    if (dna.twistSwingBehavior.twists.size() > twistIndex) {
+        return dna.twistSwingBehavior.twists[twistIndex].twistInputControlIndices;
+    }
+    return {};
+}
+
+template<class TReaderBase>
+inline ConstArrayView<std::uint16_t> ReaderImpl<TReaderBase>::getTwistOutputJointIndices(std::uint16_t twistIndex) const {
+    if (dna.twistSwingBehavior.twists.size() > twistIndex) {
+        return dna.twistSwingBehavior.twists[twistIndex].twistOutputJointIndices;
+    }
+    return {};
+}
+
+template<class TReaderBase>
+inline ConstArrayView<float> ReaderImpl<TReaderBase>::getTwistBlendWeights(std::uint16_t twistIndex) const {
+    if (dna.twistSwingBehavior.twists.size() > twistIndex) {
+        return dna.twistSwingBehavior.twists[twistIndex].twistBlendWeights;
+    }
+    return {};
+}
+
+template<class TReaderBase>
+inline std::uint16_t ReaderImpl<TReaderBase>::getSwingCount() const {
+    return static_cast<std::uint16_t>(dna.twistSwingBehavior.swings.size());
+}
+
+template<class TReaderBase>
+inline TwistAxis ReaderImpl<TReaderBase>::getSwingSetupTwistAxis(std::uint16_t swingIndex) const {
+    if (dna.twistSwingBehavior.swings.size() > swingIndex) {
+        return static_cast<TwistAxis>(dna.twistSwingBehavior.swings[swingIndex].twistAxis);
+    }
+    return {};
+}
+
+template<class TReaderBase>
+inline ConstArrayView<std::uint16_t> ReaderImpl<TReaderBase>::getSwingInputControlIndices(std::uint16_t swingIndex) const {
+    if (dna.twistSwingBehavior.swings.size() > swingIndex) {
+        return dna.twistSwingBehavior.swings[swingIndex].swingInputControlIndices;
+    }
+    return {};
+}
+
+template<class TReaderBase>
+inline ConstArrayView<std::uint16_t> ReaderImpl<TReaderBase>::getSwingOutputJointIndices(std::uint16_t swingIndex) const {
+    if (dna.twistSwingBehavior.swings.size() > swingIndex) {
+        return dna.twistSwingBehavior.swings[swingIndex].swingOutputJointIndices;
+    }
+    return {};
+}
+
+template<class TReaderBase>
+inline ConstArrayView<float> ReaderImpl<TReaderBase>::getSwingBlendWeights(std::uint16_t swingIndex) const {
+    if (dna.twistSwingBehavior.swings.size() > swingIndex) {
+        return dna.twistSwingBehavior.swings[swingIndex].swingBlendWeights;
+    }
+    return {};
+}
+
 #ifdef _MSC_VER
     #pragma warning(pop)
 #endif
diff --git a/Engine/Plugins/Animation/RigLogic/Source/RigLogicLib/Private/dna/Writer.cpp b/Engine/Plugins/Animation/RigLogic/Source/RigLogicLib/Private/dna/Writer.cpp
--- a/Engine/Plugins/Animation/RigLogic/Source/RigLogicLib/Private/dna/Writer.cpp
+++ b/Engine/Plugins/Animation/RigLogic/Source/RigLogicLib/Private/dna/Writer.cpp
@@ -32,6 +32,9 @@ DefinitionWriter::~DefinitionWriter() = default;
 BehaviorWriter::~BehaviorWriter() = default;
 GeometryWriter::~GeometryWriter() = default;
 MachineLearnedBehaviorWriter::~MachineLearnedBehaviorWriter() = default;
+RBFBehaviorWriter::~RBFBehaviorWriter() = default;
+TwistSwingBehaviorWriter::~TwistSwingBehaviorWriter() = default;
+JointBehaviorMetadataWriter::~JointBehaviorMetadataWriter() = default;
 Writer::~Writer() = default;
 
 template<typename TVector, typename TGetter>
@@ -55,7 +58,7 @@ static void copyDescriptor(const DescriptorReader* source, DescriptorWriter* des
     destination->setGender(source->getGender());
     destination->setAge(source->getAge());
     destination->clearMetaData();
-    for (std::uint32_t i = 0u; i < source->getMetaDataCount(); ++i) {
+    for (std::uint32_t i = {}; i < source->getMetaDataCount(); ++i) {
         const auto key = source->getMetaDataKey(i);
         const auto value = source->getMetaDataValue(key);
         destination->setMetaData(key, value);
@@ -92,8 +95,8 @@ static void copyLODIndices(IndicesGetter getIndices,
                            std::uint16_t lodCount,
                            MemoryResource* memRes) {
     Matrix<std::uint16_t> allIndices{memRes};
-    std::uint16_t index = 0u;
-    for (std::uint16_t lod = 0u; lod < lodCount; ++lod) {
+    std::uint16_t index = {};
+    for (std::uint16_t lod = {}; lod < lodCount; ++lod) {
         auto indices = getIndices(lod);
         // Check if these same indices were perhaps already used for previous LODs
         auto found = findIndices(allIndices, indices);
@@ -181,12 +184,12 @@ static void copyDefinition(const DefinitionReader* source, DefinitionWriter* des
 
     Vector<std::uint16_t> jointHierarchy{memRes};
     jointHierarchy.reserve(source->getJointCount());
-    for (std::uint16_t i = 0u; i < source->getJointCount(); ++i) {
+    for (std::uint16_t i = {}; i < source->getJointCount(); ++i) {
         jointHierarchy.push_back(source->getJointParentIndex(i));
     }
     destination->setJointHierarchy(jointHierarchy.data(), static_cast<std::uint16_t>(jointHierarchy.size()));
 
-    for (std::uint16_t i = 0u; i < source->getMeshBlendShapeChannelMappingCount(); ++i) {
+    for (std::uint16_t i = {}; i < source->getMeshBlendShapeChannelMappingCount(); ++i) {
         auto mapping = source->getMeshBlendShapeChannelMapping(i);
         destination->setMeshBlendShapeChannelMapping(i, mapping.meshIndex, mapping.blendShapeChannelIndex);
     }
@@ -501,6 +504,113 @@ static void copyMachineLearnedBehavior(const MachineLearnedBehaviorReader* sourc
     }
 }
 
+static bool hasRBFBehavior(const RBFBehaviorReader* source) {
+    // Heuristic for determining whether source DNA actually has any RBF behavior data
+    return (source->getRBFPoseCount() != 0u) || (source->getRBFSolverCount() != 0u);
+}
+
+static void copyRBFBehavior(const RBFBehaviorReader* source, RBFBehaviorWriter* destination, MemoryResource* memRes) {
+    destination->clearRBFPoses();
+    destination->clearRBFSolverIndices();
+    destination->clearRBFSolvers();
+
+    if (!hasRBFBehavior(source)) {
+        // Source DNA was loaded without RBF behavior layer
+        return;
+    }
+
+    using namespace std::placeholders;
+    copyLODIndices(
+        std::bind(&RBFBehaviorReader::getRBFSolverIndicesForLOD, source, _1),
+        std::bind(&RBFBehaviorWriter::setRBFSolverIndices, destination, _1, _2, _3),
+        std::bind(&RBFBehaviorWriter::setLODRBFSolverMapping, destination, _1, _2),
+        source->getLODCount(),
+        memRes);
+
+    for (std::uint16_t posePlusOne = source->getRBFPoseCount(); posePlusOne > 0u; --posePlusOne) {
+        const auto pi = static_cast<std::uint16_t>(posePlusOne - 1u);
+        destination->setRBFPoseName(pi, source->getRBFPoseName(pi).data());
+        destination->setRBFPoseScale(pi, source->getRBFPoseScale(pi));
+    }
+
+    for (std::uint16_t solverPlusOne = source->getRBFSolverCount(); solverPlusOne > 0u; --solverPlusOne) {
+        const auto si = static_cast<std::uint16_t>(solverPlusOne - 1u);
+        destination->setRBFSolverName(si, source->getRBFSolverName(si).data());
+        const auto rawControlIndices = source->getRBFSolverRawControlIndices(si);
+        destination->setRBFSolverRawControlIndices(si, rawControlIndices.data(),
+                                                   static_cast<std::uint16_t>(rawControlIndices.size()));
+        const auto poseIndices = source->getRBFSolverPoseIndices(si);
+        destination->setRBFSolverPoseIndices(si, poseIndices.data(), static_cast<std::uint16_t>(poseIndices.size()));
+        const auto rawControlValues = source->getRBFSolverRawControlValues(si);
+        destination->setRBFSolverRawControlValues(si, rawControlValues.data(),
+                                                  static_cast<std::uint16_t>(rawControlValues.size()));
+        destination->setRBFSolverType(si, source->getRBFSolverType(si));
+        destination->setRBFSolverRadius(si, source->getRBFSolverRadius(si));
+        destination->setRBFSolverAutomaticRadius(si, source->getRBFSolverAutomaticRadius(si));
+        destination->setRBFSolverWeightThreshold(si, source->getRBFSolverWeightThreshold(si));
+        destination->setRBFSolverDistanceMethod(si, source->getRBFSolverDistanceMethod(si));
+        destination->setRBFSolverNormalizeMethod(si, source->getRBFSolverNormalizeMethod(si));
+        destination->setRBFSolverFunctionType(si, source->getRBFSolverFunctionType(si));
+        destination->setRBFSolverTwistAxis(si, source->getRBFSolverTwistAxis(si));
+    }
+}
+
+static void copyJointBehaviorMetaData(const JointBehaviorMetadataReader* source,
+                                      JointBehaviorMetadataWriter* destination,
+                                      MemoryResource*  /*unused*/) {
+    destination->clearJointRepresentations();
+
+    for (std::uint16_t jointPlusOne = source->getJointCount(); jointPlusOne > 0u; --jointPlusOne) {
+        const auto ji = static_cast<std::uint16_t>(jointPlusOne - 1u);
+        destination->setJointTranslationRepresentation(ji, source->getJointTranslationRepresentation(ji));
+        destination->setJointRotationRepresentation(ji, source->getJointRotationRepresentation(ji));
+        destination->setJointScaleRepresentation(ji, source->getJointScaleRepresentation(ji));
+    }
+}
+
+static bool hasTwistSwingBehavior(const TwistSwingBehaviorReader* source) {
+    // Heuristic for determining whether source DNA actually has any TwistSwing behavior data
+    return (source->getTwistCount() != 0u) || (source->getSwingCount() != 0u);
+}
+
+static void copyTwistSwingBehavior(const TwistSwingBehaviorReader* source, TwistSwingBehaviorWriter* destination,
+                                   MemoryResource*  /*unused*/) {
+    destination->clearTwists();
+    destination->clearSwings();
+
+    if (!hasTwistSwingBehavior(source)) {
+        // Source DNA was loaded without twist and swig behavior layer
+        return;
+    }
+
+    for (std::uint16_t twistIndexPlusOne = source->getTwistCount(); twistIndexPlusOne > 0u; --twistIndexPlusOne) {
+        const auto ti = static_cast<std::uint16_t>(twistIndexPlusOne - 1u);
+        destination->setTwistSetupTwistAxis(ti, source->getTwistSetupTwistAxis(ti));
+        const auto twistInputControlIndices = source->getTwistInputControlIndices(ti);
+        destination->setTwistInputControlIndices(ti, twistInputControlIndices.data(),
+                                                 static_cast<std::uint16_t>(twistInputControlIndices.size()));
+        const auto twistOutputJointIndices = source->getTwistOutputJointIndices(ti);
+        destination->setTwistOutputJointIndices(ti, twistOutputJointIndices.data(),
+                                                static_cast<std::uint16_t>(twistOutputJointIndices.size()));
+        const auto twistBlendWeights = source->getTwistBlendWeights(ti);
+        destination->setTwistBlendWeights(ti, twistBlendWeights.data(),
+                                          static_cast<std::uint16_t>(twistBlendWeights.size()));
+    }
+
+    for (std::uint16_t swingIndexPlusOne = source->getSwingCount(); swingIndexPlusOne > 0u; --swingIndexPlusOne) {
+        const auto si = static_cast<std::uint16_t>(swingIndexPlusOne - 1u);
+        destination->setSwingSetupTwistAxis(si, source->getSwingSetupTwistAxis(si));
+        const auto swingInputControlIndices = source->getSwingInputControlIndices(si);
+        destination->setSwingInputControlIndices(si, swingInputControlIndices.data(),
+                                                 static_cast<std::uint16_t>(swingInputControlIndices.size()));
+        const auto swingOutputJointIndices = source->getSwingOutputJointIndices(si);
+        destination->setSwingOutputJointIndices(si, swingOutputJointIndices.data(),
+                                                static_cast<std::uint16_t>(swingOutputJointIndices.size()));
+        const auto swingBlendWeights = source->getSwingBlendWeights(si);
+        destination->setSwingBlendWeights(si, swingBlendWeights.data(), static_cast<std::uint16_t>(swingBlendWeights.size()));
+    }
+}
+
 void Writer::setFrom(const Reader* source, DataLayer layer, UnknownLayerPolicy  /*unused*/, MemoryResource* memRes) {
     // Unknown layers are inaccessible when copying through the public API
     if (source == nullptr) {
@@ -525,6 +635,15 @@ void Writer::setFrom(const Reader* source, DataLayer layer, UnknownLayerPolicy
     if (contains(bitmask, DataLayerBitmask::MachineLearnedBehavior)) {
         copyMachineLearnedBehavior(source, this, memRes);
     }
+    if (contains(bitmask, DataLayerBitmask::RBFBehavior)) {
+        copyRBFBehavior(source, this, memRes);
+    }
+    if (contains(bitmask, DataLayerBitmask::JointBehaviorMetadata)) {
+        copyJointBehaviorMetaData(source, this, memRes);
+    }
+    if (contains(bitmask, DataLayerBitmask::TwistSwingBehavior)) {
+        copyTwistSwingBehavior(source, this, memRes);
+    }
 }
 
 }  // namespace dna
diff --git a/Engine/Plugins/Animation/RigLogic/Source/RigLogicLib/Private/dna/WriterImpl.h b/Engine/Plugins/Animation/RigLogic/Source/RigLogicLib/Private/dna/WriterImpl.h
--- a/Engine/Plugins/Animation/RigLogic/Source/RigLogicLib/Private/dna/WriterImpl.h
+++ b/Engine/Plugins/Animation/RigLogic/Source/RigLogicLib/Private/dna/WriterImpl.h
@@ -3,11 +3,8 @@
 #pragma once
 
 #include "dna/BaseImpl.h"
-#include "dna/layers/Descriptor.h"
-#include "dna/layers/Geometry.h"
-#include "dna/layers/MachineLearnedBehavior.h"
 #include "dna/TypeDefs.h"
-#include "dna/types/Vector3.h"
+#include "dna/Writer.h"
 #include "dna/utils/Extd.h"
 
 #ifdef _MSC_VER
@@ -18,6 +15,7 @@
 #include <cstddef>
 #include <cstring>
 #include <tuple>
+#include <type_traits>
 #include <utility>
 #ifdef _MSC_VER
     #pragma warning(pop)
@@ -25,16 +23,38 @@
 
 namespace dna {
 
-template<class TContainer, typename ... Args>
-void ensureHasSize(TContainer& target,
-                   std::size_t size,
-                   Args&& ... args) {
+template<class TContainer>
+typename std::enable_if<std::is_constructible<typename TContainer::value_type, MemoryResource*>::value>::type
+ensureHasSize(TContainer& target, std::size_t size) {
     target.reserve(size);
     while (target.size() < size) {
-        target.push_back(typename TContainer::value_type(std::forward<Args>(args)...));
+        target.push_back(typename TContainer::value_type(target.get_allocator().getMemoryResource()));
     }
 }
 
+template<class TContainer>
+typename std::enable_if<!std::is_constructible<typename TContainer::value_type, MemoryResource*>::value>::type
+ensureHasSize(TContainer& target, std::size_t size) {
+    if (target.size() < size) {
+        target.resize(size);
+    }
+}
+
+template<class TContainer, typename U>
+typename std::enable_if<std::is_integral<U>::value, typename TContainer::value_type&>::type
+getAt(TContainer& target, U index) {
+    ensureHasSize(target, index + 1ul);
+    return target[index];
+}
+
+template<class TContainer, typename TSize, typename TValue>
+typename std::enable_if<std::is_integral<TSize>::value>::type
+setAt(TContainer& target,
+      TSize index,
+      const TValue& value) {
+    getAt(target, index) = value;
+}
+
 template<class TWriterBase>
 class WriterImpl : public TWriterBase, public virtual BaseImpl {
     public:
@@ -201,8 +221,56 @@ class WriterImpl : public TWriterBase, public virtual BaseImpl {
                                           const float* weights,
                                           std::uint32_t count) override;
 
-};
+        // RBFBehaviorWriter methods
+        void clearRBFPoses() override;
+        void setRBFPoseName(std::uint16_t poseIndex, const char* name) override;
+        void setRBFPoseScale(std::uint16_t poseIndex, float scale) override;
+        void clearRBFSolvers() override;
+        void clearRBFSolverIndices() override;
+        void setRBFSolverIndices(std::uint16_t index, const std::uint16_t* solverIndices, std::uint16_t count) override;
+        void clearLODRBFSolverMappings() override;
+        void setLODRBFSolverMapping(std::uint16_t lod, std::uint16_t index) override;
+        void setRBFSolverName(std::uint16_t solverIndex, const char* name) override;
+        void setRBFSolverRawControlIndices(std::uint16_t solverIndex, const std::uint16_t* inputIndices,
+                                           std::uint16_t count) override;
+        void setRBFSolverPoseIndices(std::uint16_t solverIndex, const std::uint16_t* poseIndices, std::uint16_t count) override;
+        void setRBFSolverRawControlValues(std::uint16_t solverIndex, const float* values, std::uint16_t count) override;
+        void setRBFSolverType(std::uint16_t solverIndex, RBFSolverType type) override;
+        void setRBFSolverRadius(std::uint16_t solverIndex, float radius) override;
+        void setRBFSolverAutomaticRadius(std::uint16_t solverIndex, AutomaticRadius automaticRadius) override;
+        void setRBFSolverWeightThreshold(std::uint16_t solverIndex, float weightThreshold) override;
+        void setRBFSolverDistanceMethod(std::uint16_t solverIndex, RBFDistanceMethod distanceMethod) override;
+        void setRBFSolverNormalizeMethod(std::uint16_t solverIndex, RBFNormalizeMethod normalizeMethod) override;
+        void setRBFSolverFunctionType(std::uint16_t solverIndex, RBFFunctionType functionType) override;
+        void setRBFSolverTwistAxis(std::uint16_t solverIndex, TwistAxis twistAxis) override;
+
+        // JointBehaviorMetadataWriter
+        void clearJointRepresentations() override;
+        void setJointTranslationRepresentation(std::uint16_t jointIndex, TranslationRepresentation representation) override;
+        void setJointRotationRepresentation(std::uint16_t jointIndex, RotationRepresentation representation) override;
+        void setJointScaleRepresentation(std::uint16_t jointIndex, ScaleRepresentation representation) override;
+
+        // TwistSwingBehaviorWriter
+        void clearTwists() override;
+        void deleteTwist(std::uint16_t twistIndex) override;
+        void setTwistSetupTwistAxis(std::uint16_t twistIndex, TwistAxis twistAxis) override;
+        void setTwistInputControlIndices(std::uint16_t twistIndex,
+                                         const std::uint16_t* controlIndices,
+                                         std::uint16_t controlIndexCount) override;
+        void setTwistOutputJointIndices(std::uint16_t twistIndex, const std::uint16_t* jointIndices,
+                                        std::uint16_t jointIndexCount) override;
+        void setTwistBlendWeights(std::uint16_t twistIndex, const float* blendWeights, std::uint16_t blendWeightCount) override;
+        void clearSwings() override;
+        void deleteSwing(std::uint16_t swingIndex) override;
+        void setSwingSetupTwistAxis(std::uint16_t swingIndex, TwistAxis twistAxis) override;
+        void setSwingInputControlIndices(std::uint16_t swingIndex,
+                                         const std::uint16_t* controlIndices,
+                                         std::uint16_t controlIndexCount) override;
+        void setSwingOutputJointIndices(std::uint16_t swingIndex, const std::uint16_t* jointIndices,
+                                        std::uint16_t jointIndexCount) override;
+        void setSwingBlendWeights(std::uint16_t swingIndex, const float* blendWeights, std::uint16_t blendWeightCount) override;
 
+};
 
 #ifdef _MSC_VER
     #pragma warning(push)
@@ -319,8 +387,7 @@ inline void WriterImpl<TWriterBase>::clearGUIControlNames() {
 
 template<class TWriterBase>
 inline void WriterImpl<TWriterBase>::setGUIControlName(std::uint16_t index, const char* name) {
-    ensureHasSize(dna.definition.guiControlNames, index + 1ul, memRes);
-    dna.definition.guiControlNames[index] = name;
+    setAt(dna.definition.guiControlNames, index, name);
 }
 
 template<class TWriterBase>
@@ -330,8 +397,7 @@ inline void WriterImpl<TWriterBase>::clearRawControlNames() {
 
 template<class TWriterBase>
 inline void WriterImpl<TWriterBase>::setRawControlName(std::uint16_t index, const char* name) {
-    ensureHasSize(dna.definition.rawControlNames, index + 1ul, memRes);
-    dna.definition.rawControlNames[index] = name;
+    setAt(dna.definition.rawControlNames, index, name);
 }
 
 template<class TWriterBase>
@@ -341,8 +407,7 @@ inline void WriterImpl<TWriterBase>::clearJointNames() {
 
 template<class TWriterBase>
 inline void WriterImpl<TWriterBase>::setJointName(std::uint16_t index, const char* name) {
-    ensureHasSize(dna.definition.jointNames, index + 1ul, memRes);
-    dna.definition.jointNames[index] = name;
+    setAt(dna.definition.jointNames, index, name);
 }
 
 template<class TWriterBase>
@@ -379,8 +444,7 @@ inline void WriterImpl<TWriterBase>::clearBlendShapeChannelNames() {
 
 template<class TWriterBase>
 inline void WriterImpl<TWriterBase>::setBlendShapeChannelName(std::uint16_t index, const char* name) {
-    ensureHasSize(dna.definition.blendShapeChannelNames, index + 1ul, memRes);
-    dna.definition.blendShapeChannelNames[index] = name;
+    setAt(dna.definition.blendShapeChannelNames, index, name);
 }
 
 template<class TWriterBase>
@@ -413,8 +477,7 @@ inline void WriterImpl<TWriterBase>::clearAnimatedMapNames() {
 
 template<class TWriterBase>
 inline void WriterImpl<TWriterBase>::setAnimatedMapName(std::uint16_t index, const char* name) {
-    ensureHasSize(dna.definition.animatedMapNames, index + 1ul, memRes);
-    dna.definition.animatedMapNames[index] = name;
+    setAt(dna.definition.animatedMapNames, index, name);
 }
 
 template<class TWriterBase>
@@ -447,8 +510,7 @@ inline void WriterImpl<TWriterBase>::clearMeshNames() {
 
 template<class TWriterBase>
 inline void WriterImpl<TWriterBase>::setMeshName(std::uint16_t index, const char* name) {
-    ensureHasSize(dna.definition.meshNames, index + 1ul, memRes);
-    dna.definition.meshNames[index] = name;
+    setAt(dna.definition.meshNames, index, name);
 }
 
 template<class TWriterBase>
@@ -571,44 +633,39 @@ template<class TWriterBase>
 inline void WriterImpl<TWriterBase>::setJointGroupLODs(std::uint16_t jointGroupIndex,
                                                        const std::uint16_t* lods,
                                                        std::uint16_t count) {
-    auto& jointGroups = dna.behavior.joints.jointGroups;
-    ensureHasSize(jointGroups, jointGroupIndex + 1ul, memRes);
-    jointGroups[jointGroupIndex].lods.assign(lods, lods + count);
+    auto& jointGroup = getAt(dna.behavior.joints.jointGroups, jointGroupIndex);
+    jointGroup.lods.assign(lods, lods + count);
 }
 
 template<class TWriterBase>
 inline void WriterImpl<TWriterBase>::setJointGroupInputIndices(std::uint16_t jointGroupIndex,
                                                                const std::uint16_t* inputIndices,
                                                                std::uint16_t count) {
-    auto& jointGroups = dna.behavior.joints.jointGroups;
-    ensureHasSize(jointGroups, jointGroupIndex + 1ul, memRes);
-    jointGroups[jointGroupIndex].inputIndices.assign(inputIndices, inputIndices + count);
+    auto& jointGroup = getAt(dna.behavior.joints.jointGroups, jointGroupIndex);
+    jointGroup.inputIndices.assign(inputIndices, inputIndices + count);
 }
 
 template<class TWriterBase>
 inline void WriterImpl<TWriterBase>::setJointGroupOutputIndices(std::uint16_t jointGroupIndex,
                                                                 const std::uint16_t* outputIndices,
                                                                 std::uint16_t count) {
-    auto& jointGroups = dna.behavior.joints.jointGroups;
-    ensureHasSize(jointGroups, jointGroupIndex + 1ul, memRes);
-    jointGroups[jointGroupIndex].outputIndices.assign(outputIndices, outputIndices + count);
+    auto& jointGroup = getAt(dna.behavior.joints.jointGroups, jointGroupIndex);
+    jointGroup.outputIndices.assign(outputIndices, outputIndices + count);
 }
 
 template<class TWriterBase>
 inline void WriterImpl<TWriterBase>::setJointGroupValues(std::uint16_t jointGroupIndex, const float* values,
                                                          std::uint32_t count) {
-    auto& jointGroups = dna.behavior.joints.jointGroups;
-    ensureHasSize(jointGroups, jointGroupIndex + 1ul, memRes);
-    jointGroups[jointGroupIndex].values.assign(values, values + count);
+    auto& jointGroup = getAt(dna.behavior.joints.jointGroups, jointGroupIndex);
+    jointGroup.values.assign(values, values + count);
 }
 
 template<class TWriterBase>
 inline void WriterImpl<TWriterBase>::setJointGroupJointIndices(std::uint16_t jointGroupIndex,
                                                                const std::uint16_t* jointIndices,
                                                                std::uint16_t count) {
-    auto& jointGroups = dna.behavior.joints.jointGroups;
-    ensureHasSize(jointGroups, jointGroupIndex + 1ul, memRes);
-    jointGroups[jointGroupIndex].jointIndices.assign(jointIndices, jointIndices + count);
+    auto& jointGroup = getAt(dna.behavior.joints.jointGroups, jointGroupIndex);
+    jointGroup.jointIndices.assign(jointIndices, jointIndices + count);
 }
 
 template<class TWriterBase>
@@ -676,16 +733,16 @@ inline void WriterImpl<TWriterBase>::deleteMesh(std::uint16_t meshIndex) {
 
 template<class TWriterBase>
 inline void WriterImpl<TWriterBase>::setVertexPositions(std::uint16_t meshIndex, const Position* positions, std::uint32_t count) {
-    ensureHasSize(dna.geometry.meshes, meshIndex + 1ul, memRes);
-    dna.geometry.meshes[meshIndex].positions.assign(positions, positions + count);
+    auto& mesh = getAt(dna.geometry.meshes, meshIndex);
+    mesh.positions.assign(positions, positions + count);
 }
 
 template<class TWriterBase>
 inline void WriterImpl<TWriterBase>::setVertexTextureCoordinates(std::uint16_t meshIndex,
                                                                  const TextureCoordinate* textureCoordinates,
                                                                  std::uint32_t count) {
-    ensureHasSize(dna.geometry.meshes, meshIndex + 1ul, memRes);
-    auto& destination = dna.geometry.meshes[meshIndex].textureCoordinates;
+    auto& mesh = getAt(dna.geometry.meshes, meshIndex);
+    auto& destination = mesh.textureCoordinates;
     destination.clear();
     destination.us.resize_uninitialized(count);
     destination.vs.resize_uninitialized(count);
@@ -697,14 +754,14 @@ inline void WriterImpl<TWriterBase>::setVertexTextureCoordinates(std::uint16_t m
 
 template<class TWriterBase>
 inline void WriterImpl<TWriterBase>::setVertexNormals(std::uint16_t meshIndex, const Normal* normals, std::uint32_t count) {
-    ensureHasSize(dna.geometry.meshes, meshIndex + 1ul, memRes);
-    dna.geometry.meshes[meshIndex].normals.assign(normals, normals + count);
+    auto& mesh = getAt(dna.geometry.meshes, meshIndex);
+    mesh.normals.assign(normals, normals + count);
 }
 
 template<class TWriterBase>
 inline void WriterImpl<TWriterBase>::setVertexLayouts(std::uint16_t meshIndex, const VertexLayout* layouts, std::uint32_t count) {
-    ensureHasSize(dna.geometry.meshes, meshIndex + 1ul, memRes);
-    auto& destination = dna.geometry.meshes[meshIndex].layouts;
+    auto& mesh = getAt(dna.geometry.meshes, meshIndex);
+    auto& destination = mesh.layouts;
     destination.clear();
     destination.positions.resize_uninitialized(count);
     destination.textureCoordinates.resize_uninitialized(count);
@@ -728,16 +785,15 @@ inline void WriterImpl<TWriterBase>::setFaceVertexLayoutIndices(std::uint16_t me
                                                                 std::uint32_t faceIndex,
                                                                 const std::uint32_t* layoutIndices,
                                                                 std::uint32_t count) {
-    ensureHasSize(dna.geometry.meshes, meshIndex + 1ul, memRes);
-    auto& faces = dna.geometry.meshes[meshIndex].faces;
-    ensureHasSize(faces, faceIndex + 1ul, memRes);
-    faces[faceIndex].layoutIndices.assign(layoutIndices, layoutIndices + count);
+    auto& mesh = getAt(dna.geometry.meshes, meshIndex);
+    auto& face = getAt(mesh.faces, faceIndex);
+    face.layoutIndices.assign(layoutIndices, layoutIndices + count);
 }
 
 template<class TWriterBase>
 inline void WriterImpl<TWriterBase>::setMaximumInfluencePerVertex(std::uint16_t meshIndex, std::uint16_t maxInfluenceCount) {
-    ensureHasSize(dna.geometry.meshes, meshIndex + 1ul, memRes);
-    dna.geometry.meshes[meshIndex].maximumInfluencePerVertex = maxInfluenceCount;
+    auto& mesh = getAt(dna.geometry.meshes, meshIndex);
+    mesh.maximumInfluencePerVertex = maxInfluenceCount;
 }
 
 template<class TWriterBase>
@@ -752,10 +808,9 @@ inline void WriterImpl<TWriterBase>::setSkinWeightsValues(std::uint16_t meshInde
                                                           std::uint32_t vertexIndex,
                                                           const float* weights,
                                                           std::uint16_t count) {
-    ensureHasSize(dna.geometry.meshes, meshIndex + 1ul, memRes);
-    auto& skinWeights = dna.geometry.meshes[meshIndex].skinWeights;
-    ensureHasSize(skinWeights, vertexIndex + 1ul, memRes);
-    skinWeights[vertexIndex].weights.assign(weights, weights + count);
+    auto& mesh = getAt(dna.geometry.meshes, meshIndex);
+    auto& vertexSkinWeights = getAt(mesh.skinWeights, vertexIndex);
+    vertexSkinWeights.weights.assign(weights, weights + count);
 }
 
 template<class TWriterBase>
@@ -763,10 +818,9 @@ inline void WriterImpl<TWriterBase>::setSkinWeightsJointIndices(std::uint16_t me
                                                                 std::uint32_t vertexIndex,
                                                                 const std::uint16_t* jointIndices,
                                                                 std::uint16_t count) {
-    ensureHasSize(dna.geometry.meshes, meshIndex + 1ul, memRes);
-    auto& skinWeights = dna.geometry.meshes[meshIndex].skinWeights;
-    ensureHasSize(skinWeights, vertexIndex + 1ul, memRes);
-    skinWeights[vertexIndex].jointIndices.assign(jointIndices, jointIndices + count);
+    auto& mesh = getAt(dna.geometry.meshes, meshIndex);
+    auto& vertexSkinWeights = getAt(mesh.skinWeights, vertexIndex);
+    vertexSkinWeights.jointIndices.assign(jointIndices, jointIndices + count);
 }
 
 template<class TWriterBase>
@@ -780,9 +834,9 @@ template<class TWriterBase>
 inline void WriterImpl<TWriterBase>::setBlendShapeChannelIndex(std::uint16_t meshIndex,
                                                                std::uint16_t blendShapeTargetIndex,
                                                                std::uint16_t blendShapeChannelIndex) {
-    ensureHasSize(dna.geometry.meshes, meshIndex + 1ul, memRes);
-    ensureHasSize(dna.geometry.meshes[meshIndex].blendShapeTargets, blendShapeTargetIndex + 1ul, memRes);
-    dna.geometry.meshes[meshIndex].blendShapeTargets[blendShapeTargetIndex].blendShapeChannelIndex = blendShapeChannelIndex;
+    auto& mesh = getAt(dna.geometry.meshes, meshIndex);
+    auto& blendShapeTarget = getAt(mesh.blendShapeTargets, blendShapeTargetIndex);
+    blendShapeTarget.blendShapeChannelIndex = blendShapeChannelIndex;
 }
 
 template<class TWriterBase>
@@ -790,9 +844,9 @@ inline void WriterImpl<TWriterBase>::setBlendShapeTargetDeltas(std::uint16_t mes
                                                                std::uint16_t blendShapeTargetIndex,
                                                                const Delta* deltas,
                                                                std::uint32_t count) {
-    ensureHasSize(dna.geometry.meshes, meshIndex + 1ul, memRes);
-    ensureHasSize(dna.geometry.meshes[meshIndex].blendShapeTargets, blendShapeTargetIndex + 1ul, memRes);
-    dna.geometry.meshes[meshIndex].blendShapeTargets[blendShapeTargetIndex].deltas.assign(deltas, deltas + count);
+    auto& mesh = getAt(dna.geometry.meshes, meshIndex);
+    auto& blendShapeTarget = getAt(mesh.blendShapeTargets, blendShapeTargetIndex);
+    blendShapeTarget.deltas.assign(deltas, deltas + count);
 }
 
 template<class TWriterBase>
@@ -800,10 +854,9 @@ inline void WriterImpl<TWriterBase>::setBlendShapeTargetVertexIndices(std::uint1
                                                                       std::uint16_t blendShapeTargetIndex,
                                                                       const std::uint32_t* vertexIndices,
                                                                       std::uint32_t count) {
-    ensureHasSize(dna.geometry.meshes, meshIndex + 1ul, memRes);
-    auto& blendShapeTargets = dna.geometry.meshes[meshIndex].blendShapeTargets;
-    ensureHasSize(blendShapeTargets, blendShapeTargetIndex + 1ul, memRes);
-    blendShapeTargets[blendShapeTargetIndex].vertexIndices.assign(vertexIndices, vertexIndices + count);
+    auto& mesh = getAt(dna.geometry.meshes, meshIndex);
+    auto& blendShapeTarget = getAt(mesh.blendShapeTargets, blendShapeTargetIndex);
+    blendShapeTarget.vertexIndices.assign(vertexIndices, vertexIndices + count);
 }
 
 template<class TWriterBase>
@@ -813,8 +866,7 @@ inline void WriterImpl<TWriterBase>::clearMLControlNames() {
 
 template<class TWriterBase>
 inline void WriterImpl<TWriterBase>::setMLControlName(std::uint16_t index, const char* name) {
-    ensureHasSize(dna.machineLearnedBehavior.mlControlNames, index + 1ul, memRes);
-    dna.machineLearnedBehavior.mlControlNames[index] = name;
+    setAt(dna.machineLearnedBehavior.mlControlNames, index, name);
 }
 
 template<class TWriterBase>
@@ -859,9 +911,8 @@ inline void WriterImpl<TWriterBase>::clearMeshRegionNames(std::uint16_t meshInde
 
 template<class TWriterBase>
 inline void WriterImpl<TWriterBase>::setMeshRegionName(std::uint16_t meshIndex, std::uint16_t regionIndex, const char* name) {
-    ensureHasSize(dna.machineLearnedBehavior.neuralNetworkToMeshRegion.regionNames, meshIndex + 1ul, memRes);
-    ensureHasSize(dna.machineLearnedBehavior.neuralNetworkToMeshRegion.regionNames[meshIndex], regionIndex + 1ul, memRes);
-    dna.machineLearnedBehavior.neuralNetworkToMeshRegion.regionNames[meshIndex][regionIndex] = name;
+    auto& meshRegionNames = getAt(dna.machineLearnedBehavior.neuralNetworkToMeshRegion.regionNames, meshIndex);
+    setAt(meshRegionNames, regionIndex, name);
 }
 
 template<class TWriterBase>
@@ -874,9 +925,8 @@ inline void WriterImpl<TWriterBase>::setNeuralNetworkIndicesForMeshRegion(std::u
                                                                           std::uint16_t regionIndex,
                                                                           const std::uint16_t* netIndices,
                                                                           std::uint16_t count) {
-    ensureHasSize(dna.machineLearnedBehavior.neuralNetworkToMeshRegion.indices, meshIndex + 1ul, memRes);
-    ensureHasSize(dna.machineLearnedBehavior.neuralNetworkToMeshRegion.indices[meshIndex], regionIndex + 1ul, memRes);
-    auto& region = dna.machineLearnedBehavior.neuralNetworkToMeshRegion.indices[meshIndex][regionIndex];
+    auto& neuralNetworkToMeshRegionIndices = getAt(dna.machineLearnedBehavior.neuralNetworkToMeshRegion.indices, meshIndex);
+    auto& region = getAt(neuralNetworkToMeshRegionIndices, regionIndex);
     region.assign(netIndices, netIndices + count);
 }
 
@@ -892,8 +942,7 @@ template<class TWriterBase>
 inline void WriterImpl<TWriterBase>::setNeuralNetworkInputIndices(std::uint16_t netIndex,
                                                                   const std::uint16_t* inputIndices,
                                                                   std::uint16_t count) {
-    ensureHasSize(dna.machineLearnedBehavior.neuralNetworks, netIndex + 1ul, memRes);
-    auto& neuralNet = dna.machineLearnedBehavior.neuralNetworks[netIndex];
+    auto& neuralNet = getAt(dna.machineLearnedBehavior.neuralNetworks, netIndex);
     neuralNet.inputIndices.assign(inputIndices, inputIndices + count);
 }
 
@@ -901,24 +950,22 @@ template<class TWriterBase>
 inline void WriterImpl<TWriterBase>::setNeuralNetworkOutputIndices(std::uint16_t netIndex,
                                                                    const std::uint16_t* outputIndices,
                                                                    std::uint16_t count) {
-    ensureHasSize(dna.machineLearnedBehavior.neuralNetworks, netIndex + 1ul, memRes);
-    auto& neuralNet = dna.machineLearnedBehavior.neuralNetworks[netIndex];
+    auto& neuralNet = getAt(dna.machineLearnedBehavior.neuralNetworks, netIndex);
     neuralNet.outputIndices.assign(outputIndices, outputIndices + count);
 }
 
 template<class TWriterBase>
 inline void WriterImpl<TWriterBase>::clearNeuralNetworkLayers(std::uint16_t netIndex) {
-    ensureHasSize(dna.machineLearnedBehavior.neuralNetworks, netIndex + 1ul, memRes);
-    dna.machineLearnedBehavior.neuralNetworks[netIndex].layers.clear();
+    auto& neuralNet = getAt(dna.machineLearnedBehavior.neuralNetworks, netIndex);
+    neuralNet.layers.clear();
 }
 
 template<class TWriterBase>
 inline void WriterImpl<TWriterBase>::setNeuralNetworkLayerActivationFunction(std::uint16_t netIndex,
                                                                              std::uint16_t layerIndex,
                                                                              ActivationFunction function) {
-    ensureHasSize(dna.machineLearnedBehavior.neuralNetworks, netIndex + 1ul, memRes);
-    ensureHasSize(dna.machineLearnedBehavior.neuralNetworks[netIndex].layers, layerIndex + 1ul, memRes);
-    auto& layer = dna.machineLearnedBehavior.neuralNetworks[netIndex].layers[layerIndex];
+    auto& neuralNet = getAt(dna.machineLearnedBehavior.neuralNetworks, netIndex);
+    auto& layer = getAt(neuralNet.layers, layerIndex);
     layer.activationFunction.functionId = static_cast<std::uint16_t>(function);
 }
 
@@ -927,9 +974,8 @@ inline void WriterImpl<TWriterBase>::setNeuralNetworkLayerActivationFunctionPara
                                                                                        std::uint16_t layerIndex,
                                                                                        const float* activationFunctionParameters,
                                                                                        std::uint16_t count) {
-    ensureHasSize(dna.machineLearnedBehavior.neuralNetworks, netIndex + 1ul, memRes);
-    ensureHasSize(dna.machineLearnedBehavior.neuralNetworks[netIndex].layers, layerIndex + 1ul, memRes);
-    auto& layer = dna.machineLearnedBehavior.neuralNetworks[netIndex].layers[layerIndex];
+    auto& neuralNet = getAt(dna.machineLearnedBehavior.neuralNetworks, netIndex);
+    auto& layer = getAt(neuralNet.layers, layerIndex);
     layer.activationFunction.parameters.assign(activationFunctionParameters, activationFunctionParameters + count);
 }
 
@@ -938,9 +984,8 @@ inline void WriterImpl<TWriterBase>::setNeuralNetworkLayerBiases(std::uint16_t n
                                                                  std::uint16_t layerIndex,
                                                                  const float* biases,
                                                                  std::uint32_t count) {
-    ensureHasSize(dna.machineLearnedBehavior.neuralNetworks, netIndex + 1ul, memRes);
-    ensureHasSize(dna.machineLearnedBehavior.neuralNetworks[netIndex].layers, layerIndex + 1ul, memRes);
-    auto& layer = dna.machineLearnedBehavior.neuralNetworks[netIndex].layers[layerIndex];
+    auto& neuralNet = getAt(dna.machineLearnedBehavior.neuralNetworks, netIndex);
+    auto& layer = getAt(neuralNet.layers, layerIndex);
     layer.biases.assign(biases, biases + count);
 }
 
@@ -949,12 +994,245 @@ inline void WriterImpl<TWriterBase>::setNeuralNetworkLayerWeights(std::uint16_t
                                                                   std::uint16_t layerIndex,
                                                                   const float* weights,
                                                                   std::uint32_t count) {
-    ensureHasSize(dna.machineLearnedBehavior.neuralNetworks, netIndex + 1ul, memRes);
-    ensureHasSize(dna.machineLearnedBehavior.neuralNetworks[netIndex].layers, layerIndex + 1ul, memRes);
-    auto& layer = dna.machineLearnedBehavior.neuralNetworks[netIndex].layers[layerIndex];
+    auto& neuralNet = getAt(dna.machineLearnedBehavior.neuralNetworks, netIndex);
+    auto& layer = getAt(neuralNet.layers, layerIndex);
     layer.weights.assign(weights, weights + count);
 }
 
+template<class TWriterBase>
+inline void WriterImpl<TWriterBase>::clearRBFPoses() {
+    dna.rbfBehavior.poses.clear();
+}
+
+template<class TWriterBase>
+inline void WriterImpl<TWriterBase>::setRBFPoseName(std::uint16_t poseIndex, const char* name) {
+    auto& pose = getAt(dna.rbfBehavior.poses, poseIndex);
+    pose.name = name;
+}
+
+template<class TWriterBase>
+inline void WriterImpl<TWriterBase>::setRBFPoseScale(std::uint16_t poseIndex, float scale) {
+    auto& pose = getAt(dna.rbfBehavior.poses, poseIndex);
+    pose.scale = scale;
+}
+
+template<class TWriterBase>
+inline void WriterImpl<TWriterBase>::clearRBFSolvers() {
+    dna.rbfBehavior.solvers.clear();
+}
+
+template<class TWriterBase>
+inline void WriterImpl<TWriterBase>::clearRBFSolverIndices() {
+    dna.rbfBehavior.lodSolverMapping.resetIndices();
+}
+
+template<class TWriterBase>
+inline void WriterImpl<TWriterBase>::setRBFSolverIndices(std::uint16_t index,
+                                                         const std::uint16_t* solverIndices,
+                                                         std::uint16_t count) {
+    dna.rbfBehavior.lodSolverMapping.clearIndices(index);
+    dna.rbfBehavior.lodSolverMapping.addIndices(index, solverIndices, count);
+}
+
+template<class TWriterBase>
+inline void WriterImpl<TWriterBase>::clearLODRBFSolverMappings() {
+    dna.rbfBehavior.lodSolverMapping.resetLODs();
+}
+
+template<class TWriterBase>
+inline void WriterImpl<TWriterBase>::setLODRBFSolverMapping(std::uint16_t lod, std::uint16_t index) {
+    dna.rbfBehavior.lodSolverMapping.associateLODWithIndices(lod, index);
+}
+
+template<class TWriterBase>
+inline void WriterImpl<TWriterBase>::setRBFSolverName(std::uint16_t solverIndex, const char* name) {
+    auto& solver = getAt(dna.rbfBehavior.solvers, solverIndex);
+    solver.name = name;
+}
+
+template<class TWriterBase>
+inline void WriterImpl<TWriterBase>::setRBFSolverRawControlIndices(std::uint16_t solverIndex,
+                                                                   const std::uint16_t* rawControlIndices,
+                                                                   std::uint16_t count) {
+    auto& solver = getAt(dna.rbfBehavior.solvers, solverIndex);
+    solver.rawControlIndices.assign(rawControlIndices, rawControlIndices + count);
+}
+
+template<class TWriterBase>
+inline void WriterImpl<TWriterBase>::setRBFSolverPoseIndices(std::uint16_t solverIndex,
+                                                             const std::uint16_t* poseIndices,
+                                                             std::uint16_t count) {
+    auto& solver = getAt(dna.rbfBehavior.solvers, solverIndex);
+    solver.poseIndices.assign(poseIndices, poseIndices + count);
+}
+
+template<class TWriterBase>
+inline void WriterImpl<TWriterBase>::setRBFSolverRawControlValues(std::uint16_t solverIndex,
+                                                                  const float* values,
+                                                                  std::uint16_t count) {
+    auto& solver = getAt(dna.rbfBehavior.solvers, solverIndex);
+    solver.rawControlValues.assign(values, values + count);
+}
+
+template<class TWriterBase>
+inline void WriterImpl<TWriterBase>::setRBFSolverType(std::uint16_t solverIndex, RBFSolverType type) {
+    auto& solver = getAt(dna.rbfBehavior.solvers, solverIndex);
+    solver.solverType = static_cast<std::uint16_t>(type);
+}
+
+template<class TWriterBase>
+inline void WriterImpl<TWriterBase>::setRBFSolverRadius(std::uint16_t solverIndex, float radius) {
+    auto& solver = getAt(dna.rbfBehavior.solvers, solverIndex);
+    solver.radius = radius;
+}
+
+template<class TWriterBase>
+inline void WriterImpl<TWriterBase>::setRBFSolverAutomaticRadius(std::uint16_t solverIndex, AutomaticRadius automaticRadius) {
+    auto& solver = getAt(dna.rbfBehavior.solvers, solverIndex);
+    solver.automaticRadius = static_cast<std::uint16_t>(automaticRadius);
+}
+
+template<class TWriterBase>
+inline void WriterImpl<TWriterBase>::setRBFSolverWeightThreshold(std::uint16_t solverIndex, float weightThreshold) {
+    auto& solver = getAt(dna.rbfBehavior.solvers, solverIndex);
+    solver.weightThreshold = weightThreshold;
+}
+
+template<class TWriterBase>
+inline void WriterImpl<TWriterBase>::setRBFSolverDistanceMethod(std::uint16_t solverIndex, RBFDistanceMethod distanceMethod) {
+    auto& solver = getAt(dna.rbfBehavior.solvers, solverIndex);
+    solver.distanceMethod = static_cast<std::uint16_t>(distanceMethod);
+}
+
+template<class TWriterBase>
+inline void WriterImpl<TWriterBase>::setRBFSolverNormalizeMethod(std::uint16_t solverIndex, RBFNormalizeMethod normalizeMethod) {
+    auto& solver = getAt(dna.rbfBehavior.solvers, solverIndex);
+    solver.normalizeMethod = static_cast<std::uint16_t>(normalizeMethod);
+}
+
+template<class TWriterBase>
+inline void WriterImpl<TWriterBase>::setRBFSolverFunctionType(std::uint16_t solverIndex, RBFFunctionType functionType) {
+    auto& solver = getAt(dna.rbfBehavior.solvers, solverIndex);
+    solver.functionType = static_cast<std::uint16_t>(functionType);
+}
+
+template<class TWriterBase>
+inline void WriterImpl<TWriterBase>::setRBFSolverTwistAxis(std::uint16_t solverIndex, TwistAxis twistAxis) {
+    auto& solver = getAt(dna.rbfBehavior.solvers, solverIndex);
+    solver.twistAxis = static_cast<std::uint16_t>(twistAxis);
+}
+
+template<class TWriterBase>
+inline void WriterImpl<TWriterBase>::clearJointRepresentations() {
+    dna.jointBehaviorMetadata.jointRepresentations.clear();
+}
+
+template<class TWriterBase>
+inline void WriterImpl<TWriterBase>::setJointTranslationRepresentation(std::uint16_t jointIndex,
+                                                                       TranslationRepresentation representation) {
+    auto& jointRepresentation = getAt(dna.jointBehaviorMetadata.jointRepresentations, jointIndex);
+    jointRepresentation.translation = static_cast<std::uint16_t>(representation);
+}
+
+template<class TWriterBase>
+inline void WriterImpl<TWriterBase>::setJointRotationRepresentation(std::uint16_t jointIndex,
+                                                                    RotationRepresentation representation) {
+    auto& jointRepresentation = getAt(dna.jointBehaviorMetadata.jointRepresentations, jointIndex);
+    jointRepresentation.rotation = static_cast<std::uint16_t>(representation);
+}
+
+template<class TWriterBase>
+inline void WriterImpl<TWriterBase>::setJointScaleRepresentation(std::uint16_t jointIndex, ScaleRepresentation representation) {
+    auto& jointRepresentation = getAt(dna.jointBehaviorMetadata.jointRepresentations, jointIndex);
+    jointRepresentation.scale = static_cast<std::uint16_t>(representation);
+}
+
+template<class TWriterBase>
+inline void WriterImpl<TWriterBase>::clearTwists() {
+    dna.twistSwingBehavior.twists.clear();
+}
+
+template<class TWriterBase>
+inline void WriterImpl<TWriterBase>::deleteTwist(std::uint16_t twistIndex) {
+    if (twistIndex < dna.twistSwingBehavior.twists.size()) {
+        auto it = extd::advanced(dna.twistSwingBehavior.twists.begin(), twistIndex);
+        dna.twistSwingBehavior.twists.erase(it);
+    }
+}
+
+template<class TWriterBase>
+inline void WriterImpl<TWriterBase>::setTwistSetupTwistAxis(std::uint16_t twistIndex, TwistAxis twistAxis) {
+    auto& twist = getAt(dna.twistSwingBehavior.twists, twistIndex);
+    twist.twistAxis = static_cast<std::uint16_t>(twistAxis);
+}
+
+template<class TWriterBase>
+inline void WriterImpl<TWriterBase>::setTwistInputControlIndices(std::uint16_t twistIndex,
+                                                                 const std::uint16_t* controlIndices,
+                                                                 std::uint16_t controlIndexCount) {
+    auto& twist = getAt(dna.twistSwingBehavior.twists, twistIndex);
+    twist.twistInputControlIndices.assign(controlIndices, controlIndices + controlIndexCount);
+}
+
+template<class TWriterBase>
+inline void WriterImpl<TWriterBase>::setTwistOutputJointIndices(std::uint16_t twistIndex,
+                                                                const std::uint16_t* jointIndices,
+                                                                std::uint16_t jointIndexCount) {
+    auto& twist = getAt(dna.twistSwingBehavior.twists, twistIndex);
+    twist.twistOutputJointIndices.assign(jointIndices, jointIndices + jointIndexCount);
+}
+
+template<class TWriterBase>
+inline void WriterImpl<TWriterBase>::setTwistBlendWeights(std::uint16_t twistIndex,
+                                                          const float* blendWeights,
+                                                          std::uint16_t blendWeightCount) {
+    auto& twist = getAt(dna.twistSwingBehavior.twists, twistIndex);
+    twist.twistBlendWeights.assign(blendWeights, blendWeights + blendWeightCount);
+}
+
+template<class TWriterBase>
+inline void WriterImpl<TWriterBase>::clearSwings() {
+    dna.twistSwingBehavior.swings.clear();
+}
+
+template<class TWriterBase>
+inline void WriterImpl<TWriterBase>::deleteSwing(std::uint16_t swingIndex) {
+    if (swingIndex < dna.twistSwingBehavior.swings.size()) {
+        auto it = extd::advanced(dna.twistSwingBehavior.swings.begin(), swingIndex);
+        dna.twistSwingBehavior.swings.erase(it);
+    }
+}
+
+template<class TWriterBase>
+inline void WriterImpl<TWriterBase>::setSwingSetupTwistAxis(std::uint16_t swingIndex, TwistAxis twistAxis) {
+    auto& swing = getAt(dna.twistSwingBehavior.swings, swingIndex);
+    swing.twistAxis = static_cast<std::uint16_t>(twistAxis);
+}
+
+template<class TWriterBase>
+inline void WriterImpl<TWriterBase>::setSwingInputControlIndices(std::uint16_t swingIndex,
+                                                                 const std::uint16_t* controlIndices,
+                                                                 std::uint16_t controlIndexCount) {
+    auto& swing = getAt(dna.twistSwingBehavior.swings, swingIndex);
+    swing.swingInputControlIndices.assign(controlIndices, controlIndices + controlIndexCount);
+}
+
+template<class TWriterBase>
+inline void WriterImpl<TWriterBase>::setSwingOutputJointIndices(std::uint16_t swingIndex,
+                                                                const std::uint16_t* jointIndices,
+                                                                std::uint16_t jointIndexCount) {
+    auto& swing = getAt(dna.twistSwingBehavior.swings, swingIndex);
+    swing.swingOutputJointIndices.assign(jointIndices, jointIndices + jointIndexCount);
+}
+
+template<class TWriterBase>
+inline void WriterImpl<TWriterBase>::setSwingBlendWeights(std::uint16_t swingIndex,
+                                                          const float* blendWeights,
+                                                          std::uint16_t blendWeightCount) {
+    auto& swing = getAt(dna.twistSwingBehavior.swings, swingIndex);
+    swing.swingBlendWeights.assign(blendWeights, blendWeights + blendWeightCount);
+}
+
 #ifdef _MSC_VER
     #pragma warning(pop)
 #endif
