Update to new RigLogic implementing body rig support

diff --git a/Engine/Plugins/Animation/RigLogic/Source/RigLogicLib/Private/dna/filters/JointFilter.cpp b/Engine/Plugins/Animation/RigLogic/Source/RigLogicLib/Private/dna/filters/JointFilter.cpp
--- a/Engine/Plugins/Animation/RigLogic/Source/RigLogicLib/Private/dna/filters/JointFilter.cpp
+++ b/Engine/Plugins/Animation/RigLogic/Source/RigLogicLib/Private/dna/filters/JointFilter.cpp
@@ -42,7 +42,7 @@ void JointFilter::apply(RawDefinition& dest) {
         jntIdx = remappedIndices[jntIdx];
     }
     // Find root joint index
-    for (std::uint16_t jointIdx = 0u; jointIdx < dest.jointHierarchy.size(); ++jointIdx) {
+    for (std::uint16_t jointIdx = {}; jointIdx < dest.jointHierarchy.size(); ++jointIdx) {
         if (dest.jointHierarchy[jointIdx] == jointIdx) {
             rootJointIndex = jointIdx;
             break;
@@ -71,7 +71,7 @@ void JointFilter::apply(RawBehavior& dest) {
             }
         }
         // Collect row indices of removed output indices to be used for joint delta removal
-        Set<std::size_t> rowsToDelete{memRes};
+        UnorderedSet<std::size_t> rowsToDelete(memRes);
         // Remove output indices belonging to the deletable joint
         extd::filter(jointGroup.outputIndices, [this, &rowsToDelete](std::uint16_t outputIndex, std::size_t rowIndex) {
                 const auto jointIndex = static_cast<std::uint16_t>(outputIndex / jointAttributeCount);
@@ -104,7 +104,7 @@ void JointFilter::apply(RawBehavior& dest) {
             });
         // Recompute LODs
         for (auto& lod : jointGroup.lods) {
-            std::uint16_t decrementBy = 0u;
+            std::uint16_t decrementBy = {};
             for (const auto rowIndex : rowsToDelete) {
                 if (rowIndex < lod) {
                     ++decrementBy;
@@ -167,6 +167,93 @@ void JointFilter::apply(RawVertexSkinWeights& dest) {
     }
 }
 
+void JointFilter::apply(RawJointBehaviorMetadata& dest) {
+    const auto jointCount = static_cast<std::uint16_t>(dest.jointRepresentations.size());
+    std::uint16_t newCount = {};
+    for (std::uint16_t ji = {}; ji < jointCount; ++ji) {
+        if (passes(ji)) {
+            dest.jointRepresentations[remapped(ji)] = dest.jointRepresentations[ji];
+            newCount++;
+        }
+    }
+    dest.jointRepresentations.resize(newCount);
+}
+
+void JointFilter::apply(RawTwistSwingBehavior& dest) {
+    if (option != Option::All) {
+        return;
+    }
+
+    auto filterTwists = [this](RawTwistSwingBehavior& tsbh) {
+            auto& twists = tsbh.twists;
+            const auto twistCount = static_cast<std::uint16_t>(twists.size());
+            std::uint16_t di = {};
+            for (std::uint16_t ti = {}; ti < twistCount; ++ti) {
+                const auto& src = twists[ti];
+                auto& dst = twists[di];
+
+                dst.twistAxis = src.twistAxis;
+                dst.twistInputControlIndices = src.twistInputControlIndices;
+
+                const auto outputIndicesCount = static_cast<std::uint16_t>(src.twistOutputJointIndices.size());
+                dst.twistOutputJointIndices.resize(outputIndicesCount);
+                dst.twistBlendWeights.resize(outputIndicesCount);
+                std::uint16_t doi = {};
+                for (std::uint16_t soi = {}; soi < outputIndicesCount; ++soi) {
+                    const auto srcJointIndex = src.twistOutputJointIndices[soi];
+                    if (passes(srcJointIndex)) {
+                        dst.twistOutputJointIndices[doi] = remapped(srcJointIndex);
+                        dst.twistBlendWeights[doi] = src.twistBlendWeights[soi];
+                        doi++;
+                    }
+                }
+
+                if (doi != 0u) {
+                    dst.twistBlendWeights.resize(doi);
+                    dst.twistOutputJointIndices.resize(doi);
+                    di++;
+                }
+            }
+            twists.resize(di, RawTwist{memRes});
+        };
+
+    auto filterSwings = [this](RawTwistSwingBehavior& tsbh) {
+            auto& swings = tsbh.swings;
+            const auto swingCount = static_cast<std::uint16_t>(swings.size());
+            std::uint16_t di = {};
+            for (std::uint16_t si = {}; si < swingCount; ++si) {
+                const auto& src = swings[si];
+                auto& dst = swings[di];
+
+                dst.twistAxis = src.twistAxis;
+                dst.swingInputControlIndices = src.swingInputControlIndices;
+
+                const auto outputIndicesCount = static_cast<std::uint16_t>(src.swingOutputJointIndices.size());
+                dst.swingOutputJointIndices.resize(outputIndicesCount);
+                dst.swingBlendWeights.resize(outputIndicesCount);
+                std::uint16_t doi = {};
+                for (std::uint16_t soi = {}; soi < outputIndicesCount; ++soi) {
+                    const auto srcJointIndex = src.swingOutputJointIndices[soi];
+                    if (passes(srcJointIndex)) {
+                        dst.swingOutputJointIndices[doi] = remapped(srcJointIndex);
+                        dst.swingBlendWeights[doi] = src.swingBlendWeights[soi];
+                        doi++;
+                    }
+                }
+
+                if (doi != 0u) {
+                    dst.swingBlendWeights.resize(doi);
+                    dst.swingOutputJointIndices.resize(doi);
+                    di++;
+                }
+            }
+            swings.resize(di, RawSwing{memRes});
+        };
+
+    filterTwists(dest);
+    filterSwings(dest);
+}
+
 bool JointFilter::passes(std::uint16_t index) const {
     return extd::contains(passingIndices, index);
 }
@@ -175,12 +262,8 @@ std::uint16_t JointFilter::remapped(std::uint16_t oldIndex) const {
     return remappedIndices.at(oldIndex);
 }
 
-std::uint16_t JointFilter::maxRemappedIndex() const {
-    using IndexPair = std::pair<std::uint16_t, std::uint16_t>;
-    const auto compare = [](const IndexPair& lhs, const IndexPair& rhs) {
-            return lhs.second < rhs.second;
-        };
-    return (remappedIndices.empty() ? static_cast<std::uint16_t>(0) : extd::maxOf(remappedIndices, compare).second);
+std::uint16_t JointFilter::filteredJointCount() const {
+    return static_cast<std::uint16_t>(passingIndices.size());
 }
 
 }  // namespace dna
diff --git a/Engine/Plugins/Animation/RigLogic/Source/RigLogicLib/Private/dna/filters/JointFilter.h b/Engine/Plugins/Animation/RigLogic/Source/RigLogicLib/Private/dna/filters/JointFilter.h
--- a/Engine/Plugins/Animation/RigLogic/Source/RigLogicLib/Private/dna/filters/JointFilter.h
+++ b/Engine/Plugins/Animation/RigLogic/Source/RigLogicLib/Private/dna/filters/JointFilter.h
@@ -10,7 +10,9 @@ namespace dna {
 
 struct RawBehavior;
 struct RawDefinition;
+struct RawJointBehaviorMetadata;
 struct RawLODMapping;
+struct RawTwistSwingBehavior;
 struct RawVertexSkinWeights;
 
 class JointFilter {
@@ -26,9 +28,11 @@ class JointFilter {
         void apply(RawDefinition& dest);
         void apply(RawBehavior& dest);
         void apply(RawVertexSkinWeights& dest);
+        void apply(RawJointBehaviorMetadata& dest);
+        void apply(RawTwistSwingBehavior& dest);
         bool passes(std::uint16_t index) const;
         std::uint16_t remapped(std::uint16_t oldIndex) const;
-        std::uint16_t maxRemappedIndex() const;
+        std::uint16_t filteredJointCount() const;
 
     private:
         MemoryResource* memRes;
diff --git a/Engine/Plugins/Animation/RigLogic/Source/RigLogicLib/Private/dna/stream/BinaryStreamReaderImpl.cpp b/Engine/Plugins/Animation/RigLogic/Source/RigLogicLib/Private/dna/stream/BinaryStreamReaderImpl.cpp
--- a/Engine/Plugins/Animation/RigLogic/Source/RigLogicLib/Private/dna/stream/BinaryStreamReaderImpl.cpp
+++ b/Engine/Plugins/Animation/RigLogic/Source/RigLogicLib/Private/dna/stream/BinaryStreamReaderImpl.cpp
@@ -92,13 +92,23 @@ void BinaryStreamReaderImpl::unload(DataLayer layer) {
     if ((layer == DataLayer::All) ||
         (layer == DataLayer::Descriptor)) {
         dna = DNA{dna.layers.unknownPolicy, dna.layers.upgradePolicy, memRes};
+    } else if (layer == DataLayer::TwistSwingBehavior) {
+        dna.unloadTwistSwingBehavior();
+    } else if (layer == DataLayer::RBFBehavior) {
+        dna.unloadRBFBehavior();
+    } else if (layer == DataLayer::JointBehaviorMetadata) {
+        dna.unloadJointBehaviorMetadata();
     } else if (layer == DataLayer::MachineLearnedBehavior) {
         dna.unloadMachineLearnedBehavior();
     } else if ((layer == DataLayer::Geometry) || (layer == DataLayer::GeometryWithoutBlendShapes)) {
         dna.unloadGeometry();
     } else if (layer == DataLayer::Behavior) {
+        dna.unloadRBFBehavior();
         dna.unloadBehavior();
     } else if (layer == DataLayer::Definition) {
+        dna.unloadJointBehaviorMetadata();
+        dna.unloadTwistSwingBehavior();
+        dna.unloadRBFBehavior();
         dna.unloadMachineLearnedBehavior();
         dna.unloadGeometry();
         dna.unloadBehavior();
diff --git a/Engine/Plugins/Animation/RigLogic/Source/RigLogicLib/Private/dna/stream/FilteredBinaryInputArchive.cpp b/Engine/Plugins/Animation/RigLogic/Source/RigLogicLib/Private/dna/stream/FilteredBinaryInputArchive.cpp
--- a/Engine/Plugins/Animation/RigLogic/Source/RigLogicLib/Private/dna/stream/FilteredBinaryInputArchive.cpp
+++ b/Engine/Plugins/Animation/RigLogic/Source/RigLogicLib/Private/dna/stream/FilteredBinaryInputArchive.cpp
@@ -4,6 +4,7 @@
 
 #include "dna/DNA.h"
 #include "dna/TypeDefs.h"
+#include "dna/filters/JointFilter.h"
 #include "dna/filters/Remap.h"
 #include "dna/types/Limits.h"
 #include "dna/utils/Extd.h"
@@ -212,8 +213,7 @@ void FilteredBinaryInputArchive::process(RawJoints& dest) {
 
         dest.jointGroups.push_back(std::move(jointGroup));
     }
-    const auto uncompressedJointCount = static_cast<std::uint16_t>(JointFilter::maxRemappedIndex() + 1u);
-    dest.rowCount = static_cast<std::uint16_t>(uncompressedJointCount * jointAttributeCount);
+    dest.rowCount = static_cast<std::uint16_t>(JointFilter::filteredJointCount() * jointAttributeCount);
 }
 
 void FilteredBinaryInputArchive::process(RawBlendShapeChannels& dest) {
@@ -450,8 +450,8 @@ void FilteredBinaryInputArchive::removeUnreferencedBlendShapes(DNA& dest) {
 
     // Remove channel from definition
     dest.definition.lodBlendShapeMapping.filterIndices([&unreferencedChannels](std::uint16_t index) {
-        return !extd::contains(unreferencedChannels, index);
-    });
+            return !extd::contains(unreferencedChannels, index);
+        });
 
     BlendShapeFilter::configure(static_cast<std::uint16_t>(dest.definition.blendShapeChannelNames.size()),
                                 dest.definition.lodBlendShapeMapping.getCombinedDistinctIndices(memRes));
@@ -468,4 +468,84 @@ void FilteredBinaryInputArchive::removeUnreferencedBlendShapes(DNA& dest) {
     }
 }
 
+void FilteredBinaryInputArchive::process(RawRBFBehavior& dest) {
+    if (malformed) {
+        return;
+    }
+
+    if (!contains(layerBitmask, DataLayerBitmask::RBFBehavior)) {
+        return;
+    }
+
+    process(dest.lodSolverMapping);
+    if (!lodConstraint.hasImpactOn(unconstrainedLODCount)) {
+        process(dest.solvers);
+        process(dest.poses);
+        return;
+    }
+
+    // Perform filtered load only if a different maxLOD is set
+    const auto solverCount = static_cast<std::uint16_t>(processSize());
+    dest.lodSolverMapping.discardLODs(lodConstraint);
+
+    const auto passingIndices = dest.lodSolverMapping.getCombinedDistinctIndices(memRes);
+    UnorderedMap<std::uint16_t, std::uint16_t> remappedIndices{memRes};
+    remap(solverCount, passingIndices, remappedIndices);
+    dest.lodSolverMapping.mapIndices([&remappedIndices](std::uint16_t value) {
+            return remappedIndices.at(value);
+        });
+
+    dest.solvers.reserve(passingIndices.size());
+    for (std::uint16_t i = {}; i < solverCount; ++i) {
+        // Check if the solver indices filtered for the current maxLOD permit loading
+        // this solver
+        if (extd::contains(passingIndices, i)) {
+            RawRBFSolver solver{memRes};
+            process(solver);
+            dest.solvers.push_back(std::move(solver));
+        } else {
+            // Jump over solver
+            decltype(RawRBFSolver::size) solverSize{};
+            decltype(RawRBFSolver::baseMarker) solverBase{};
+            decltype(RawRBFSolver::sizeMarker) solverSizeMarker{solverSize, solverBase};
+            process(solverSize);
+            process(solverBase);
+            process(solverSizeMarker);
+        }
+    }
+    process(dest.poses);
+}
+
+void FilteredBinaryInputArchive::process(RawJointBehaviorMetadata& dest) {
+    if (malformed) {
+        return;
+    }
+    if (!contains(layerBitmask, DataLayerBitmask::JointBehaviorMetadata)) {
+        return;
+    }
+
+    BaseArchive::process(dest);
+    // No filtering is done, unless LOD constraint may have some effect
+    if (!lodConstraint.hasImpactOn(unconstrainedLODCount)) {
+        return;
+    }
+    JointFilter::apply(dest);
+}
+
+void FilteredBinaryInputArchive::process(RawTwistSwingBehavior& dest) {
+    if (malformed) {
+        return;
+    }
+    if (!contains(layerBitmask, DataLayerBitmask::TwistSwingBehavior)) {
+        return;
+    }
+
+    BaseArchive::process(dest);
+    // No filtering is done, unless LOD constraint may have some effect
+    if (!lodConstraint.hasImpactOn(unconstrainedLODCount)) {
+        return;
+    }
+    JointFilter::apply(dest);
+}
+
 }  // namespace dna
diff --git a/Engine/Plugins/Animation/RigLogic/Source/RigLogicLib/Private/dna/stream/FilteredBinaryInputArchive.h b/Engine/Plugins/Animation/RigLogic/Source/RigLogicLib/Private/dna/stream/FilteredBinaryInputArchive.h
--- a/Engine/Plugins/Animation/RigLogic/Source/RigLogicLib/Private/dna/stream/FilteredBinaryInputArchive.h
+++ b/Engine/Plugins/Animation/RigLogic/Source/RigLogicLib/Private/dna/stream/FilteredBinaryInputArchive.h
@@ -6,12 +6,10 @@
 #include "dna/DataLayerBitmask.h"
 #include "dna/DNA.h"
 #include "dna/LODConstraint.h"
-#include "dna/TypeDefs.h"
 #include "dna/filters/AnimatedMapFilter.h"
 #include "dna/filters/BlendShapeFilter.h"
 #include "dna/filters/JointFilter.h"
 #include "dna/filters/MeshFilter.h"
-
 #include <terse/archives/binary/InputArchive.h>
 
 #include <cstddef>
@@ -83,6 +81,9 @@ class FilteredBinaryInputArchive final : public AnimatedMapFilter, public BlendS
         void process(Vector<RawBlendShapeTarget>& dest);
         void process(RawVertexSkinWeights& dest);
         void process(RawMachineLearnedBehavior& dest);
+        void process(RawRBFBehavior& dest);
+        void process(RawJointBehaviorMetadata& dest);
+        void process(RawTwistSwingBehavior& dest);
 
         void process(DNA& dest);
         void removeUnreferencedBlendShapes(DNA& dest);
diff --git a/Engine/Plugins/Animation/RigLogic/Source/RigLogicLib/Private/dna/stream/JSONStreamReaderImpl.cpp b/Engine/Plugins/Animation/RigLogic/Source/RigLogicLib/Private/dna/stream/JSONStreamReaderImpl.cpp
--- a/Engine/Plugins/Animation/RigLogic/Source/RigLogicLib/Private/dna/stream/JSONStreamReaderImpl.cpp
+++ b/Engine/Plugins/Animation/RigLogic/Source/RigLogicLib/Private/dna/stream/JSONStreamReaderImpl.cpp
@@ -48,13 +48,23 @@ void JSONStreamReaderImpl::unload(DataLayer layer) {
     if ((layer == DataLayer::All) ||
         (layer == DataLayer::Descriptor)) {
         dna = DNA{dna.layers.unknownPolicy, dna.layers.upgradePolicy, memRes};
+    } else if (layer == DataLayer::TwistSwingBehavior) {
+        dna.unloadTwistSwingBehavior();
+    } else if (layer == DataLayer::RBFBehavior) {
+        dna.unloadRBFBehavior();
+    } else if (layer == DataLayer::JointBehaviorMetadata) {
+        dna.unloadJointBehaviorMetadata();
     } else if (layer == DataLayer::MachineLearnedBehavior) {
         dna.unloadMachineLearnedBehavior();
     } else if ((layer == DataLayer::Geometry) || (layer == DataLayer::GeometryWithoutBlendShapes)) {
         dna.unloadGeometry();
     } else if (layer == DataLayer::Behavior) {
+        dna.unloadRBFBehavior();
         dna.unloadBehavior();
     } else if (layer == DataLayer::Definition) {
+        dna.unloadJointBehaviorMetadata();
+        dna.unloadTwistSwingBehavior();
+        dna.unloadRBFBehavior();
         dna.unloadMachineLearnedBehavior();
         dna.unloadGeometry();
         dna.unloadBehavior();
diff --git a/Engine/Plugins/Animation/RigLogic/Source/RigLogicLib/Private/dna/utils/Extd.h b/Engine/Plugins/Animation/RigLogic/Source/RigLogicLib/Private/dna/utils/Extd.h
--- a/Engine/Plugins/Animation/RigLogic/Source/RigLogicLib/Private/dna/utils/Extd.h
+++ b/Engine/Plugins/Animation/RigLogic/Source/RigLogicLib/Private/dna/utils/Extd.h
@@ -28,7 +28,7 @@ inline T clamp(T value, T low, T high) {
 
 template<typename T>
 inline T roundUp(T number, T multiple) {
-    return ((number + multiple - 1) / multiple) * multiple;
+    return static_cast<T>(((number + multiple - 1) / multiple) * multiple);
 }
 
 template<typename T>
@@ -108,6 +108,22 @@ inline impl::LUTFilter<TLookUpTable, impl::LUTStrategy::ByPosition> byPosition(c
     return impl::LUTFilter<TLookUpTable, impl::LUTStrategy::ByPosition>{lookUpTable};
 }
 
+template<class TContainer, class TComparator>
+inline typename TContainer::value_type minOf(const TContainer& container, TComparator&& comparator) {
+    using ValueType = typename TContainer::value_type;
+    const auto it = std::min_element(container.begin(), container.end(), comparator);
+    return (it == container.end() ? ValueType{} : *it);
+}
+
+template<class TContainer>
+inline typename TContainer::value_type minOf(const TContainer& container) {
+    using ValueType = typename TContainer::value_type;
+    const auto compare = [](const ValueType& lhs, const ValueType& rhs) {
+            return lhs < rhs;
+        };
+    return minOf(container, compare);
+}
+
 template<class TContainer, class TComparator>
 inline typename TContainer::value_type maxOf(const TContainer& container, TComparator&& comparator) {
     using ValueType = typename TContainer::value_type;
diff --git a/Engine/Plugins/Animation/RigLogic/Source/RigLogicLib/Private/riglogic/animatedmaps/AnimatedMapsFactory.cpp b/Engine/Plugins/Animation/RigLogic/Source/RigLogicLib/Private/riglogic/animatedmaps/AnimatedMapsFactory.cpp
--- a/Engine/Plugins/Animation/RigLogic/Source/RigLogicLib/Private/riglogic/animatedmaps/AnimatedMapsFactory.cpp
+++ b/Engine/Plugins/Animation/RigLogic/Source/RigLogicLib/Private/riglogic/animatedmaps/AnimatedMapsFactory.cpp
@@ -28,7 +28,7 @@ static AnimatedMapsOutputInstance::Factory createAnimatedMapsOutputInstanceFacto
 
 AnimatedMaps::Pointer AnimatedMapsFactory::create(const Configuration& config, const RigMetrics& metrics,
                                                   MemoryResource* memRes) {
-    if (!config.loadAnimatedMaps) {
+    if (!config.loadAnimatedMaps || (metrics.animatedMapCount == 0u)) {
         return UniqueInstance<AnimatedMapsNull, AnimatedMaps>::with(memRes).create();
     }
     auto instanceFactory = createAnimatedMapsOutputInstanceFactory(config, metrics.animatedMapCount);
@@ -39,7 +39,7 @@ AnimatedMaps::Pointer AnimatedMapsFactory::create(const Configuration& config, c
 AnimatedMaps::Pointer AnimatedMapsFactory::create(const Configuration& config,
                                                   const dna::BehaviorReader* reader,
                                                   MemoryResource* memRes) {
-    if (!config.loadAnimatedMaps) {
+    if (!config.loadAnimatedMaps || (reader->getAnimatedMapCount() == 0u)) {
         return UniqueInstance<AnimatedMapsNull, AnimatedMaps>::with(memRes).create();
     }
 
diff --git a/Engine/Plugins/Animation/RigLogic/Source/RigLogicLib/Private/riglogic/animatedmaps/AnimatedMapsImplOutputInstance.cpp b/Engine/Plugins/Animation/RigLogic/Source/RigLogicLib/Private/riglogic/animatedmaps/AnimatedMapsImplOutputInstance.cpp
--- a/Engine/Plugins/Animation/RigLogic/Source/RigLogicLib/Private/riglogic/animatedmaps/AnimatedMapsImplOutputInstance.cpp
+++ b/Engine/Plugins/Animation/RigLogic/Source/RigLogicLib/Private/riglogic/animatedmaps/AnimatedMapsImplOutputInstance.cpp
@@ -4,7 +4,15 @@
 
 #include "riglogic/TypeDefs.h"
 
+#ifdef _MSC_VER
+    #pragma warning(push)
+    #pragma warning(disable : 4365 4987)
+#endif
+#include <algorithm>
 #include <cstdint>
+#ifdef _MSC_VER
+    #pragma warning(pop)
+#endif
 
 namespace rl4 {
 
@@ -16,4 +24,8 @@ ArrayView<float> AnimatedMapsImplOutputInstance::getOutputBuffer() {
     return ArrayView<float>{outputBuffer};
 }
 
+void AnimatedMapsImplOutputInstance::resetOutputBuffer() {
+    std::fill(outputBuffer.begin(), outputBuffer.end(), 0.0f);
+}
+
 }  // namespace rl4
diff --git a/Engine/Plugins/Animation/RigLogic/Source/RigLogicLib/Private/riglogic/animatedmaps/AnimatedMapsImplOutputInstance.h b/Engine/Plugins/Animation/RigLogic/Source/RigLogicLib/Private/riglogic/animatedmaps/AnimatedMapsImplOutputInstance.h
--- a/Engine/Plugins/Animation/RigLogic/Source/RigLogicLib/Private/riglogic/animatedmaps/AnimatedMapsImplOutputInstance.h
+++ b/Engine/Plugins/Animation/RigLogic/Source/RigLogicLib/Private/riglogic/animatedmaps/AnimatedMapsImplOutputInstance.h
@@ -13,6 +13,7 @@ class AnimatedMapsImplOutputInstance : public AnimatedMapsOutputInstance {
     public:
         AnimatedMapsImplOutputInstance(std::uint16_t animatedMapCount, MemoryResource* memRes);
         ArrayView<float> getOutputBuffer() override;
+        void resetOutputBuffer() override;
 
     private:
         Vector<float> outputBuffer;
diff --git a/Engine/Plugins/Animation/RigLogic/Source/RigLogicLib/Private/riglogic/animatedmaps/AnimatedMapsNullOutputInstance.cpp b/Engine/Plugins/Animation/RigLogic/Source/RigLogicLib/Private/riglogic/animatedmaps/AnimatedMapsNullOutputInstance.cpp
--- a/Engine/Plugins/Animation/RigLogic/Source/RigLogicLib/Private/riglogic/animatedmaps/AnimatedMapsNullOutputInstance.cpp
+++ b/Engine/Plugins/Animation/RigLogic/Source/RigLogicLib/Private/riglogic/animatedmaps/AnimatedMapsNullOutputInstance.cpp
@@ -10,4 +10,7 @@ ArrayView<float> AnimatedMapsNullOutputInstance::getOutputBuffer() {
     return {};
 }
 
+void AnimatedMapsNullOutputInstance::resetOutputBuffer() {
+}
+
 }  // namespace rl4
diff --git a/Engine/Plugins/Animation/RigLogic/Source/RigLogicLib/Private/riglogic/animatedmaps/AnimatedMapsNullOutputInstance.h b/Engine/Plugins/Animation/RigLogic/Source/RigLogicLib/Private/riglogic/animatedmaps/AnimatedMapsNullOutputInstance.h
--- a/Engine/Plugins/Animation/RigLogic/Source/RigLogicLib/Private/riglogic/animatedmaps/AnimatedMapsNullOutputInstance.h
+++ b/Engine/Plugins/Animation/RigLogic/Source/RigLogicLib/Private/riglogic/animatedmaps/AnimatedMapsNullOutputInstance.h
@@ -10,6 +10,7 @@ namespace rl4 {
 class AnimatedMapsNullOutputInstance : public AnimatedMapsOutputInstance {
     public:
         ArrayView<float> getOutputBuffer() override;
+        void resetOutputBuffer() override;
 
 };
 
diff --git a/Engine/Plugins/Animation/RigLogic/Source/RigLogicLib/Private/riglogic/animatedmaps/AnimatedMapsOutputInstance.h b/Engine/Plugins/Animation/RigLogic/Source/RigLogicLib/Private/riglogic/animatedmaps/AnimatedMapsOutputInstance.h
--- a/Engine/Plugins/Animation/RigLogic/Source/RigLogicLib/Private/riglogic/animatedmaps/AnimatedMapsOutputInstance.h
+++ b/Engine/Plugins/Animation/RigLogic/Source/RigLogicLib/Private/riglogic/animatedmaps/AnimatedMapsOutputInstance.h
@@ -23,6 +23,7 @@ class AnimatedMapsOutputInstance {
     public:
         virtual ~AnimatedMapsOutputInstance();
         virtual ArrayView<float> getOutputBuffer() = 0;
+        virtual void resetOutputBuffer() = 0;
 
 };
 
diff --git a/Engine/Plugins/Animation/RigLogic/Source/RigLogicLib/Private/riglogic/blendshapes/BlendShapesFactory.cpp b/Engine/Plugins/Animation/RigLogic/Source/RigLogicLib/Private/riglogic/blendshapes/BlendShapesFactory.cpp
--- a/Engine/Plugins/Animation/RigLogic/Source/RigLogicLib/Private/riglogic/blendshapes/BlendShapesFactory.cpp
+++ b/Engine/Plugins/Animation/RigLogic/Source/RigLogicLib/Private/riglogic/blendshapes/BlendShapesFactory.cpp
@@ -33,7 +33,7 @@ static BlendShapesOutputInstance::Factory createBlendShapesOutputInstanceFactory
 }
 
 BlendShapes::Pointer BlendShapesFactory::create(const Configuration& config, const RigMetrics& metrics, MemoryResource* memRes) {
-    if (!config.loadBlendShapes) {
+    if (!config.loadBlendShapes || (metrics.blendShapeCount == 0u)) {
         return UniqueInstance<BlendShapesNull, BlendShapes>::with(memRes).create();
     }
     auto instanceFactory = createBlendShapesOutputInstanceFactory(config, metrics.blendShapeCount);
@@ -47,7 +47,7 @@ BlendShapes::Pointer BlendShapesFactory::create(const Configuration& config, con
 BlendShapes::Pointer BlendShapesFactory::create(const Configuration& config,
                                                 const dna::BehaviorReader* reader,
                                                 MemoryResource* memRes) {
-    if (!config.loadBlendShapes) {
+    if (!config.loadBlendShapes || (reader->getBlendShapeChannelCount() == 0u)) {
         return UniqueInstance<BlendShapesNull, BlendShapes>::with(memRes).create();
     }
 
diff --git a/Engine/Plugins/Animation/RigLogic/Source/RigLogicLib/Private/riglogic/blendshapes/BlendShapesImplOutputInstance.cpp b/Engine/Plugins/Animation/RigLogic/Source/RigLogicLib/Private/riglogic/blendshapes/BlendShapesImplOutputInstance.cpp
--- a/Engine/Plugins/Animation/RigLogic/Source/RigLogicLib/Private/riglogic/blendshapes/BlendShapesImplOutputInstance.cpp
+++ b/Engine/Plugins/Animation/RigLogic/Source/RigLogicLib/Private/riglogic/blendshapes/BlendShapesImplOutputInstance.cpp
@@ -2,7 +2,15 @@
 
 #include "riglogic/blendshapes/BlendShapesImplOutputInstance.h"
 
+#ifdef _MSC_VER
+    #pragma warning(push)
+    #pragma warning(disable : 4365 4987)
+#endif
+#include <algorithm>
 #include <cstdint>
+#ifdef _MSC_VER
+    #pragma warning(pop)
+#endif
 
 namespace rl4 {
 
@@ -14,4 +22,8 @@ ArrayView<float> BlendShapesImplOutputInstance::getOutputBuffer() {
     return ArrayView<float>{outputBuffer};
 }
 
+void BlendShapesImplOutputInstance::resetOutputBuffer() {
+    std::fill(outputBuffer.begin(), outputBuffer.end(), 0.0f);
+}
+
 }  // namespace rl4
diff --git a/Engine/Plugins/Animation/RigLogic/Source/RigLogicLib/Private/riglogic/blendshapes/BlendShapesNullOutputInstance.cpp b/Engine/Plugins/Animation/RigLogic/Source/RigLogicLib/Private/riglogic/blendshapes/BlendShapesNullOutputInstance.cpp
--- a/Engine/Plugins/Animation/RigLogic/Source/RigLogicLib/Private/riglogic/blendshapes/BlendShapesNullOutputInstance.cpp
+++ b/Engine/Plugins/Animation/RigLogic/Source/RigLogicLib/Private/riglogic/blendshapes/BlendShapesNullOutputInstance.cpp
@@ -10,4 +10,7 @@ ArrayView<float> BlendShapesNullOutputInstance::getOutputBuffer() {
     return {};
 }
 
+void BlendShapesNullOutputInstance::resetOutputBuffer() {
+}
+
 }  // namespace rl4
diff --git a/Engine/Plugins/Animation/RigLogic/Source/RigLogicLib/Private/riglogic/blendshapes/BlendShapesOutputInstance.h b/Engine/Plugins/Animation/RigLogic/Source/RigLogicLib/Private/riglogic/blendshapes/BlendShapesOutputInstance.h
--- a/Engine/Plugins/Animation/RigLogic/Source/RigLogicLib/Private/riglogic/blendshapes/BlendShapesOutputInstance.h
+++ b/Engine/Plugins/Animation/RigLogic/Source/RigLogicLib/Private/riglogic/blendshapes/BlendShapesOutputInstance.h
@@ -24,6 +24,7 @@ class BlendShapesOutputInstance {
     public:
         virtual ~BlendShapesOutputInstance();
         virtual ArrayView<float> getOutputBuffer() = 0;
+        virtual void resetOutputBuffer() = 0;
 
 };
 
diff --git a/Engine/Plugins/Animation/RigLogic/Source/RigLogicLib/Private/riglogic/controls/ControlsFactory.cpp b/Engine/Plugins/Animation/RigLogic/Source/RigLogicLib/Private/riglogic/controls/ControlsFactory.cpp
--- a/Engine/Plugins/Animation/RigLogic/Source/RigLogicLib/Private/riglogic/controls/ControlsFactory.cpp
+++ b/Engine/Plugins/Animation/RigLogic/Source/RigLogicLib/Private/riglogic/controls/ControlsFactory.cpp
@@ -19,10 +19,11 @@ static ControlsInputInstance::Factory createInstanceFactory(const Configuration&
                                                             std::uint16_t guiControlCount,
                                                             std::uint16_t rawControlCount,
                                                             std::uint16_t psdControlCount,
-                                                            std::uint16_t mlControlCount) {
+                                                            std::uint16_t mlControlCount,
+                                                            std::uint16_t rbfControlCount) {
     return [ = ](MemoryResource* memRes) {
                auto factory = UniqueInstance<StandardControlsInputInstance, ControlsInputInstance>::with(memRes);
-               return factory.create(guiControlCount, rawControlCount, psdControlCount, mlControlCount, memRes);
+               return factory.create(guiControlCount, rawControlCount, psdControlCount, mlControlCount, rbfControlCount, memRes);
     };
 }
 
@@ -31,7 +32,8 @@ Controls::Pointer ControlsFactory::create(const Configuration& config, const Rig
                                                  metrics.guiControlCount,
                                                  metrics.rawControlCount,
                                                  metrics.psdControlCount,
-                                                 metrics.mlControlCount);
+                                                 metrics.mlControlCount,
+                                                 metrics.rbfControlCount);
     return UniqueInstance<Controls>::with(memRes).create(ConditionalTable{memRes}, PSDMatrix{memRes}, instanceFactory);
 }
 
@@ -82,7 +84,8 @@ Controls::Pointer ControlsFactory::create(const Configuration& config, const dna
                                                  conditionals.getInputCount(),
                                                  conditionals.getOutputCount(),
                                                  psds.getDistinctPSDCount(),
-                                                 reader->getMLControlCount());
+                                                 reader->getMLControlCount(),
+                                                 reader->getRBFPoseCount());
 
     return UniqueInstance<Controls>::with(memRes).create(std::move(conditionals), std::move(psds), instanceFactory);
 }
diff --git a/Engine/Plugins/Animation/RigLogic/Source/RigLogicLib/Private/riglogic/controls/instances/StandardControlsInputInstance.cpp b/Engine/Plugins/Animation/RigLogic/Source/RigLogicLib/Private/riglogic/controls/instances/StandardControlsInputInstance.cpp
--- a/Engine/Plugins/Animation/RigLogic/Source/RigLogicLib/Private/riglogic/controls/instances/StandardControlsInputInstance.cpp
+++ b/Engine/Plugins/Animation/RigLogic/Source/RigLogicLib/Private/riglogic/controls/instances/StandardControlsInputInstance.cpp
@@ -11,9 +11,10 @@ StandardControlsInputInstance::StandardControlsInputInstance(std::uint16_t guiCo
                                                              std::uint16_t rawControlCount,
                                                              std::uint16_t psdControlCount,
                                                              std::uint16_t mlControlCount,
+                                                             std::uint16_t rbfControlCount,
                                                              MemoryResource* memRes) :
     guiControlBuffer{guiControlCount, {}, memRes},
-    inputBuffer{static_cast<std::size_t>(rawControlCount + psdControlCount + mlControlCount), {}, memRes} {
+    inputBuffer{static_cast<std::size_t>(rawControlCount + psdControlCount + mlControlCount + rbfControlCount), {}, memRes} {
 }
 
 ArrayView<float> StandardControlsInputInstance::getGUIControlBuffer() {
diff --git a/Engine/Plugins/Animation/RigLogic/Source/RigLogicLib/Private/riglogic/controls/instances/StandardControlsInputInstance.h b/Engine/Plugins/Animation/RigLogic/Source/RigLogicLib/Private/riglogic/controls/instances/StandardControlsInputInstance.h
--- a/Engine/Plugins/Animation/RigLogic/Source/RigLogicLib/Private/riglogic/controls/instances/StandardControlsInputInstance.h
+++ b/Engine/Plugins/Animation/RigLogic/Source/RigLogicLib/Private/riglogic/controls/instances/StandardControlsInputInstance.h
@@ -15,6 +15,7 @@ class StandardControlsInputInstance : public ControlsInputInstance {
                                       std::uint16_t rawControlCount,
                                       std::uint16_t psdControlCount,
                                       std::uint16_t mlControlCount,
+                                      std::uint16_t rbfControlCount,
                                       MemoryResource* memRes);
 
         ArrayView<float> getGUIControlBuffer() override;
diff --git a/Engine/Plugins/Animation/RigLogic/Source/RigLogicLib/Private/riglogic/joints/Joints.cpp b/Engine/Plugins/Animation/RigLogic/Source/RigLogicLib/Private/riglogic/joints/Joints.cpp
--- a/Engine/Plugins/Animation/RigLogic/Source/RigLogicLib/Private/riglogic/joints/Joints.cpp
+++ b/Engine/Plugins/Animation/RigLogic/Source/RigLogicLib/Private/riglogic/joints/Joints.cpp
@@ -4,7 +4,6 @@
 
 #include "riglogic/joints/JointsEvaluator.h"
 #include "riglogic/riglogic/RigInstance.h"
-#include "riglogic/transformation/Transformation.h"
 #include "riglogic/types/Aliases.h"
 
 #include <cstdint>
@@ -43,14 +42,10 @@ void Joints::calculate(const ControlsInputInstance* inputs,
     evaluator->calculate(inputs, outputs, lod, jointGroupIndex);
 }
 
-ConstArrayView<float> Joints::getRawNeutralValues() const {
+ConstArrayView<float> Joints::getNeutralValues() const {
     return ConstArrayView<float>{neutralValues};
 }
 
-TransformationArrayView Joints::getNeutralValues() const {
-    return TransformationArrayView{neutralValues.data(), neutralValues.size()};
-}
-
 std::uint16_t Joints::getJointGroupCount() const {
     return jointGroupCount;
 }
diff --git a/Engine/Plugins/Animation/RigLogic/Source/RigLogicLib/Private/riglogic/joints/Joints.h b/Engine/Plugins/Animation/RigLogic/Source/RigLogicLib/Private/riglogic/joints/Joints.h
--- a/Engine/Plugins/Animation/RigLogic/Source/RigLogicLib/Private/riglogic/joints/Joints.h
+++ b/Engine/Plugins/Animation/RigLogic/Source/RigLogicLib/Private/riglogic/joints/Joints.h
@@ -5,7 +5,6 @@
 #include "riglogic/TypeDefs.h"
 #include "riglogic/joints/JointsEvaluator.h"
 #include "riglogic/joints/JointsOutputInstance.h"
-#include "riglogic/transformation/Transformation.h"
 
 #include <cstdint>
 
@@ -45,8 +44,7 @@ class Joints {
         }
 
         std::uint16_t getJointGroupCount() const;
-        ConstArrayView<float> getRawNeutralValues() const;
-        TransformationArrayView getNeutralValues() const;
+        ConstArrayView<float> getNeutralValues() const;
         ConstArrayView<std::uint16_t> getVariableAttributeIndices(std::uint16_t lod) const;
 
     private:
diff --git a/Engine/Plugins/Animation/RigLogic/Source/RigLogicLib/Private/riglogic/joints/JointsBuilder.cpp b/Engine/Plugins/Animation/RigLogic/Source/RigLogicLib/Private/riglogic/joints/JointsBuilder.cpp
--- a/Engine/Plugins/Animation/RigLogic/Source/RigLogicLib/Private/riglogic/joints/JointsBuilder.cpp
+++ b/Engine/Plugins/Animation/RigLogic/Source/RigLogicLib/Private/riglogic/joints/JointsBuilder.cpp
@@ -3,56 +3,14 @@
 #include "riglogic/joints/JointsBuilder.h"
 
 #include "riglogic/TypeDefs.h"
-#include "riglogic/joints/bpcm/Builder.h"
-#include "riglogic/system/simd/Detect.h"
+#include "riglogic/joints/cpu/CPUJointsBuilder.h"
 
 namespace rl4 {
 
 JointsBuilder::~JointsBuilder() = default;
 
-JointsBuilder::Pointer JointsBuilder::create(Configuration config, MemoryResource* memRes) {
-    #if defined(__clang__)
-        #pragma clang diagnostic push
-        #pragma clang diagnostic ignored "-Wunused-local-typedef"
-    #endif
-    #if !defined(__clang__) && defined(__GNUC__)
-        #pragma GCC diagnostic push
-        #pragma GCC diagnostic ignored "-Wunused-local-typedefs"
-    #endif
-    #ifdef RL_USE_HALF_FLOATS
-        using StorageType = std::uint16_t;
-    #else
-        using StorageType = float;
-    #endif  // RL_USE_HALF_FLOATS
-    #if !defined(__clang__) && defined(__GNUC__)
-        #pragma GCC diagnostic pop
-    #endif
-    #if defined(__clang__)
-        #pragma clang diagnostic pop
-    #endif
-
-    // Work around unused parameter warning when building without SSE and AVX
-    static_cast<void>(config);
-    #ifdef RL_BUILD_WITH_SSE
-        if (config.calculationType == CalculationType::SSE) {
-            using SSEJointsBuilder = bpcm::BPCMJointsBuilder<StorageType, trimd::sse::F128>;
-            return UniqueInstance<SSEJointsBuilder, JointsBuilder>::with(memRes).create(memRes);
-        }
-    #endif  // RL_BUILD_WITH_SSE
-    #ifdef RL_BUILD_WITH_AVX
-        if (config.calculationType == CalculationType::AVX) {
-            using AVXJointsBuilder = bpcm::BPCMJointsBuilder<StorageType, trimd::avx::F256>;
-            return UniqueInstance<AVXJointsBuilder, JointsBuilder>::with(memRes).create(memRes);
-        }
-    #endif  // RL_BUILD_WITH_AVX
-    #ifdef RL_BUILD_WITH_NEON
-        if (config.calculationType == CalculationType::NEON) {
-            using NEONJointsBuilder = bpcm::BPCMJointsBuilder<StorageType, trimd::neon::F128>;
-            return UniqueInstance<NEONJointsBuilder, JointsBuilder>::with(memRes).create(memRes);
-        }
-    #endif  // RL_BUILD_WITH_NEON
-    using ScalarJointsBuilder = bpcm::BPCMJointsBuilder<float, trimd::scalar::F128>;
-    return UniqueInstance<ScalarJointsBuilder, JointsBuilder>::with(memRes).create(memRes);
+JointsBuilder::Pointer JointsBuilder::create(const Configuration& config, MemoryResource* memRes) {
+    return UniqueInstance<CPUJointsBuilder, JointsBuilder>::with(memRes).create(config, memRes);
 }
 
 }  // namespace rl4
diff --git a/Engine/Plugins/Animation/RigLogic/Source/RigLogicLib/Private/riglogic/joints/JointsBuilder.h b/Engine/Plugins/Animation/RigLogic/Source/RigLogicLib/Private/riglogic/joints/JointsBuilder.h
--- a/Engine/Plugins/Animation/RigLogic/Source/RigLogicLib/Private/riglogic/joints/JointsBuilder.h
+++ b/Engine/Plugins/Animation/RigLogic/Source/RigLogicLib/Private/riglogic/joints/JointsBuilder.h
@@ -8,6 +8,7 @@
 
 namespace rl4 {
 
+class JointBehaviorFilter;
 struct RigMetrics;
 
 class JointsBuilder {
@@ -17,12 +18,12 @@ class JointsBuilder {
     public:
         virtual ~JointsBuilder();
 
-        static Pointer create(Configuration config, MemoryResource* memRes);
+        static Pointer create(const Configuration& config, MemoryResource* memRes);
 
         virtual void computeStorageRequirements(const RigMetrics& source) = 0;
-        virtual void computeStorageRequirements(const dna::BehaviorReader* source) = 0;
-        virtual void allocateStorage(const dna::BehaviorReader* source) = 0;
-        virtual void fillStorage(const dna::BehaviorReader* source) = 0;
+        virtual void computeStorageRequirements(const JointBehaviorFilter& source) = 0;
+        virtual void allocateStorage(const JointBehaviorFilter& source) = 0;
+        virtual void fillStorage(const JointBehaviorFilter& source) = 0;
         virtual JointsEvaluator::Pointer build() = 0;
 };
 
diff --git a/Engine/Plugins/Animation/RigLogic/Source/RigLogicLib/Private/riglogic/joints/JointsFactory.cpp b/Engine/Plugins/Animation/RigLogic/Source/RigLogicLib/Private/riglogic/joints/JointsFactory.cpp
--- a/Engine/Plugins/Animation/RigLogic/Source/RigLogicLib/Private/riglogic/joints/JointsFactory.cpp
+++ b/Engine/Plugins/Animation/RigLogic/Source/RigLogicLib/Private/riglogic/joints/JointsFactory.cpp
@@ -3,12 +3,14 @@
 #include "riglogic/joints/JointsFactory.h"
 
 #include "riglogic/TypeDefs.h"
+#include "riglogic/joints/JointBehaviorFilter.h"
 #include "riglogic/joints/JointsBuilder.h"
 #include "riglogic/joints/JointsEvaluator.h"
 #include "riglogic/joints/JointsNullEvaluator.h"
 #include "riglogic/riglogic/Configuration.h"
 #include "riglogic/riglogic/RigMetrics.h"
-#include "riglogic/transformation/Transformation.h"
+
+#include <tdm/Quat.h>
 
 #ifdef _MSC_VER
     #pragma warning(push)
@@ -25,17 +27,6 @@
 
 namespace rl4 {
 
-namespace {
-
-constexpr std::size_t txOffset = 0ul;
-constexpr std::size_t tyOffset = 1ul;
-constexpr std::size_t tzOffset = 2ul;
-constexpr std::size_t rxOffset = 3ul;
-constexpr std::size_t ryOffset = 4ul;
-constexpr std::size_t rzOffset = 5ul;
-
-}  // namespace
-
 template<class TSource, class TDestination>
 static void scatter(const TSource& source, TDestination& destination, std::size_t stride, std::size_t offset) {
     const auto size = static_cast<std::size_t>(std::distance(std::begin(source), std::end(source)));
@@ -44,50 +35,126 @@ static void scatter(const TSource& source, TDestination& destination, std::size_
     }
 }
 
-static Vector<float> copyNeutralValues(const dna::BehaviorReader* reader, MemoryResource* memRes) {
+static Vector<float> copyNeutralValues(const Configuration& config, const dna::Reader* reader, MemoryResource* memRes) {
     Vector<float> neutralValues{memRes};
-    neutralValues.resize(reader->getJointCount() * Transformation::size());
-    std::fill(neutralValues.begin(), neutralValues.end(), 1.0f);
-    scatter(reader->getNeutralJointTranslationXs(), neutralValues, Transformation::size(), txOffset);
-    scatter(reader->getNeutralJointTranslationYs(), neutralValues, Transformation::size(), tyOffset);
-    scatter(reader->getNeutralJointTranslationZs(), neutralValues, Transformation::size(), tzOffset);
-    scatter(reader->getNeutralJointRotationXs(), neutralValues, Transformation::size(), rxOffset);
-    scatter(reader->getNeutralJointRotationYs(), neutralValues, Transformation::size(), ryOffset);
-    scatter(reader->getNeutralJointRotationZs(), neutralValues, Transformation::size(), rzOffset);
+    const auto numAttrsPerJoint = static_cast<std::size_t>(static_cast<std::uint8_t>(config.translationType) +
+                                                           static_cast<std::uint8_t>(config.rotationType) +
+                                                           static_cast<std::uint8_t>(config.scaleType));
+    neutralValues.resize(reader->getJointCount() * numAttrsPerJoint, 1.0f);
+    scatter(reader->getNeutralJointTranslationXs(), neutralValues, numAttrsPerJoint, 0ul);
+    scatter(reader->getNeutralJointTranslationYs(), neutralValues, numAttrsPerJoint, 1ul);
+    scatter(reader->getNeutralJointTranslationZs(), neutralValues, numAttrsPerJoint, 2ul);
+
+    if (config.rotationType == RotationType::EulerAngles) {
+        scatter(reader->getNeutralJointRotationXs(), neutralValues, numAttrsPerJoint, 3ul);
+        scatter(reader->getNeutralJointRotationYs(), neutralValues, numAttrsPerJoint, 4ul);
+        scatter(reader->getNeutralJointRotationZs(), neutralValues, numAttrsPerJoint, 5ul);
+    } else {
+        const auto rotationXs = reader->getNeutralJointRotationXs();
+        const auto rotationYs = reader->getNeutralJointRotationYs();
+        const auto rotationZs = reader->getNeutralJointRotationZs();
+        const auto rotationUnit = reader->getRotationUnit();
+        auto toRad = [rotationUnit](float x) {
+                return (rotationUnit == dna::RotationUnit::radians) ? tdm::frad{x} : tdm::frad{tdm::fdeg{x}};
+            };
+        for (std::size_t jointIndex = {}; jointIndex < rotationXs.size(); ++jointIndex) {
+            const std::size_t absAttrIndexBase = jointIndex * numAttrsPerJoint;
+            const tdm::frad3 euler{toRad(rotationXs[jointIndex]),
+                                   toRad(rotationYs[jointIndex]),
+                                   toRad(rotationZs[jointIndex])};
+            const tdm::fquat q{euler, static_cast<tdm::rot_seq>(config.rotationOrder)};
+            neutralValues[absAttrIndexBase + 3ul] = q.x;
+            neutralValues[absAttrIndexBase + 4ul] = q.y;
+            neutralValues[absAttrIndexBase + 5ul] = q.z;
+            neutralValues[absAttrIndexBase + 6ul] = q.w;
+        }
+    }
+
     return neutralValues;
 }
 
-static Matrix<std::uint16_t> copyVariableAttributeIndices(const dna::BehaviorReader* reader, MemoryResource* memRes) {
+static Matrix<std::uint16_t> copyVariableAttributeIndices(const Configuration& config,
+                                                          const dna::Reader* reader,
+                                                          MemoryResource* memRes) {
     Matrix<std::uint16_t> variableAttributeIndices{memRes};
     variableAttributeIndices.resize(reader->getLODCount());
-    for (std::uint16_t lod = 0u; lod < reader->getLODCount(); ++lod) {
-        auto indices = reader->getJointVariableAttributeIndices(lod);
-        variableAttributeIndices[lod].assign(indices.begin(), indices.end());
+
+    if (config.rotationType == RotationType::EulerAngles) {
+        for (std::uint16_t lod = 0u; lod < reader->getLODCount(); ++lod) {
+            const auto indices = reader->getJointVariableAttributeIndices(lod);
+            variableAttributeIndices[lod].assign(indices.begin(), indices.end());
+        }
+    } else {
+        const auto translationOffset = static_cast<std::uint8_t>(0);
+        const auto rotationOffset = translationOffset + static_cast<std::uint8_t>(config.translationType);
+        const auto scaleOffset = rotationOffset + static_cast<std::uint8_t>(config.rotationType);
+        const auto numAttrsPerJoint = static_cast<std::size_t>(static_cast<std::uint8_t>(config.translationType) +
+                                                               static_cast<std::uint8_t>(config.rotationType) +
+                                                               static_cast<std::uint8_t>(config.scaleType));
+        for (std::uint16_t lod = 0u; lod < reader->getLODCount(); ++lod) {
+            const auto indices = reader->getJointVariableAttributeIndices(lod);
+            variableAttributeIndices[lod].reserve(indices.size());
+            Vector<bool> markers(reader->getJointCount(), false, memRes);
+            for (const auto absAttrIndex : indices) {
+                const auto jointIndex = static_cast<std::uint16_t>(absAttrIndex / 9);
+                const auto relAttrIndex = static_cast<std::uint8_t>(absAttrIndex % 9);
+                const auto remappedBaseIndex = static_cast<std::uint16_t>(jointIndex * numAttrsPerJoint);
+                // Within the DNA, the structure is always fixed [tx, ty, tz, rx, ry, rz, sx, sy, sz]
+                if (relAttrIndex < 3) {
+                    const auto remapped = static_cast<std::uint16_t>(remappedBaseIndex + translationOffset + (relAttrIndex % 3));
+                    variableAttributeIndices[lod].push_back(remapped);
+                } else if (relAttrIndex < 6) {
+                    // Even if only a single rotation attribute is variable by DNA definition, when working with quaternions
+                    // all four attributes must be provided
+                    if (!markers[jointIndex]) {
+                        const auto attrBase = static_cast<std::uint16_t>(remappedBaseIndex + rotationOffset);
+                        std::uint16_t attrIndices[4] = {attrBase,
+                                                        static_cast<std::uint16_t>(attrBase + 1),
+                                                        static_cast<std::uint16_t>(attrBase + 2),
+                                                        static_cast<std::uint16_t>(attrBase + 3)};
+                        variableAttributeIndices[lod].insert(variableAttributeIndices[lod].end(),
+                                                             std::begin(attrIndices),
+                                                             std::end(attrIndices));
+                        markers[jointIndex] = true;
+                    }
+                } else {
+                    const auto remapped = static_cast<std::uint16_t>(remappedBaseIndex + scaleOffset + (relAttrIndex % 3));
+                    variableAttributeIndices[lod].push_back(remapped);
+                }
+            }
+        }
     }
+
     return variableAttributeIndices;
 }
 
-Joints::Pointer JointsFactory::create(const Configuration& config, const dna::BehaviorReader* reader, MemoryResource* memRes) {
-    if (!config.loadJoints) {
+Joints::Pointer JointsFactory::create(const Configuration& config, const dna::Reader* reader, MemoryResource* memRes) {
+    if (!config.loadJoints || (reader->getJointCount() == 0u)) {
         auto evaluator = UniqueInstance<JointsNullEvaluator, JointsEvaluator>::with(memRes).create();
         return UniqueInstance<Joints>::with(memRes).create(std::move(evaluator), memRes);
     }
 
-    auto neutralValues = copyNeutralValues(reader, memRes);
-    auto variableAttributeIndices = copyVariableAttributeIndices(reader, memRes);
+    JointBehaviorFilter filter{reader, memRes};
+    filter.include(dna::TranslationRepresentation::Vector);
+    filter.include(dna::RotationRepresentation::EulerAngles);
+    filter.include(dna::RotationRepresentation::Quaternion);
+    filter.include(dna::ScaleRepresentation::Vector);
+
     auto builder = JointsBuilder::create(config, memRes);
-    builder->computeStorageRequirements(reader);
-    builder->allocateStorage(reader);
-    builder->fillStorage(reader);
+    builder->computeStorageRequirements(filter);
+    builder->allocateStorage(filter);
+    builder->fillStorage(filter);
     auto evaluator = builder->build();
+    auto neutralValues = copyNeutralValues(config, reader, memRes);
+    auto variableAttributeIndices = copyVariableAttributeIndices(config, reader, memRes);
     return UniqueInstance<Joints>::with(memRes).create(std::move(evaluator),
                                                        std::move(neutralValues),
                                                        std::move(variableAttributeIndices),
-                                                       reader->getJointGroupCount());
+                                                       filter.getJointGroupCount());
 }
 
 Joints::Pointer JointsFactory::create(const Configuration& config, const RigMetrics& metrics, MemoryResource* memRes) {
-    if (!config.loadJoints) {
+    if (!config.loadJoints || (metrics.jointAttributeCount == 0u)) {
         auto evaluator = UniqueInstance<JointsNullEvaluator, JointsEvaluator>::with(memRes).create();
         return UniqueInstance<Joints>::with(memRes).create(std::move(evaluator), memRes);
     }
diff --git a/Engine/Plugins/Animation/RigLogic/Source/RigLogicLib/Private/riglogic/joints/JointsFactory.h b/Engine/Plugins/Animation/RigLogic/Source/RigLogicLib/Private/riglogic/joints/JointsFactory.h
--- a/Engine/Plugins/Animation/RigLogic/Source/RigLogicLib/Private/riglogic/joints/JointsFactory.h
+++ b/Engine/Plugins/Animation/RigLogic/Source/RigLogicLib/Private/riglogic/joints/JointsFactory.h
@@ -11,7 +11,7 @@ struct Configuration;
 struct RigMetrics;
 
 struct JointsFactory {
-    static Joints::Pointer create(const Configuration& config, const dna::BehaviorReader* reader, MemoryResource* memRes);
+    static Joints::Pointer create(const Configuration& config, const dna::Reader* reader, MemoryResource* memRes);
     static Joints::Pointer create(const Configuration& config, const RigMetrics& metrics, MemoryResource* memRes);
 
 };
diff --git a/Engine/Plugins/Animation/RigLogic/Source/RigLogicLib/Private/riglogic/joints/JointsNullOutputInstance.cpp b/Engine/Plugins/Animation/RigLogic/Source/RigLogicLib/Private/riglogic/joints/JointsNullOutputInstance.cpp
--- a/Engine/Plugins/Animation/RigLogic/Source/RigLogicLib/Private/riglogic/joints/JointsNullOutputInstance.cpp
+++ b/Engine/Plugins/Animation/RigLogic/Source/RigLogicLib/Private/riglogic/joints/JointsNullOutputInstance.cpp
@@ -8,4 +8,7 @@ ArrayView<float> JointsNullOutputInstance::getOutputBuffer() {
     return {};
 }
 
+void JointsNullOutputInstance::resetOutputBuffer() {
+}
+
 }  // namespace rl4
diff --git a/Engine/Plugins/Animation/RigLogic/Source/RigLogicLib/Private/riglogic/joints/cpu/CPUJointsBuilder.cpp b/Engine/Plugins/Animation/RigLogic/Source/RigLogicLib/Private/riglogic/joints/cpu/CPUJointsBuilder.cpp
new file mode 100644
--- /dev/null
+++ b/Engine/Plugins/Animation/RigLogic/Source/RigLogicLib/Private/riglogic/joints/cpu/CPUJointsBuilder.cpp
@@ -0,0 +1,81 @@
+// Copyright Epic Games, Inc. All Rights Reserved.
+
+#include "riglogic/joints/cpu/CPUJointsBuilder.h"
+
+#include "riglogic/TypeDefs.h"
+#include "riglogic/joints/cpu/CPUJointsEvaluator.h"
+#include "riglogic/joints/cpu/CPUJointsOutputInstance.h"
+#include "riglogic/joints/cpu/bpcm/BPCMJointsBuilderFactory.h"
+#include "riglogic/joints/cpu/quaternions/QuaternionJointsBuilderFactory.h"
+#include "riglogic/joints/cpu/twistswing/TwistSwingJointsBuilderFactory.h"
+#include "riglogic/riglogic/RigMetrics.h"
+#include "riglogic/system/simd/Detect.h"
+
+namespace rl4 {
+
+CPUJointsBuilder::CPUJointsBuilder(const Configuration& config_, MemoryResource* memRes_) :
+    memRes{memRes_},
+    config{config_},
+    bpcmBuilder{},
+    quaternionBuilder{},
+    twistSwingBuilder{},
+    numAttrsPerJoint{},
+    jointAttributeCount{} {
+
+    numAttrsPerJoint = static_cast<std::uint16_t>(static_cast<std::uint8_t>(config.translationType) +
+                                                  static_cast<std::uint8_t>(config.rotationType) +
+                                                  static_cast<std::uint8_t>(config.scaleType));
+    bpcmBuilder = BPCMJointsBuilderFactory::create(config, memRes);
+    quaternionBuilder = QuaternionJointsBuilderFactory::create(config, memRes);
+    twistSwingBuilder = TwistSwingJointsBuilderFactory::create(config, memRes);
+}
+
+void CPUJointsBuilder::computeStorageRequirements(const RigMetrics& source) {
+    jointAttributeCount = source.jointAttributeCount;
+    bpcmBuilder->computeStorageRequirements(source);
+    quaternionBuilder->computeStorageRequirements(source);
+    twistSwingBuilder->computeStorageRequirements(source);
+}
+
+void CPUJointsBuilder::computeStorageRequirements(const JointBehaviorFilter& source) {
+    jointAttributeCount = static_cast<std::uint16_t>(source.getJointCount() * numAttrsPerJoint);
+    bpcmBuilder->computeStorageRequirements(source.excluded(dna::RotationRepresentation::Quaternion));
+    quaternionBuilder->computeStorageRequirements(source.only(dna::RotationRepresentation::Quaternion));
+    twistSwingBuilder->computeStorageRequirements(source);
+}
+
+void CPUJointsBuilder::allocateStorage(const JointBehaviorFilter& source) {
+    bpcmBuilder->allocateStorage(source.excluded(dna::RotationRepresentation::Quaternion));
+    quaternionBuilder->allocateStorage(source.only(dna::RotationRepresentation::Quaternion));
+    twistSwingBuilder->allocateStorage(source);
+}
+
+void CPUJointsBuilder::fillStorage(const JointBehaviorFilter& source) {
+    bpcmBuilder->fillStorage(source.excluded(dna::RotationRepresentation::Quaternion));
+    quaternionBuilder->fillStorage(source.only(dna::RotationRepresentation::Quaternion));
+    twistSwingBuilder->fillStorage(source);
+}
+
+JointsEvaluator::Pointer CPUJointsBuilder::build() {
+    const auto attrCount = jointAttributeCount;
+    const auto translationType = config.translationType;
+    const auto rotationType = config.rotationType;
+    const auto scaleType = config.scaleType;
+    auto instanceFactory = [attrCount, translationType, rotationType, scaleType](MemoryResource* instanceMemRes) {
+            return UniqueInstance<CPUJointsOutputInstance, JointsOutputInstance>::with(instanceMemRes).create(attrCount,
+                                                                                                              translationType,
+                                                                                                              rotationType,
+                                                                                                              scaleType,
+                                                                                                              instanceMemRes);
+        };
+    auto bpcmEvaluator = bpcmBuilder->build();
+    auto quaternionEvaluator = quaternionBuilder->build();
+    auto twistSwingEvaluator = twistSwingBuilder->build();
+    auto factory = UniqueInstance<CPUJointsEvaluator, JointsEvaluator>::with(memRes);
+    return factory.create(std::move(bpcmEvaluator),
+                          std::move(quaternionEvaluator),
+                          std::move(twistSwingEvaluator),
+                          instanceFactory);
+}
+
+}  // namespace rl4
diff --git a/Engine/Plugins/Animation/RigLogic/Source/RigLogicLib/Private/riglogic/joints/cpu/bpcm/BPCMJointsBuilder.h b/Engine/Plugins/Animation/RigLogic/Source/RigLogicLib/Private/riglogic/joints/cpu/bpcm/BPCMJointsBuilder.h
new file mode 100644
--- /dev/null
+++ b/Engine/Plugins/Animation/RigLogic/Source/RigLogicLib/Private/riglogic/joints/cpu/bpcm/BPCMJointsBuilder.h
@@ -0,0 +1,381 @@
+// Copyright Epic Games, Inc. All Rights Reserved.
+
+#pragma once
+
+#include "riglogic/TypeDefs.h"
+#include "riglogic/joints/JointBehaviorFilter.h"
+#include "riglogic/joints/JointsBuilder.h"
+#include "riglogic/joints/cpu/bpcm/BPCMJointsEvaluator.h"
+#include "riglogic/joints/cpu/bpcm/CalculationStrategy.h"
+#include "riglogic/joints/cpu/bpcm/RotationAdapters.h"
+#include "riglogic/joints/cpu/bpcm/Storage.h"
+#include "riglogic/joints/cpu/bpcm/StorageSize.h"
+#include "riglogic/riglogic/Configuration.h"
+#include "riglogic/riglogic/RigMetrics.h"
+#include "riglogic/types/Aliases.h"
+#include "riglogic/types/bpcm/Optimizer.h"
+#include "riglogic/utils/Extd.h"
+
+#ifdef _MSC_VER
+    #pragma warning(push)
+    #pragma warning(disable : 4365 4987)
+#endif
+#include <algorithm>
+#include <cstddef>
+#include <cstdint>
+#include <functional>
+#include <memory>
+#include <utility>
+#ifdef _MSC_VER
+    #pragma warning(pop)
+#endif
+
+namespace rl4 {
+
+namespace bpcm {
+
+template<typename TIterator>
+static void remapOutputIndicesForQuaternions(TIterator begin, TIterator end) {
+    for (auto it = begin; it != end; ++it) {
+        const auto absAttrIndex = *it;
+        const auto jointIndex = static_cast<std::uint16_t>(absAttrIndex / 9u);
+        const auto relAttrIndex = static_cast<std::uint16_t>(absAttrIndex % 9u);
+        const auto newAttrBase = static_cast<std::uint16_t>(jointIndex * 10u);
+        // Only scale relative attribute index is offset by one when output is in quaternions
+        const auto newRelAttrIndex = (relAttrIndex < 6u ? relAttrIndex : static_cast<std::uint16_t>(relAttrIndex + 1u));
+        *it = static_cast<std::uint16_t>(newAttrBase + newRelAttrIndex);
+    }
+}
+
+template<typename TValue, typename TFVec>
+class BPCMJointsBuilder : public JointsBuilder {
+    public:
+        BPCMJointsBuilder(const Configuration& config_, MemoryResource* memRes_);
+
+        void computeStorageRequirements(const RigMetrics& source) override;
+        void computeStorageRequirements(const JointBehaviorFilter& source) override;
+        void allocateStorage(const JointBehaviorFilter& source) override;
+        void fillStorage(const JointBehaviorFilter& source) override;
+        JointsEvaluator::Pointer build() override;
+
+    private:
+        void setValues(const JointBehaviorFilter& source);
+        void setInputIndices(const JointBehaviorFilter& source);
+        void setOutputIndices(const JointBehaviorFilter& source);
+        void setLODs(const JointBehaviorFilter& source);
+        void setOutputRotationIndices(const JointBehaviorFilter& source);
+        void setOutputRotationLODs(const JointBehaviorFilter& source,
+                                   ConstArrayView<std::uint16_t> outputRotationIndices,
+                                   std::uint32_t outputOffset,
+                                   std::uint16_t jointGroupIndex);
+        static constexpr std::uint32_t BlockHeight() {
+            return static_cast<std::uint32_t>(TFVec::size() * 2ul);
+        }
+
+        static constexpr std::uint32_t PadTo() {
+            return static_cast<std::uint32_t>(TFVec::size());
+        }
+
+    private:
+        Configuration config;
+        MemoryResource* memRes;
+        StorageSize sizeReqs;
+        JointStorage<TValue> storage;
+        dna::RotationUnit rotationUnit;
+};
+
+template<typename TValue, typename TFVec>
+BPCMJointsBuilder<TValue, TFVec>::BPCMJointsBuilder(const Configuration& config_, MemoryResource* memRes_) :
+    config{config_},
+    memRes{memRes_},
+    sizeReqs{memRes},
+    storage{memRes},
+    rotationUnit{} {
+}
+
+template<typename TValue, typename TFVec>
+void BPCMJointsBuilder<TValue, TFVec>::computeStorageRequirements(const RigMetrics&  /*unused*/) {
+}
+
+template<typename TValue, typename TFVec>
+void BPCMJointsBuilder<TValue, TFVec>::computeStorageRequirements(const JointBehaviorFilter& source) {
+    sizeReqs.computeFrom(source, PadTo());
+}
+
+template<typename TValue, typename TFVec>
+void BPCMJointsBuilder<TValue, TFVec>::allocateStorage(const JointBehaviorFilter&  /*unused*/) {
+    storage.values.resize(sizeReqs.valueCount);
+    storage.inputIndices.resize(sizeReqs.inputIndexCount);
+    storage.outputIndices.resize(sizeReqs.outputIndexCount);
+    storage.lodRegions.reserve(sizeReqs.lodRegionCount);
+    storage.jointGroups.resize(sizeReqs.jointGroups.size());
+    if (config.rotationType == RotationType::Quaternions) {
+        storage.outputRotationLODs.resize(sizeReqs.jointGroups.size() * sizeReqs.lodCount);
+    }
+}
+
+template<typename TValue, typename TFVec>
+void BPCMJointsBuilder<TValue, TFVec>::fillStorage(const JointBehaviorFilter& source) {
+    rotationUnit = source.getRotationUnit();
+    setValues(source);
+    setInputIndices(source);
+    setOutputIndices(source);
+    setLODs(source);
+    if (config.rotationType == RotationType::Quaternions) {
+        // Remap output indices from 9-attribute joints to 10-attribute joints
+        remapOutputIndicesForQuaternions(storage.outputIndices.begin(), storage.outputIndices.end());
+        const JointBehaviorFilter filtered = source.only(dna::RotationRepresentation::EulerAngles);
+        setOutputRotationIndices(filtered);
+    }
+}
+
+template<typename TValue, typename TFVec>
+void BPCMJointsBuilder<TValue, TFVec>::setValues(const JointBehaviorFilter& source) {
+    std::uint32_t offset = {};
+    Vector<float> buffer{memRes};
+    for (std::uint16_t i = {}; i < source.getJointGroupCount(); ++i) {
+        const auto jointGroupSize = sizeReqs.getJointGroupSize(i);
+        storage.jointGroups[i].valuesOffset = offset;
+        storage.jointGroups[i].valuesSize = jointGroupSize.padded.size();
+
+        buffer.resize(jointGroupSize.padded.size());
+        source.copyValues(i, buffer);
+
+        offset += Optimizer<TFVec, BlockHeight(), PadTo(), 1u>::optimize(storage.values.data() + offset,
+                                                                         buffer.data(),
+                                                                         jointGroupSize.original);
+    }
+}
+
+template<typename TValue, typename TFVec>
+void BPCMJointsBuilder<TValue, TFVec>::setInputIndices(const JointBehaviorFilter& source) {
+    std::uint32_t offset = {};
+    for (std::uint16_t i = {}; i < source.getJointGroupCount(); ++i) {
+        const auto jointGroupSize = sizeReqs.getJointGroupSize(i);
+        source.copyInputIndices(i, {storage.inputIndices.data() + offset, jointGroupSize.padded.cols});
+        storage.jointGroups[i].inputIndicesOffset = offset;
+        storage.jointGroups[i].inputIndicesSize = jointGroupSize.padded.cols;
+        storage.jointGroups[i].inputIndicesSizeAlignedTo4 = jointGroupSize.padded.cols - (jointGroupSize.padded.cols % 4u);
+        storage.jointGroups[i].inputIndicesSizeAlignedTo8 = jointGroupSize.padded.cols - (jointGroupSize.padded.cols % 8u);
+        offset += jointGroupSize.padded.cols;
+    }
+}
+
+template<typename TValue, typename TFVec>
+void BPCMJointsBuilder<TValue, TFVec>::setOutputIndices(const JointBehaviorFilter& source) {
+    std::uint32_t offset = {};
+    for (std::uint16_t i = {}; i < source.getJointGroupCount(); ++i) {
+        const auto jointGroupSize = sizeReqs.getJointGroupSize(i);
+        source.copyOutputIndices(i, {storage.outputIndices.data() + offset, jointGroupSize.padded.rows});
+        storage.jointGroups[i].outputIndicesOffset = offset;
+        offset += jointGroupSize.padded.rows;
+    }
+}
+
+template<typename TValue, typename TFVec>
+void BPCMJointsBuilder<TValue, TFVec>::setLODs(const JointBehaviorFilter& source) {
+    std::uint32_t offset = {};
+    for (std::uint16_t i = {}; i < source.getJointGroupCount(); ++i) {
+        const auto jointGroupSize = sizeReqs.getJointGroupSize(i);
+        const auto paddedRowCount = jointGroupSize.padded.rows;
+        for (std::uint16_t lod = {}; lod < source.getLODCount(); ++lod) {
+            storage.lodRegions.emplace_back(source.getRowCountForLOD(i, lod),
+                                            paddedRowCount,
+                                            BlockHeight(),
+                                            PadTo());
+        }
+        storage.jointGroups[i].lodsOffset = offset;
+        offset += source.getLODCount();
+    }
+}
+
+template<typename TValue, typename TFVec>
+void BPCMJointsBuilder<TValue, TFVec>::setOutputRotationIndices(const JointBehaviorFilter& source) {
+    std::uint32_t outputOffset = {};
+    for (std::uint16_t jgi = {}; jgi < source.getJointGroupCount(); ++jgi) {
+        const auto rowCount = source.getRowCount(jgi);
+        Vector<std::uint16_t> outputRotationIndices{rowCount, {}, memRes};
+        source.copyOutputIndices(jgi, {outputRotationIndices.data(), rowCount});
+        // Remap output indices from 9-attribute joints to 10-attribute joints rx -> qx
+        remapOutputIndicesForQuaternions(outputRotationIndices.begin(), outputRotationIndices.end());
+        // Given any rotation indices (qx, qy, qz), return only qx indices for all joints in the group
+        #if !defined(__clang__) && defined(__GNUC__)
+            #pragma GCC diagnostic push
+            #pragma GCC diagnostic ignored "-Wattributes"
+        #endif
+        auto deduplicate = [this](Vector<std::uint16_t>& v) {
+                UnorderedSet<std::uint16_t> deduplicator{memRes};
+                v.erase(std::remove_if(v.begin(), v.end(), [&deduplicator](const std::uint16_t value) {
+                        return !deduplicator.insert(value).second;
+                    }), v.end());
+            };
+        #if !defined(__clang__) && defined(__GNUC__)
+            #pragma GCC diagnostic pop
+        #endif
+        Vector<std::uint16_t> outputRotationBaseIndices{memRes};
+        outputRotationBaseIndices.reserve(outputRotationIndices.size() / 3ul);
+        std::transform(outputRotationIndices.begin(),
+                       outputRotationIndices.end(),
+                       std::back_inserter(outputRotationBaseIndices),
+                       [](std::uint16_t outputIndex) {
+                    return static_cast<std::uint16_t>((outputIndex / 10) * 10 + 3);
+                });
+        deduplicate(outputRotationBaseIndices);
+        // Copy remapped qx indices into destination storage
+        std::copy(outputRotationBaseIndices.begin(),
+                  outputRotationBaseIndices.end(),
+                  std::back_inserter(storage.outputRotationIndices));
+        storage.jointGroups[jgi].outputRotationIndicesOffset = outputOffset;
+        // Must be called before outputOffset is adjusted
+        setOutputRotationLODs(source, outputRotationIndices, outputOffset, jgi);
+        outputOffset += static_cast<std::uint32_t>(outputRotationBaseIndices.size());
+    }
+}
+
+template<typename TValue, typename TFVec>
+void BPCMJointsBuilder<TValue, TFVec>::setOutputRotationLODs(const JointBehaviorFilter& source,
+                                                             ConstArrayView<std::uint16_t> outputRotationIndices,
+                                                             std::uint32_t outputOffset,
+                                                             std::uint16_t jointGroupIndex) {
+    const auto lodCount = source.getLODCount();
+    const auto offset = static_cast<std::uint32_t>(jointGroupIndex * lodCount);
+    for (std::uint16_t lod = {}; lod < lodCount; ++lod) {
+        const auto oldLODRowCount = source.getRowCountForLOD(jointGroupIndex, lod);
+        if (oldLODRowCount == 0) {
+            storage.outputRotationLODs[offset + lod] = 0;
+            continue;
+        }
+
+        assert(oldLODRowCount <= outputRotationIndices.size());
+        const auto qxRotationIndexAtOldLODRowCount = (outputRotationIndices[oldLODRowCount - 1ul] / 10) * 10 + 3;
+        auto start = extd::advanced(storage.outputRotationIndices.begin(), outputOffset);
+        auto it = std::find(start, storage.outputRotationIndices.end(), qxRotationIndexAtOldLODRowCount);
+        assert(it != storage.outputRotationIndices.end());
+        const auto newLODRowCount = static_cast<std::uint16_t>(std::distance(start, it) + 1);
+        storage.outputRotationLODs[offset + lod] = newLODRowCount;
+    }
+    storage.jointGroups[jointGroupIndex].outputRotationLODsOffset = offset;
+}
+
+template<typename T, typename TFVec>
+typename UniqueInstance<JointGroupLinearCalculationStrategy<T> >::PointerType createJointGroupLinearStrategy(
+    RotationType rotationType,
+    RotationOrder rotationOrder,
+    dna::RotationUnit rotationUnit,
+    MemoryResource* memRes) {
+
+    if (rotationType == RotationType::EulerAngles) {
+        using CalculationStrategy = VectorizedJointGroupLinearCalculationStrategy<T, TFVec, NoopAdapter>;
+        return UniqueInstance<CalculationStrategy, JointGroupLinearCalculationStrategy<T> >::with(memRes).create();
+    }
+
+    #ifdef RL_BUILD_WITH_XYZ_ROTATION_ORDER
+        if (rotationOrder == RotationOrder::XYZ) {
+            if (rotationUnit == dna::RotationUnit::degrees) {
+                using CalculationStrategy = VectorizedJointGroupLinearCalculationStrategy<T, TFVec,
+                                                                                          EulerAnglesToQuaternions<tdm::fdeg,
+                                                                                                                   tdm::rot_seq::xyz> >;
+                return UniqueInstance<CalculationStrategy, JointGroupLinearCalculationStrategy<T> >::with(memRes).create();
+            } else {
+                using CalculationStrategy = VectorizedJointGroupLinearCalculationStrategy<T, TFVec,
+                                                                                          EulerAnglesToQuaternions<tdm::frad,
+                                                                                                                   tdm::rot_seq::xyz> >;
+                return UniqueInstance<CalculationStrategy, JointGroupLinearCalculationStrategy<T> >::with(memRes).create();
+            }
+        }
+    #endif  // RL_BUILD_WITH_XYZ_ROTATION_ORDER
+
+    #ifdef RL_BUILD_WITH_XZY_ROTATION_ORDER
+        if (rotationOrder == RotationOrder::XZY) {
+            if (rotationUnit == dna::RotationUnit::degrees) {
+                using CalculationStrategy = VectorizedJointGroupLinearCalculationStrategy<T, TFVec,
+                                                                                          EulerAnglesToQuaternions<tdm::fdeg,
+                                                                                                                   tdm::rot_seq::xzy> >;
+                return UniqueInstance<CalculationStrategy, JointGroupLinearCalculationStrategy<T> >::with(memRes).create();
+            } else {
+                using CalculationStrategy = VectorizedJointGroupLinearCalculationStrategy<T, TFVec,
+                                                                                          EulerAnglesToQuaternions<tdm::frad,
+                                                                                                                   tdm::rot_seq::xzy> >;
+                return UniqueInstance<CalculationStrategy, JointGroupLinearCalculationStrategy<T> >::with(memRes).create();
+            }
+        }
+    #endif  // RL_BUILD_WITH_XZY_ROTATION_ORDER
+
+    #ifdef RL_BUILD_WITH_YXZ_ROTATION_ORDER
+        if (rotationOrder == RotationOrder::YXZ) {
+            if (rotationUnit == dna::RotationUnit::degrees) {
+                using CalculationStrategy = VectorizedJointGroupLinearCalculationStrategy<T, TFVec,
+                                                                                          EulerAnglesToQuaternions<tdm::fdeg,
+                                                                                                                   tdm::rot_seq::yxz> >;
+                return UniqueInstance<CalculationStrategy, JointGroupLinearCalculationStrategy<T> >::with(memRes).create();
+            } else {
+                using CalculationStrategy = VectorizedJointGroupLinearCalculationStrategy<T, TFVec,
+                                                                                          EulerAnglesToQuaternions<tdm::frad,
+                                                                                                                   tdm::rot_seq::yxz> >;
+                return UniqueInstance<CalculationStrategy, JointGroupLinearCalculationStrategy<T> >::with(memRes).create();
+            }
+        }
+    #endif  // RL_BUILD_WITH_YXZ_ROTATION_ORDER
+
+    #ifdef RL_BUILD_WITH_YZX_ROTATION_ORDER
+        if (rotationOrder == RotationOrder::YZX) {
+            if (rotationUnit == dna::RotationUnit::degrees) {
+                using CalculationStrategy = VectorizedJointGroupLinearCalculationStrategy<T, TFVec,
+                                                                                          EulerAnglesToQuaternions<tdm::fdeg,
+                                                                                                                   tdm::rot_seq::yzx> >;
+                return UniqueInstance<CalculationStrategy, JointGroupLinearCalculationStrategy<T> >::with(memRes).create();
+            } else {
+                using CalculationStrategy = VectorizedJointGroupLinearCalculationStrategy<T, TFVec,
+                                                                                          EulerAnglesToQuaternions<tdm::frad,
+                                                                                                                   tdm::rot_seq::yzx> >;
+                return UniqueInstance<CalculationStrategy, JointGroupLinearCalculationStrategy<T> >::with(memRes).create();
+            }
+        }
+    #endif  // RL_BUILD_WITH_YZX_ROTATION_ORDER
+
+    #ifdef RL_BUILD_WITH_ZXY_ROTATION_ORDER
+        if (rotationOrder == RotationOrder::ZXY) {
+            if (rotationUnit == dna::RotationUnit::degrees) {
+                using CalculationStrategy = VectorizedJointGroupLinearCalculationStrategy<T, TFVec,
+                                                                                          EulerAnglesToQuaternions<tdm::fdeg,
+                                                                                                                   tdm::rot_seq::zxy> >;
+                return UniqueInstance<CalculationStrategy, JointGroupLinearCalculationStrategy<T> >::with(memRes).create();
+            } else {
+                using CalculationStrategy = VectorizedJointGroupLinearCalculationStrategy<T, TFVec,
+                                                                                          EulerAnglesToQuaternions<tdm::frad,
+                                                                                                                   tdm::rot_seq::zxy> >;
+                return UniqueInstance<CalculationStrategy, JointGroupLinearCalculationStrategy<T> >::with(memRes).create();
+            }
+        }
+    #endif  // RL_BUILD_WITH_ZXY_ROTATION_ORDER
+
+    #ifdef RL_BUILD_WITH_ZYX_ROTATION_ORDER
+        if (rotationOrder == RotationOrder::ZYX) {
+            if (rotationUnit == dna::RotationUnit::degrees) {
+                using CalculationStrategy = VectorizedJointGroupLinearCalculationStrategy<T, TFVec,
+                                                                                          EulerAnglesToQuaternions<tdm::fdeg,
+                                                                                                                   tdm::rot_seq::zyx> >;
+                return UniqueInstance<CalculationStrategy, JointGroupLinearCalculationStrategy<T> >::with(memRes).create();
+            } else {
+                using CalculationStrategy = VectorizedJointGroupLinearCalculationStrategy<T, TFVec,
+                                                                                          EulerAnglesToQuaternions<tdm::frad,
+                                                                                                                   tdm::rot_seq::zyx> >;
+                return UniqueInstance<CalculationStrategy, JointGroupLinearCalculationStrategy<T> >::with(memRes).create();
+            }
+        }
+    #endif  // RL_BUILD_WITH_ZYX_ROTATION_ORDER
+
+    return nullptr;
+}
+
+template<typename TValue, typename TFVec>
+JointsEvaluator::Pointer BPCMJointsBuilder<TValue, TFVec>::build() {
+    auto strategy =
+        createJointGroupLinearStrategy<TValue, TFVec>(config.rotationType, config.rotationOrder, rotationUnit, memRes);
+    auto factory = UniqueInstance<Evaluator<TValue>, JointsEvaluator>::with(memRes);
+    return factory.create(std::move(storage), std::move(strategy), nullptr, memRes);
+}
+
+}  // namespace bpcm
+
+}  // namespace rl4
diff --git a/Engine/Plugins/Animation/RigLogic/Source/RigLogicLib/Private/riglogic/joints/cpu/bpcm/BPCMJointsBuilderFactory.cpp b/Engine/Plugins/Animation/RigLogic/Source/RigLogicLib/Private/riglogic/joints/cpu/bpcm/BPCMJointsBuilderFactory.cpp
new file mode 100644
--- /dev/null
+++ b/Engine/Plugins/Animation/RigLogic/Source/RigLogicLib/Private/riglogic/joints/cpu/bpcm/BPCMJointsBuilderFactory.cpp
@@ -0,0 +1,53 @@
+// Copyright Epic Games, Inc. All Rights Reserved.
+
+#include "riglogic/joints/cpu/bpcm/BPCMJointsBuilderFactory.h"
+
+#include "riglogic/joints/cpu/bpcm/BPCMJointsBuilder.h"
+#include "riglogic/system/simd/Detect.h"
+
+namespace rl4 {
+
+UniqueInstance<JointsBuilder>::PointerType BPCMJointsBuilderFactory::create(const Configuration& config, MemoryResource* memRes) {
+    #if defined(__clang__)
+        #pragma clang diagnostic push
+        #pragma clang diagnostic ignored "-Wunused-local-typedef"
+    #endif
+    #if !defined(__clang__) && defined(__GNUC__)
+        #pragma GCC diagnostic push
+        #pragma GCC diagnostic ignored "-Wunused-local-typedefs"
+    #endif
+    #ifdef RL_USE_HALF_FLOATS
+        using StorageType = std::uint16_t;
+    #else
+        using StorageType = float;
+    #endif  // RL_USE_HALF_FLOATS
+    #if !defined(__clang__) && defined(__GNUC__)
+        #pragma GCC diagnostic pop
+    #endif
+    #if defined(__clang__)
+        #pragma clang diagnostic pop
+    #endif
+
+    #ifdef RL_BUILD_WITH_SSE
+        if ((config.calculationType == CalculationType::SSE) || (config.calculationType == CalculationType::AnyVector)) {
+            using SSEBPCMJointsBuilder = bpcm::BPCMJointsBuilder<StorageType, trimd::sse::F128>;
+            return UniqueInstance<SSEBPCMJointsBuilder, JointsBuilder>::with(memRes).create(config, memRes);
+        }
+    #endif  // RL_BUILD_WITH_SSE
+    #ifdef RL_BUILD_WITH_AVX
+        if ((config.calculationType == CalculationType::AVX) || (config.calculationType == CalculationType::AnyVector)) {
+            using AVXBPCMJointsBuilder = bpcm::BPCMJointsBuilder<StorageType, trimd::avx::F256>;
+            return UniqueInstance<AVXBPCMJointsBuilder, JointsBuilder>::with(memRes).create(config, memRes);
+        }
+    #endif  // RL_BUILD_WITH_AVX
+    #ifdef RL_BUILD_WITH_NEON
+        if ((config.calculationType == CalculationType::NEON) || (config.calculationType == CalculationType::AnyVector)) {
+            using NEONBPCMJointsBuilder = bpcm::BPCMJointsBuilder<StorageType, trimd::neon::F128>;
+            return UniqueInstance<NEONBPCMJointsBuilder, JointsBuilder>::with(memRes).create(config, memRes);
+        }
+    #endif  // RL_BUILD_WITH_NEON
+    using ScalarBPCMJointsBuilder = bpcm::BPCMJointsBuilder<float, trimd::scalar::F128>;
+    return UniqueInstance<ScalarBPCMJointsBuilder, JointsBuilder>::with(memRes).create(config, memRes);
+}
+
+}  // namespace rl4
diff --git a/Engine/Plugins/Animation/RigLogic/Source/RigLogicLib/Private/riglogic/joints/cpu/bpcm/BPCMJointsBuilderFactory.h b/Engine/Plugins/Animation/RigLogic/Source/RigLogicLib/Private/riglogic/joints/cpu/bpcm/BPCMJointsBuilderFactory.h
new file mode 100644
--- /dev/null
+++ b/Engine/Plugins/Animation/RigLogic/Source/RigLogicLib/Private/riglogic/joints/cpu/bpcm/BPCMJointsBuilderFactory.h
@@ -0,0 +1,16 @@
+// Copyright Epic Games, Inc. All Rights Reserved.
+
+#pragma once
+
+#include "riglogic/TypeDefs.h"
+#include "riglogic/joints/JointsBuilder.h"
+#include "riglogic/riglogic/Configuration.h"
+
+namespace rl4 {
+
+struct BPCMJointsBuilderFactory {
+    static UniqueInstance<JointsBuilder>::PointerType create(const Configuration& config, MemoryResource* memRes);
+
+};
+
+}  // namespace rl4
diff --git a/Engine/Plugins/Animation/RigLogic/Source/RigLogicLib/Private/riglogic/joints/cpu/bpcm/BPCMJointsEvaluator.h b/Engine/Plugins/Animation/RigLogic/Source/RigLogicLib/Private/riglogic/joints/cpu/bpcm/BPCMJointsEvaluator.h
new file mode 100644
--- /dev/null
+++ b/Engine/Plugins/Animation/RigLogic/Source/RigLogicLib/Private/riglogic/joints/cpu/bpcm/BPCMJointsEvaluator.h
@@ -0,0 +1,81 @@
+// Copyright Epic Games, Inc. All Rights Reserved.
+
+#pragma once
+
+#include "riglogic/TypeDefs.h"
+#include "riglogic/controls/ControlsInputInstance.h"
+#include "riglogic/joints/JointsEvaluator.h"
+#include "riglogic/joints/JointsOutputInstance.h"
+#include "riglogic/joints/cpu/bpcm/CalculationStrategy.h"
+#include "riglogic/joints/cpu/bpcm/Storage.h"
+#include "riglogic/joints/cpu/bpcm/StorageView.h"
+#include "riglogic/riglogic/RigInstanceImpl.h"
+
+#include <cstdint>
+
+namespace rl4 {
+
+namespace bpcm {
+
+template<typename TValue>
+class Evaluator : public JointsEvaluator {
+    public:
+        using CalculationStrategyPointer = typename UniqueInstance<JointGroupLinearCalculationStrategy<TValue> >::PointerType;
+
+        struct Accessor;
+        friend Accessor;
+
+    public:
+        Evaluator(JointStorage<TValue>&& storage_,
+                  CalculationStrategyPointer strategy_,
+                  JointsOutputInstance::Factory instanceFactory_,
+                  MemoryResource* memRes_) :
+            memRes{memRes_},
+            storage{std::move(storage_)},
+            jointGroups{takeStorageSnapshot(storage, memRes)},
+            strategy{std::move(strategy_)},
+            instanceFactory{instanceFactory_} {
+        }
+
+        JointsOutputInstance::Pointer createInstance(MemoryResource* instanceMemRes) const override {
+            return instanceFactory(instanceMemRes);
+        }
+
+        void calculate(const ControlsInputInstance* inputs, JointsOutputInstance* outputs, std::uint16_t lod) const override {
+            assert(strategy != nullptr);
+            for (std::size_t i = {}; i < jointGroups.size(); ++i) {
+                calculate(inputs, outputs, lod, static_cast<std::uint16_t>(i));
+            }
+        }
+
+        void calculate(const ControlsInputInstance* inputs,
+                       JointsOutputInstance* outputs,
+                       std::uint16_t lod,
+                       std::uint16_t jointGroupIndex) const override {
+            assert(strategy != nullptr);
+            strategy->calculate(jointGroups[jointGroupIndex],
+                                inputs->getInputBuffer(),
+                                outputs->getOutputBuffer(),
+                                lod);
+        }
+
+        void load(terse::BinaryInputArchive<BoundedIOStream>& archive) override {
+            archive(storage);
+            jointGroups = takeStorageSnapshot(storage, memRes);
+        }
+
+        void save(terse::BinaryOutputArchive<BoundedIOStream>& archive) override {
+            archive(storage);
+        }
+
+    private:
+        MemoryResource* memRes;
+        JointStorage<TValue> storage;
+        Vector<JointGroupView<TValue> > jointGroups;
+        CalculationStrategyPointer strategy;
+        JointsOutputInstance::Factory instanceFactory;
+};
+
+}  // namespace bpcm
+
+}  // namespace rl4
