Update to new RigLogic implementing body rig support

diff --git a/Engine/Plugins/Animation/RigLogic/Source/RigLogicLib/Private/riglogic/joints/cpu/bpcm/CalculationStrategy.h b/Engine/Plugins/Animation/RigLogic/Source/RigLogicLib/Private/riglogic/joints/cpu/bpcm/CalculationStrategy.h
new file mode 100644
--- /dev/null
+++ b/Engine/Plugins/Animation/RigLogic/Source/RigLogicLib/Private/riglogic/joints/cpu/bpcm/CalculationStrategy.h
@@ -0,0 +1,363 @@
+// Copyright Epic Games, Inc. All Rights Reserved.
+
+#pragma once
+
+#include "riglogic/joints/cpu/bpcm/Consts.h"
+#include "riglogic/joints/cpu/bpcm/CalculationStrategy.h"
+#include "riglogic/joints/cpu/bpcm/JointGroup.h"
+#include "riglogic/joints/cpu/bpcm/StorageView.h"
+#include "riglogic/system/simd/SIMD.h"
+#include "riglogic/types/Aliases.h"
+#include "riglogic/utils/Macros.h"
+
+#include <cstddef>
+#include <cstdint>
+
+namespace rl4 {
+
+namespace bpcm {
+
+/*
+ * Process the remainder portion after 8x4 blocks
+ *
+ * Consumes 4 x 2 values per input index, but input indexes are processed one
+ * by one. The maximum number of iterations performed is 3, as anything larger
+ * would be already consumed by the 8x4 block processor. It's executed only
+ * after the 8x4 block processor in cases when the matrix column count is not
+ * divisible by 8.
+ *
+ *  [ x x x x x x x x X Y Z ] Input Vector
+ *
+ *    0 1 2 3 4 5 6 7 8 9 0
+ *  0 x x x x x x x x A C E  |  X * A + Y * C + Z * E
+ *  1 x x x x x x x x A C E  |  X * A + Y * C + Z * E
+ *  2 x x x x x x x x A C E  |  X * A + Y * C + Z * E
+ *  3 x x x x x x x x A C E  |  X * A + Y * C + Z * E
+ *  4 x x x x x x x x B D F  |  X * B + Y * D + Z * F
+ *  5 x x x x x x x x B D F  |  X * B + Y * D + Z * F
+ *  6 x x x x x x x x B D F  |  X * B + Y * D + Z * F
+ *  7 x x x x x x x x B D F  |  X * B + Y * D + Z * F
+ *  8 x x x x x x x x x x x  |
+ *  9 x x x x x x x x x x x  |
+ * 10 x x x x x x x x x x x  |
+ * 11 x x x x x x x x x x x  |
+ *
+ */
+template<typename TFVec, typename T>
+static FORCE_INLINE void processBlocks8x1(const std::uint16_t* inputIndicesEndAlignedTo4,
+                                          const std::uint16_t* inputIndicesEnd,
+                                          ConstArrayView<float> inputs,
+                                          const T*& values,
+                                          TFVec& sum1,
+                                          TFVec& sum2) {
+    for (const std::uint16_t* inputIndices = inputIndicesEndAlignedTo4;
+         inputIndices < inputIndicesEnd;
+         ++inputIndices, values += (2ul * TFVec::size())) {
+        const TFVec inputVec{inputs[*inputIndices]};
+        const TFVec blk1 = TFVec::fromAlignedSource(values);
+        const TFVec blk2 = TFVec::fromAlignedSource(values + TFVec::size());
+        sum1 += (blk1 * inputVec);
+        sum2 += (blk2 * inputVec);
+    }
+}
+
+/*
+ * Process the main portion of the matrix that's partitionable into 8x4 blocks
+ *
+ * Consumes 4 x 2 values per input index, and loads 4 input indexes per iteration,
+ * thus a single iteration processes 32 values. This is the most performant of all
+ * the block processors as for every 4 x 2 values only a single input vector needs
+ * to be filled (which is expensive as it involves a lot of shuffle instructions).
+ *
+ *  [ S T U V W X Y Z x x x ] Input Vector
+ *
+ *    0 1 2 3 4 5 6 7 8 9 0
+ *  0 A C E G I K M O x x x  |  S * A + T * C + U * E + V * G + W * I + X * K + Y * M + Z * O
+ *  1 A C E G I K M O x x x  |  S * A + T * C + U * E + V * G + W * I + X * K + Y * M + Z * O
+ *  2 A C E G I K M O x x x  |  S * A + T * C + U * E + V * G + W * I + X * K + Y * M + Z * O
+ *  3 A C E G I K M O x x x  |  S * A + T * C + U * E + V * G + W * I + X * K + Y * M + Z * O
+ *  4 B D F H J L N P x x x  |  S * B + T * D + U * F + V * H + W * J + X * L + Y * N + Z * P
+ *  5 B D F H J L N P x x x  |  S * B + T * D + U * F + V * H + W * J + X * L + Y * N + Z * P
+ *  6 B D F H J L N P x x x  |  S * B + T * D + U * F + V * H + W * J + X * L + Y * N + Z * P
+ *  7 B D F H J L N P x x x  |  S * B + T * D + U * F + V * H + W * J + X * L + Y * N + Z * P
+ *  8 x x x x x x x x x x x  |
+ *  9 x x x x x x x x x x x  |
+ * 10 x x x x x x x x x x x  |
+ * 11 x x x x x x x x x x x  |
+ *
+ */
+template<typename TFVec, typename T>
+static FORCE_INLINE void processBlocks8x4(const std::uint16_t* inputIndicesStart,
+                                          const std::uint16_t* inputIndicesEndAlignedTo4,
+                                          const std::uint16_t* inputIndicesEnd,
+                                          ConstArrayView<float> inputs,
+                                          const T*& values,
+                                          float* outbuf) {
+    TFVec sum1{};
+    TFVec sum2{};
+    TFVec sum3{};
+    TFVec sum4{};
+    TFVec sum5{};
+    TFVec sum6{};
+    TFVec sum7{};
+    TFVec sum8{};
+    for (const std::uint16_t* inputIndices = inputIndicesStart;
+         inputIndices < inputIndicesEndAlignedTo4;
+         inputIndices += 4ul, values += (8ul * TFVec::size())) {
+        const TFVec inputVec1{inputs[inputIndices[0]]};
+        const TFVec inputVec2{inputs[inputIndices[1]]};
+        const TFVec inputVec3{inputs[inputIndices[2]]};
+        const TFVec inputVec4{inputs[inputIndices[3]]};
+        const TFVec blk1 = TFVec::fromAlignedSource(values);
+        const TFVec blk2 = TFVec::fromAlignedSource(values + TFVec::size());
+        const TFVec blk3 = TFVec::fromAlignedSource(values + TFVec::size() * 2);
+        const TFVec blk4 = TFVec::fromAlignedSource(values + TFVec::size() * 3);
+        const TFVec blk5 = TFVec::fromAlignedSource(values + TFVec::size() * 4);
+        const TFVec blk6 = TFVec::fromAlignedSource(values + TFVec::size() * 5);
+        const TFVec blk7 = TFVec::fromAlignedSource(values + TFVec::size() * 6);
+        const TFVec blk8 = TFVec::fromAlignedSource(values + TFVec::size() * 7);
+        #ifndef RL_NO_MANUAL_PREFETCH
+            #ifdef RL_USE_NON_TEMPORAL_PREFETCH
+                TFVec::prefetchNTA(values + (lookAheadOffset * 6ul));
+                TFVec::prefetchNTA(values + (lookAheadOffset * 7ul));
+                TFVec::prefetchNTA(values + (lookAheadOffset * 8ul));
+                TFVec::prefetchNTA(values + (lookAheadOffset * 9ul));
+                TFVec::prefetchNTA(values + (lookAheadOffset * 10ul));
+                TFVec::prefetchNTA(values + (lookAheadOffset * 11ul));
+            #else
+                TFVec::prefetchT0(values + (lookAheadOffset * 6ul));
+                TFVec::prefetchT0(values + (lookAheadOffset * 7ul));
+                TFVec::prefetchT0(values + (lookAheadOffset * 8ul));
+                TFVec::prefetchT0(values + (lookAheadOffset * 9ul));
+                TFVec::prefetchT0(values + (lookAheadOffset * 10ul));
+                TFVec::prefetchT0(values + (lookAheadOffset * 11ul));
+            #endif
+        #endif
+        sum1 += (blk1 * inputVec1);
+        sum2 += (blk2 * inputVec1);
+        sum3 += (blk3 * inputVec2);
+        sum4 += (blk4 * inputVec2);
+        sum5 += (blk5 * inputVec3);
+        sum6 += (blk6 * inputVec3);
+        sum7 += (blk7 * inputVec4);
+        sum8 += (blk8 * inputVec4);
+    }
+    // Process 8x1 horizontal remainder portion after 8x4 blocks are consumed
+    processBlocks8x1(inputIndicesEndAlignedTo4, inputIndicesEnd, inputs, values, sum1, sum2);
+
+    sum1 += sum3;
+    sum2 += sum4;
+    sum5 += sum7;
+    sum6 += sum8;
+    sum1 += sum5;
+    sum2 += sum6;
+
+    sum1.alignedStore(outbuf);
+    sum2.alignedStore(outbuf + TFVec::size());
+}
+
+/*
+ * Process the remainder portion after 4x8 blocks
+ *
+ * Consumes 4 x 1 values per input index, but input indexes are processed one
+ * by one. The maximum number of iterations performed is 7, as anything larger
+ * would be already consumed by the 4x8 block processor. It's executed only
+ * after the 4x8 block processor, in cases when the matrix column count is not
+ * divisible by 8.
+ *
+ *  [ x x x x x x x x T U V W X Y Z ] Input Vector
+ *
+ *    0 1 2 3 4 5 6 7 8 9 0 1 2 3 4
+ *  0 x x x x x x x x x x x x x x x  |
+ *  1 x x x x x x x x x x x x x x x  |
+ *  2 x x x x x x x x x x x x x x x  |
+ *  3 x x x x x x x x x x x x x x x  |
+ *  4 x x x x x x x x x x x x x x x  |
+ *  5 x x x x x x x x x x x x x x x  |
+ *  6 x x x x x x x x x x x x x x x  |
+ *  7 x x x x x x x x x x x x x x x  |
+ *  8 x x x x x x x x A B C D E F G  |  T * A + U * B + V * C + W * D + X * E + Y * F + Z * G
+ *  9 x x x x x x x x A B C D E F G  |  T * A + U * B + V * C + W * D + X * E + Y * F + Z * G
+ * 10 x x x x x x x x A B C D E F G  |  T * A + U * B + V * C + W * D + X * E + Y * F + Z * G
+ * 11 x x x x x x x x A B C D E F G  |  T * A + U * B + V * C + W * D + X * E + Y * F + Z * G
+ *
+ */
+template<typename TFVec, typename T>
+static FORCE_INLINE void processBlocks4x1(const std::uint16_t* inputIndicesEndAlignedTo8,
+                                          const std::uint16_t* inputIndicesEnd,
+                                          ConstArrayView<float> inputs,
+                                          const T*& values,
+                                          TFVec& sum1) {
+    for (const std::uint16_t* inputIndices = inputIndicesEndAlignedTo8;
+         inputIndices < inputIndicesEnd;
+         ++inputIndices, values += TFVec::size()) {
+        const TFVec inputVec{inputs[*inputIndices]};
+        const TFVec blk = TFVec::fromAlignedSource(values);
+        sum1 += (blk * inputVec);
+    }
+}
+
+/*
+ * Process the vertical remainder portion of the matrix that's partitionable into 4x8 blocks
+ *
+ * Consumes 4 x 1 values per input index, and loads 8 input indexes per
+ * iteration, thus a single iteration processes 32 values. Although it
+ * processes the same number of values per iteration as the 8x4 block
+ * processor, it does so by having to fill 8 input vectors (one for each
+ * 4 x 1 value-block), thus requiring twice the number of shuffle instructions
+ * which makes it somewhat less performant than the 8x4 variant. It gets
+ * executed only in cases when the matrix row count is not divisible by 8.
+ *
+ *  [ S T U V W X Y Z x x x ] Input Vector
+ *
+ *    0 1 2 3 4 5 6 7 8 9 0
+ *  0 x x x x x x x x x x x  |
+ *  1 x x x x x x x x x x x  |
+ *  2 x x x x x x x x x x x  |
+ *  3 x x x x x x x x x x x  |
+ *  4 x x x x x x x x x x x  |
+ *  5 x x x x x x x x x x x  |
+ *  6 x x x x x x x x x x x  |
+ *  7 x x x x x x x x x x x  |
+ *  8 A B C D E F G H x x x  |  S * A + T * B + U * C + V * D + W * E + X * F + Y * G + Z * H
+ *  9 A B C D E F G H x x x  |  S * A + T * B + U * C + V * D + W * E + X * F + Y * G + Z * H
+ * 10 A B C D E F G H x x x  |  S * A + T * B + U * C + V * D + W * E + X * F + Y * G + Z * H
+ * 11 A B C D E F G H x x x  |  S * A + T * B + U * C + V * D + W * E + X * F + Y * G + Z * H
+ *
+ */
+template<typename TFVec, typename T>
+static FORCE_INLINE void processBlocks4x8(const std::uint16_t* inputIndicesStart,
+                                          const std::uint16_t* inputIndicesEndAlignedTo8,
+                                          const std::uint16_t* inputIndicesEnd,
+                                          ConstArrayView<float> inputs,
+                                          const T*& values,
+                                          float* outbuf) {
+    TFVec sum1{};
+    TFVec sum2{};
+    TFVec sum3{};
+    TFVec sum4{};
+    TFVec sum5{};
+    TFVec sum6{};
+    TFVec sum7{};
+    TFVec sum8{};
+    for (const std::uint16_t* inputIndices = inputIndicesStart;
+         inputIndices < inputIndicesEndAlignedTo8;
+         inputIndices += 8ul, values += (TFVec::size() * 8ul)) {
+        const TFVec inputVec1{inputs[inputIndices[0]]};
+        const TFVec inputVec2{inputs[inputIndices[1]]};
+        const TFVec inputVec3{inputs[inputIndices[2]]};
+        const TFVec inputVec4{inputs[inputIndices[3]]};
+        const TFVec inputVec5{inputs[inputIndices[4]]};
+        const TFVec inputVec6{inputs[inputIndices[5]]};
+        const TFVec inputVec7{inputs[inputIndices[6]]};
+        const TFVec inputVec8{inputs[inputIndices[7]]};
+        const TFVec blk1 = TFVec::fromAlignedSource(values);
+        const TFVec blk2 = TFVec::fromAlignedSource(values + TFVec::size());
+        const TFVec blk3 = TFVec::fromAlignedSource(values + TFVec::size() * 2);
+        const TFVec blk4 = TFVec::fromAlignedSource(values + TFVec::size() * 3);
+        const TFVec blk5 = TFVec::fromAlignedSource(values + TFVec::size() * 4);
+        const TFVec blk6 = TFVec::fromAlignedSource(values + TFVec::size() * 5);
+        const TFVec blk7 = TFVec::fromAlignedSource(values + TFVec::size() * 6);
+        const TFVec blk8 = TFVec::fromAlignedSource(values + TFVec::size() * 7);
+        sum1 += (blk1 * inputVec1);
+        sum2 += (blk2 * inputVec2);
+        sum3 += (blk3 * inputVec3);
+        sum4 += (blk4 * inputVec4);
+        sum5 += (blk5 * inputVec5);
+        sum6 += (blk6 * inputVec6);
+        sum7 += (blk7 * inputVec7);
+        sum8 += (blk8 * inputVec8);
+    }
+    // Process 4x1 horizontal remainder portion after 4x8 blocks are consumed
+    processBlocks4x1(inputIndicesEndAlignedTo8, inputIndicesEnd, inputs, values, sum1);
+
+    sum1 += sum2;
+    sum3 += sum4;
+    sum5 += sum6;
+    sum7 += sum8;
+    sum1 += sum3;
+    sum5 += sum7;
+    sum1 += sum5;
+
+    sum1.alignedStore(outbuf);
+}
+
+/*
+ * Orchestrate the execution of the needed block processors for a given joint group
+ */
+template<typename TFVec, typename T>
+static FORCE_INLINE void processJointGroupBlock4(const JointGroupView<T>& jointGroup, ConstArrayView<float> inputs,
+                                                 ArrayView<float> outputs, std::uint16_t lod) {
+    const T* values = jointGroup.values;
+    const std::uint16_t* const inputIndices = jointGroup.inputIndices;
+    const std::uint16_t* const inputIndicesEnd = jointGroup.inputIndicesEnd;
+    const std::uint16_t* const inputIndicesEndAlignedTo4 = jointGroup.inputIndicesEndAlignedTo4;
+    const std::uint16_t* const inputIndicesEndAlignedTo8 = jointGroup.inputIndicesEndAlignedTo8;
+    const LODRegion& lodRegion = jointGroup.lods[lod];
+    const std::uint16_t* outputIndices = jointGroup.outputIndices;
+    const std::uint16_t* const outputIndicesEnd = outputIndices + lodRegion.size;
+    const std::uint16_t* const outputIndicesEndPaddedToLastFullBlock = outputIndices + lodRegion.sizePaddedToLastFullBlock;
+    const std::uint16_t* const outputIndicesEndPaddedToSecondLastFullBlock = outputIndices +
+        lodRegion.sizePaddedToSecondLastFullBlock;
+    constexpr std::size_t halfBlockHeight = TFVec::size();
+    constexpr std::size_t fullBlockHeight = 2ul * TFVec::size();
+    // Process portion of matrix that's partitionable into 8x4 blocks
+    for (; outputIndices < outputIndicesEndPaddedToSecondLastFullBlock; outputIndices += fullBlockHeight) {
+        alignas(TFVec::alignment()) float outbuf[fullBlockHeight];
+        processBlocks8x4<TFVec>(inputIndices, inputIndicesEndAlignedTo4, inputIndicesEnd, inputs, values,
+                                static_cast<float*>(outbuf));
+        for (std::size_t i = 0ul; i < fullBlockHeight; ++i) {
+            outputs[outputIndices[i]] = outbuf[i];
+        }
+    }
+    // Process the last 8x4 block which needs special handling when some of
+    // the output values need to be masked-off because of LODs
+    for (; outputIndices < outputIndicesEndPaddedToLastFullBlock; outputIndices += fullBlockHeight) {
+        alignas(TFVec::alignment()) float outbuf[fullBlockHeight];
+        processBlocks8x4<TFVec>(inputIndices, inputIndicesEndAlignedTo4, inputIndicesEnd, inputs, values,
+                                static_cast<float*>(outbuf));
+        // Ignore results that came from rows after the last LOD row
+        for (std::size_t i = 0ul; i < (lodRegion.size % fullBlockHeight); ++i) {
+            // NOLINTNEXTLINE(cppcoreguidelines-pro-bounds-constant-array-index)
+            outputs[outputIndices[i]] = outbuf[i];
+        }
+    }
+    // Process vertical remainder portion of matrix that's partitionable into 4x8 blocks
+    for (; outputIndices < outputIndicesEnd; outputIndices += halfBlockHeight) {
+        alignas(TFVec::alignment()) float outbuf[halfBlockHeight];
+        processBlocks4x8<TFVec>(inputIndices, inputIndicesEndAlignedTo8, inputIndicesEnd, inputs, values,
+                                static_cast<float*>(outbuf));
+        // Ignore results that came from rows after the last LOD row
+        auto maskOffStart = static_cast<std::size_t>(outputIndicesEnd - outputIndices);
+        for (std::size_t i = 0; i < maskOffStart; ++i) {
+            // NOLINTNEXTLINE(cppcoreguidelines-pro-bounds-constant-array-index)
+            outputs[outputIndices[i]] = outbuf[i];
+        }
+    }
+}
+
+template<typename T>
+struct JointGroupLinearCalculationStrategy {
+    virtual ~JointGroupLinearCalculationStrategy() = default;
+
+    virtual void calculate(const JointGroupView<T>& jointGroup,
+                           ConstArrayView<float> inputs,
+                           ArrayView<float> outputs,
+                           std::uint16_t lod) const = 0;
+
+};
+
+template<typename T, typename TFVec, class TRotationAdapter>
+struct VectorizedJointGroupLinearCalculationStrategy : public JointGroupLinearCalculationStrategy<T> {
+
+    void calculate(const JointGroupView<T>& jointGroup, ConstArrayView<float> inputs, ArrayView<float> outputs,
+                   std::uint16_t lod) const override {
+        processJointGroupBlock4<TFVec>(jointGroup, inputs, outputs, lod);
+        TRotationAdapter::adapt(jointGroup, outputs, lod);
+    }
+
+};
+
+}  // namespace bpcm
+
+}  // namespace rl4
diff --git a/Engine/Plugins/Animation/RigLogic/Source/RigLogicLib/Private/riglogic/joints/cpu/bpcm/Consts.h b/Engine/Plugins/Animation/RigLogic/Source/RigLogicLib/Private/riglogic/joints/cpu/bpcm/Consts.h
new file mode 100644
--- /dev/null
+++ b/Engine/Plugins/Animation/RigLogic/Source/RigLogicLib/Private/riglogic/joints/cpu/bpcm/Consts.h
@@ -0,0 +1,17 @@
+// Copyright Epic Games, Inc. All Rights Reserved.
+
+#pragma once
+
+#include <cstddef>
+#include <cstdint>
+
+namespace rl4 {
+
+namespace bpcm {
+
+constexpr std::size_t cacheLineSize = 64ul;  // bytes
+constexpr std::size_t lookAheadOffset = cacheLineSize / sizeof(float);  // in number of elements
+
+}  // namespace bpcm
+
+}  // namespace rl4
diff --git a/Engine/Plugins/Animation/RigLogic/Source/RigLogicLib/Private/riglogic/joints/cpu/bpcm/StorageSize.cpp b/Engine/Plugins/Animation/RigLogic/Source/RigLogicLib/Private/riglogic/joints/cpu/bpcm/StorageSize.cpp
new file mode 100644
--- /dev/null
+++ b/Engine/Plugins/Animation/RigLogic/Source/RigLogicLib/Private/riglogic/joints/cpu/bpcm/StorageSize.cpp
@@ -0,0 +1,60 @@
+// Copyright Epic Games, Inc. All Rights Reserved.
+
+#include "riglogic/joints/cpu/bpcm/StorageSize.h"
+
+#include "riglogic/TypeDefs.h"
+#include "riglogic/types/Extent.h"
+#include "riglogic/utils/Extd.h"
+#include "riglogic/joints/JointBehaviorFilter.h"
+
+#include <cassert>
+#include <cstddef>
+
+namespace rl4 {
+
+namespace bpcm {
+
+StorageSize::StorageSize(MemoryResource* memRes) :
+    valueCount{},
+    inputIndexCount{},
+    outputIndexCount{},
+    lodRegionCount{},
+    lodCount{},
+    jointGroups{memRes} {
+}
+
+void StorageSize::computeFrom(const JointBehaviorFilter& src, std::size_t padTo) {
+    const auto jointGroupCount = src.getJointGroupCount();
+    jointGroups.clear();
+    jointGroups.reserve(jointGroupCount);
+    lodCount = src.getLODCount();
+    for (std::uint16_t i = 0u; i < jointGroupCount; ++i) {
+        const auto columnCount = static_cast<std::size_t>(src.getColumnCount(i));
+        const auto rowCount = static_cast<std::size_t>(src.getRowCount(i));
+        const std::size_t padding = extd::roundUp(rowCount, padTo) - rowCount;
+        valueCount += (rowCount * columnCount) + (padding * columnCount);
+        inputIndexCount += columnCount;
+        outputIndexCount += (rowCount + padding);
+        lodRegionCount += lodCount;
+        #if !defined(__clang__) && defined(__GNUC__)
+            #pragma GCC diagnostic push
+            #pragma GCC diagnostic ignored "-Wuseless-cast"
+        #endif
+        jointGroups.push_back({
+                    {static_cast<std::uint32_t>(rowCount), static_cast<std::uint32_t>(columnCount)},
+                    {static_cast<std::uint32_t>(rowCount + padding), static_cast<std::uint32_t>(columnCount)}
+                });
+        #if !defined(__clang__) && defined(__GNUC__)
+            #pragma GCC diagnostic pop
+        #endif
+    }
+}
+
+JointGroupSize StorageSize::getJointGroupSize(std::size_t jointGroupIndex) const {
+    assert(jointGroupIndex < jointGroups.size());
+    return jointGroups[jointGroupIndex];
+}
+
+}  // namespace bpcm
+
+}  // namespace rl4
diff --git a/Engine/Plugins/Animation/RigLogic/Source/RigLogicLib/Private/riglogic/joints/cpu/bpcm/StorageSize.h b/Engine/Plugins/Animation/RigLogic/Source/RigLogicLib/Private/riglogic/joints/cpu/bpcm/StorageSize.h
new file mode 100644
--- /dev/null
+++ b/Engine/Plugins/Animation/RigLogic/Source/RigLogicLib/Private/riglogic/joints/cpu/bpcm/StorageSize.h
@@ -0,0 +1,37 @@
+// Copyright Epic Games, Inc. All Rights Reserved.
+
+#pragma once
+
+#include "riglogic/TypeDefs.h"
+#include "riglogic/types/Extent.h"
+
+#include <cstddef>
+
+namespace rl4 {
+
+class JointBehaviorFilter;
+
+namespace bpcm {
+
+struct JointGroupSize {
+    Extent original;
+    Extent padded;
+};
+
+struct StorageSize {
+    std::size_t valueCount;
+    std::size_t inputIndexCount;
+    std::size_t outputIndexCount;
+    std::size_t lodRegionCount;
+    std::size_t lodCount;
+    Vector<JointGroupSize> jointGroups;
+
+    explicit StorageSize(MemoryResource* memRes);
+
+    void computeFrom(const JointBehaviorFilter& src, std::size_t padTo);
+    JointGroupSize getJointGroupSize(std::size_t jointGroupIndex) const;
+};
+
+}  // namespace bpcm
+
+}  // namespace rl4
diff --git a/Engine/Plugins/Animation/RigLogic/Source/RigLogicLib/Private/riglogic/joints/cpu/bpcm/StorageView.h b/Engine/Plugins/Animation/RigLogic/Source/RigLogicLib/Private/riglogic/joints/cpu/bpcm/StorageView.h
new file mode 100644
--- /dev/null
+++ b/Engine/Plugins/Animation/RigLogic/Source/RigLogicLib/Private/riglogic/joints/cpu/bpcm/StorageView.h
@@ -0,0 +1,47 @@
+// Copyright Epic Games, Inc. All Rights Reserved.
+
+#pragma once
+
+#include "riglogic/TypeDefs.h"
+#include "riglogic/joints/cpu/bpcm/Storage.h"
+
+#include <cstdint>
+
+namespace rl4 {
+
+namespace bpcm {
+
+template<typename TValue>
+struct JointGroupView {
+    TValue* values;
+    std::uint16_t* inputIndices;
+    std::uint16_t* inputIndicesEnd;
+    std::uint16_t* inputIndicesEndAlignedTo4;
+    std::uint16_t* inputIndicesEndAlignedTo8;
+    std::uint16_t* outputIndices;
+    LODRegion* lods;
+    std::uint16_t* outputRotationIndices;
+    std::uint16_t* outputRotationLODs;
+};
+
+template<typename TValue>
+Vector<JointGroupView<TValue> > takeStorageSnapshot(JointStorage<TValue>& storage, MemoryResource* memRes) {
+    Vector<JointGroupView<TValue> > snapshot{storage.jointGroups.size(), {}, memRes};
+    for (std::size_t i = 0ul; i < storage.jointGroups.size(); ++i) {
+        const auto& jointGroup = storage.jointGroups[i];
+        snapshot[i].values = storage.values.data() + jointGroup.valuesOffset;
+        snapshot[i].inputIndices = storage.inputIndices.data() + jointGroup.inputIndicesOffset;
+        snapshot[i].inputIndicesEnd = snapshot[i].inputIndices + jointGroup.inputIndicesSize;
+        snapshot[i].inputIndicesEndAlignedTo4 = snapshot[i].inputIndices + jointGroup.inputIndicesSizeAlignedTo4;
+        snapshot[i].inputIndicesEndAlignedTo8 = snapshot[i].inputIndices + jointGroup.inputIndicesSizeAlignedTo8;
+        snapshot[i].outputIndices = storage.outputIndices.data() + jointGroup.outputIndicesOffset;
+        snapshot[i].lods = storage.lodRegions.data() + jointGroup.lodsOffset;
+        snapshot[i].outputRotationIndices = storage.outputRotationIndices.data() + jointGroup.outputRotationIndicesOffset;
+        snapshot[i].outputRotationLODs = storage.outputRotationLODs.data() + jointGroup.outputRotationLODsOffset;
+    }
+    return snapshot;
+}
+
+}  // namespace bpcm
+
+}  // namespace rl4
diff --git a/Engine/Plugins/Animation/RigLogic/Source/RigLogicLib/Private/riglogic/joints/cpu/quaternions/CalculationStrategy.h b/Engine/Plugins/Animation/RigLogic/Source/RigLogicLib/Private/riglogic/joints/cpu/quaternions/CalculationStrategy.h
new file mode 100644
--- /dev/null
+++ b/Engine/Plugins/Animation/RigLogic/Source/RigLogicLib/Private/riglogic/joints/cpu/quaternions/CalculationStrategy.h
@@ -0,0 +1,338 @@
+// Copyright Epic Games, Inc. All Rights Reserved.
+
+#pragma once
+
+#include "riglogic/system/simd/SIMD.h"
+#include "riglogic/types/Aliases.h"
+#include "riglogic/utils/Macros.h"
+#include "riglogic/joints/cpu/quaternions/JointGroup.h"
+
+#include <cstddef>
+#include <cstdint>
+
+namespace rl4 {
+
+/*
+ * FastLerp between 4 arbitrary and 4 identity quaternions
+ *
+ * fastlerp_with_identity(Q, WEIGHT):
+ *     IQ = [0, 0, 0, 1]
+ *     // Dot product between IQ and Q
+ *     D1 = IQ.x * Q.x + IQ.y * Q.y + IQ.z * Q.z + IQ.w * Q.w;
+ *     D1 = 0 * Q.x + 0 * Q.y + 0 * Q.z + 1 * Q.w
+ *     D1 = Q.w
+ *     // Select bias based on dot product
+ *     BIAS1 = D1 >= 0 ? 1.0 : -1.0
+ *     AW = BIAS1 * (1.0 - WEIGHT)
+ *     // Because:
+ *     // LERPED = (Q * WEIGHT) + (IQ * AW)
+ *     // And:
+ *     // IQ.x * AW = 0 * AW
+ *     // IQ.y * AW = 0 * AW
+ *     // IQ.z * AW = 0 * AW
+ *     // IQ.w * AW = 1 * AW
+ *     // It follows that:
+ *     LERPED = [
+ *         Q.x * WEIGHT
+ *         Q.y * WEIGHT
+ *         Q.z * WEIGHT
+ *         Q.w * WEIGHT + AW
+ *     ]
+ *
+ * ------------------------------
+ *
+ * fastlerp_with_identity_simplified(Q, WEIGHT):
+ *     BIAS1 = Q.w >= 0 ? 1.0 : -1.0
+ *     AW = BIAS1 * (1.0 - WEIGHT)
+ *     LERPED = [
+ *         Q.x * WEIGHT
+ *         Q.y * WEIGHT
+ *         Q.z * WEIGHT
+ *         Q.w * WEIGHT + AW
+ *     ]
+ *
+ * ------------------------------
+ *
+ * A, C, E, G are quaternions representing delta rotations as an array of 4 floats [x, y, z, w].
+ *
+ * fastlerp_with_identity_simd(A, C, E, G, WEIGHT):
+ *     NEG_ONES = [-1, -1, -1, -1]
+ *     ONES = [1, 1, 1, 1]
+ *     MASK = [
+ *         A.w >= 0
+ *         C.w >= 0
+ *         E.w >= 0
+ *         G.w >= 0
+ *     ]
+ *     BIAS = ANDNOT(MASK, NEG_ONES) | (ONES & MASK)
+ *     ONE_MINUS_WEIGHT = [
+ *         (1.0 - WEIGHT)
+ *         (1.0 - WEIGHT)
+ *         (1.0 - WEIGHT)
+ *         (1.0 - WEIGHT)
+ *     ]
+ *     AWS = BIAS * ONE_MINUS_WEIGHT
+ *     WEIGHTS = [WEIGHT, WEIGHT, WEIGHT, WEIGHT]
+ *     LERPED_X = [
+ *         A.x * WEIGHTS
+ *         C.x * WEIGHTS
+ *         E.x * WEIGHTS
+ *         G.x * WEIGHTS
+ *     ]
+ *     LERPED_Y = [
+ *         A.y * WEIGHTS
+ *         C.y * WEIGHTS
+ *         E.y * WEIGHTS
+ *         G.y * WEIGHTS
+ *     ]
+ *     LERPED_Z = [
+ *         A.z * WEIGHTS
+ *         C.z * WEIGHTS
+ *         E.z * WEIGHTS
+ *         G.z * WEIGHTS
+ *     ]
+ *     LERPED_W = [
+ *         A.w * WEIGHTS + AWS
+ *         C.w * WEIGHTS + AWS
+ *         E.w * WEIGHTS + AWS
+ *         G.w * WEIGHTS + AWS
+ *     ]
+ *     LERPED_A = [LERPED_X[0], LERPED_Y[0], LERPED_Z[0], LERPED_W[0]]
+ *     LERPED_C = [LERPED_X[1], LERPED_Y[1], LERPED_Z[1], LERPED_W[1]]
+ *     LERPED_E = [LERPED_X[2], LERPED_Y[2], LERPED_Z[2], LERPED_W[2]]
+ *     LERPED_G = [LERPED_X[3], LERPED_Y[3], LERPED_Z[3], LERPED_W[3]]
+ */
+template<typename TFVec>
+static FORCE_INLINE void fastlerpWithIdentity(TFVec& qxACEG, TFVec& qyACEG, TFVec& qzACEG, TFVec& qwACEG, TFVec weight) {
+    const TFVec negOne{-1.0f};
+    const TFVec one{1.0f};
+    const TFVec oneMinusWeight = one - weight;
+    const TFVec mask = (qwACEG >= TFVec{0.0f});
+    const TFVec bias = trimd::andnot(mask, negOne) | (one & mask);
+    const TFVec aws = bias * oneMinusWeight;
+    qxACEG = qxACEG * weight;
+    qyACEG = qyACEG * weight;
+    qzACEG = qzACEG * weight;
+    qwACEG = qwACEG * weight + aws;
+}
+
+/*
+ * Normalize quaternion
+ *
+ * normalize(Q):
+ *     SS = Q.x * Q.x + Q.y * Q.y + Q.z * Q.z + Q.w * Q.w
+ *     INVSS = 1.0 / SQRT(SS)
+ *     NORMALIZED = [
+ *         Q.x *= INVSS
+ *         Q.y *= INVSS
+ *         Q.z *= INVSS
+ *         Q.w *= INVSS
+ *     ]
+ *
+ * normalize_simd(A, C, E, G):
+ *     XS_SQUARED = [
+ *         A.x * A.x
+ *         C.x * C.x
+ *         E.x * E.x
+ *         G.x * G.x
+ *     ]
+ *     YS_SQUARED = [
+ *         A.y * A.y
+ *         C.y * C.y
+ *         E.y * E.y
+ *         G.y * G.y
+ *     ]
+ *     ZS_SQUARED = [
+ *         A.z * A.z
+ *         C.z * C.z
+ *         E.z * E.z
+ *         G.z * G.z
+ *     ]
+ *     WS_SQUARED = [
+ *         A.w * A.w
+ *         C.w * C.w
+ *         E.w * E.w
+ *         G.w * G.w
+ *     ]
+ *     S_SUM = XS_SQUARED + YS_SQUARED + ZS_SQUARED + WS_SQUARED
+ *     INVSS = 1.0 / SQRT(S_SUM)  // _mm_rsqrt_ps(S_SUM)
+ *     NORMALIZED_A = [
+ *         A.x * INVSS
+ *         A.y * INVSS
+ *         A.z * INVSS
+ *         A.w * INVSS
+ *     ]
+ *     NORMALIZED_B = [
+ *         B.x * INVSS
+ *         B.y * INVSS
+ *         B.z * INVSS
+ *         B.w * INVSS
+ *     ]
+ *     NORMALIZED_C = [
+ *         C.x * INVSS
+ *         C.y * INVSS
+ *         C.z * INVSS
+ *         C.w * INVSS
+ *     ]
+ *     NORMALIZED_D = [
+ *         D.x * INVSS
+ *         D.y * INVSS
+ *         D.z * INVSS
+ *         D.w * INVSS
+ *     ]
+ */
+template<typename TFVec>
+static FORCE_INLINE void normalize(TFVec& qxACEG, TFVec& qyACEG, TFVec& qzACEG, TFVec& qwACEG) {
+    const TFVec xsSquared = qxACEG * qxACEG;
+    const TFVec ysSquared = qyACEG * qyACEG;
+    const TFVec zsSquared = qzACEG * qzACEG;
+    const TFVec wsSquared = qwACEG * qwACEG;
+    const TFVec sSum = xsSquared + ysSquared + zsSquared + wsSquared;
+    const TFVec rSqrtSum = trimd::rsqrt(sSum);
+    qxACEG *= rSqrtSum;
+    qyACEG *= rSqrtSum;
+    qzACEG *= rSqrtSum;
+    qwACEG *= rSqrtSum;
+}
+
+/*
+ * Multiply each quaternion of a joint group by the expression weight, and combine them together into a single rotation
+ *
+ * Rotation = FastLerp(IdentityQuaternion, PoseQuaternion, ExpressionWeight);
+ * Rotation = Normalize(Rotation);
+ * OutputRotation = Rotation * OutputRotation;  // OutputRotation is the accumulator
+ *
+ * A, C, E, G are quaternions representing delta rotations as an array of 4 floats [x, y, z, w].
+ * B, D, F, H are the accumulator quaternions that will contain the combined, blended rotations.
+ *
+ * The below operation is how a freshly blended set of 4 quaternions are combined with the accumulator
+ * quaternions. This first set of multiplications is how each pair is multiplied individually:
+ *
+ * A[3]*B[0] + A[0]*B[3] + A[1]*B[2] - A[2]*B[1];
+ * A[3]*B[1] - A[0]*B[2] + A[1]*B[3] + A[2]*B[0];
+ * A[3]*B[2] + A[0]*B[1] - A[1]*B[0] + A[2]*B[3];
+ * A[3]*B[3] - A[0]*B[0] - A[1]*B[1] - A[2]*B[2];
+ *
+ * C[3]*D[0] + C[0]*D[3] + C[1]*D[2] - C[2]*D[1];
+ * C[3]*D[1] - C[0]*D[2] + C[1]*D[3] + C[2]*D[0];
+ * C[3]*D[2] + C[0]*D[1] - C[1]*D[0] + C[2]*D[3];
+ * C[3]*D[3] - C[0]*D[0] - C[1]*D[1] - C[2]*D[2];
+ *
+ * E[3]*F[0] + E[0]*F[3] + E[1]*F[2] - E[2]*F[1];
+ * E[3]*F[1] - E[0]*F[2] + E[1]*F[3] + E[2]*F[0];
+ * E[3]*F[2] + E[0]*F[1] - E[1]*F[0] + E[2]*F[3];
+ * E[3]*F[3] - E[0]*F[0] - E[1]*F[1] - E[2]*F[2];
+ *
+ * G[3]*H[0] + G[0]*H[3] + G[1]*H[2] - G[2]*H[1];
+ * G[3]*H[1] - G[0]*H[2] + G[1]*H[3] + G[2]*H[0];
+ * G[3]*H[2] + G[0]*H[1] - G[1]*H[0] + G[2]*H[3];
+ * G[3]*H[3] - G[0]*H[0] - G[1]*H[1] - G[2]*H[2];
+ *
+ * And the below transformation is the same as the above shown block, only the operations are
+ * interleaved so 4 quaternions can be processed with 1 set of SIMD registers efficiently.
+ *
+ * A[3]*B[0] + A[0]*B[3] + A[1]*B[2] - A[2]*B[1]
+ * C[3]*D[0] + C[0]*D[3] + C[1]*D[2] - C[2]*D[1]
+ * E[3]*F[0] + E[0]*F[3] + E[1]*F[2] - E[2]*F[1]
+ * G[3]*H[0] + G[0]*H[3] + G[1]*H[2] - G[2]*H[1]
+ * ---------------------------------------------
+ * A[3]*B[1] - A[0]*B[2] + A[1]*B[3] + A[2]*B[0]
+ * C[3]*D[1] - C[0]*D[2] + C[1]*D[3] + C[2]*D[0]
+ * E[3]*F[1] - E[0]*F[2] + E[1]*F[3] + E[2]*F[0]
+ * G[3]*H[1] - G[0]*H[2] + G[1]*H[3] + G[2]*H[0]
+ * ---------------------------------------------
+ * A[3]*B[2] + A[0]*B[1] - A[1]*B[0] + A[2]*B[3]
+ * C[3]*D[2] + C[0]*D[1] - C[1]*D[0] + C[2]*D[3]
+ * E[3]*F[2] + E[0]*F[1] - E[1]*F[0] + E[2]*F[3]
+ * G[3]*H[2] + G[0]*H[1] - G[1]*H[0] + G[2]*H[3]
+ * ---------------------------------------------
+ * A[3]*B[3] - A[0]*B[0] - A[1]*B[1] - A[2]*B[2]
+ * C[3]*D[3] - C[0]*D[0] - C[1]*D[1] - C[2]*D[2]
+ * E[3]*F[3] - E[0]*F[0] - E[1]*F[1] - E[2]*F[2]
+ * G[3]*H[3] - G[0]*H[0] - G[1]*H[1] - G[2]*H[2]
+ */
+template<typename TFVec, typename T>
+static FORCE_INLINE void blendQuaternions(const T*& quaternions,
+                                          ConstArrayView<std::uint16_t> inputIndices,
+                                          ConstArrayView<float> inputs,
+                                          float* outbuf) {
+    // Initialize accumulators to identity quaternion
+    TFVec qxBDFH{0.0f};
+    TFVec qyBDFH{0.0f};
+    TFVec qzBDFH{0.0f};
+    TFVec qwBDFH{1.0f};
+    for (std::size_t col = {}; col < inputIndices.size(); ++col, quaternions += (TFVec::size() * 4)) {
+        TFVec qxACEG = TFVec::fromAlignedSource(quaternions);
+        TFVec qyACEG = TFVec::fromAlignedSource(quaternions + TFVec::size());
+        TFVec qzACEG = TFVec::fromAlignedSource(quaternions + TFVec::size() * 2);
+        TFVec qwACEG = TFVec::fromAlignedSource(quaternions + TFVec::size() * 3);
+        // Blend qxACEG, qyACEG, qzACEG, qwACEG with input weights
+        const TFVec weights{inputs[inputIndices[col]]};
+        fastlerpWithIdentity(qxACEG, qyACEG, qzACEG, qwACEG, weights);
+        normalize(qxACEG, qyACEG, qzACEG, qwACEG);
+        // Accumulate blended in q[xyzw]Tmp
+        const TFVec qxTmp = (qwACEG * qxBDFH) + (qxACEG * qwBDFH) + (qyACEG * qzBDFH) - (qzACEG * qyBDFH);
+        const TFVec qyTmp = (qwACEG * qyBDFH) - (qxACEG * qzBDFH) + (qyACEG * qwBDFH) + (qzACEG * qxBDFH);
+        const TFVec qzTmp = (qwACEG * qzBDFH) + (qxACEG * qyBDFH) - (qyACEG * qxBDFH) + (qzACEG * qwBDFH);
+        const TFVec qwTmp = (qwACEG * qwBDFH) - (qxACEG * qxBDFH) - (qyACEG * qyBDFH) - (qzACEG * qzBDFH);
+        qxBDFH = qxTmp;
+        qyBDFH = qyTmp;
+        qzBDFH = qzTmp;
+        qwBDFH = qwTmp;
+    }
+    qxBDFH.alignedStore(outbuf);
+    qyBDFH.alignedStore(outbuf + TFVec::size());
+    qzBDFH.alignedStore(outbuf + TFVec::size() * 2);
+    qwBDFH.alignedStore(outbuf + TFVec::size() * 3);
+}
+
+template<typename T>
+struct JointGroupQuaternionCalculationStrategy {
+    virtual ~JointGroupQuaternionCalculationStrategy() = default;
+    virtual void calculate(const JointGroup<T>& jointGroup,
+                           ConstArrayView<float> inputs,
+                           ArrayView<float> outputs,
+                           std::uint16_t lod) = 0;
+
+};
+
+template<typename T, typename TFVec256, typename TFVec128, class TRotationAdapter>
+struct VectorizedJointGroupQuaternionCalculationStrategy : public JointGroupQuaternionCalculationStrategy<T> {
+
+    void calculate(const JointGroup<T>& jointGroup, ConstArrayView<float> inputs, ArrayView<float> outputs,
+                   std::uint16_t lod) override {
+        const T* quaternions = jointGroup.values.data();
+        ConstArrayView<std::uint16_t> inputIndices{jointGroup.inputIndices};
+        const std::uint16_t* outputIndices = jointGroup.outputIndices.data();
+        const LODRegion& lodRegion = jointGroup.lods[lod];
+        const std::uint16_t* const outputIndicesEnd = outputIndices + lodRegion.size;
+        const std::uint16_t* const outputIndicesEndPaddedToLastFullBlock = outputIndices + lodRegion.sizePaddedToLastFullBlock;
+        const std::uint16_t* const outputIndicesEndPaddedToSecondLastFullBlock = outputIndices +
+            lodRegion.sizePaddedToSecondLastFullBlock;
+
+        for (; outputIndices < outputIndicesEndPaddedToSecondLastFullBlock; outputIndices += (TFVec256::size() * 4)) {
+            alignas(TFVec256::alignment()) float outbuf[TFVec256::size() * 4];
+            blendQuaternions<TFVec256>(quaternions, inputIndices, inputs, static_cast<float*>(outbuf));
+            TRotationAdapter::forward(outbuf, TFVec256::size(), TFVec256::size(), outputIndices, outputs);
+        }
+
+        for (; outputIndices < outputIndicesEndPaddedToLastFullBlock; outputIndices += (TFVec256::size() * 4)) {
+            alignas(TFVec256::alignment()) float outbuf[TFVec256::size() * 4];
+            blendQuaternions<TFVec256>(quaternions, inputIndices, inputs, static_cast<float*>(outbuf));
+            // Ignore results that came from rows after the last LOD row
+            const auto quaternionCount = (lodRegion.size % (TFVec256::size() * 4)) / 4;
+            TRotationAdapter::forward(outbuf, quaternionCount, TFVec256::size(), outputIndices, outputs);
+        }
+
+        for (; outputIndices < outputIndicesEnd; outputIndices += (TFVec128::size() * 4)) {
+            alignas(TFVec128::alignment()) float outbuf[TFVec128::size() * 4];
+            blendQuaternions<TFVec128>(quaternions, inputIndices, inputs, static_cast<float*>(outbuf));
+            // Ignore results that came from rows after the last LOD row
+            const auto quaternionCount = static_cast<std::size_t>(outputIndicesEnd - outputIndices) / 4;
+            TRotationAdapter::forward(outbuf, quaternionCount, TFVec128::size(), outputIndices, outputs);
+        }
+    }
+
+};
+
+}  // namespace rl4
diff --git a/Engine/Plugins/Animation/RigLogic/Source/RigLogicLib/Private/riglogic/joints/cpu/quaternions/JointGroup.h b/Engine/Plugins/Animation/RigLogic/Source/RigLogicLib/Private/riglogic/joints/cpu/quaternions/JointGroup.h
new file mode 100644
--- /dev/null
+++ b/Engine/Plugins/Animation/RigLogic/Source/RigLogicLib/Private/riglogic/joints/cpu/quaternions/JointGroup.h
@@ -0,0 +1,39 @@
+// Copyright Epic Games, Inc. All Rights Reserved.
+
+#pragma once
+
+#include "riglogic/TypeDefs.h"
+#include "riglogic/types/PaddedBlockView.h"
+
+#include <cstdint>
+
+namespace rl4 {
+
+using LODRegion = PaddedBlockView;
+
+template<typename TValue>
+struct JointGroup {
+    // All non-zero values
+    AlignedVector<TValue> values;
+    // Sub-matrix col -> input vector
+    AlignedVector<std::uint16_t> inputIndices;
+    // Sub-matrix row -> output vector
+    AlignedVector<std::uint16_t> outputIndices;
+    // Output index boundaries for each LOD
+    Vector<LODRegion> lods;
+
+    explicit JointGroup(MemoryResource* memRes) :
+        values{memRes},
+        inputIndices{memRes},
+        outputIndices{memRes},
+        lods{memRes} {
+    }
+
+    template<class Archive>
+    void serialize(Archive& archive) {
+        archive(values, inputIndices, outputIndices, lods);
+    }
+
+};
+
+}  // namespace rl4
diff --git a/Engine/Plugins/Animation/RigLogic/Source/RigLogicLib/Private/riglogic/joints/cpu/quaternions/QuaternionJointsBuilder.h b/Engine/Plugins/Animation/RigLogic/Source/RigLogicLib/Private/riglogic/joints/cpu/quaternions/QuaternionJointsBuilder.h
new file mode 100644
--- /dev/null
+++ b/Engine/Plugins/Animation/RigLogic/Source/RigLogicLib/Private/riglogic/joints/cpu/quaternions/QuaternionJointsBuilder.h
@@ -0,0 +1,370 @@
+// Copyright Epic Games, Inc. All Rights Reserved.
+
+#pragma once
+
+#include "riglogic/TypeDefs.h"
+#include "riglogic/joints/JointBehaviorFilter.h"
+#include "riglogic/joints/JointsBuilder.h"
+#include "riglogic/joints/cpu/quaternions/QuaternionJointsEvaluator.h"
+#include "riglogic/joints/cpu/quaternions/RotationAdapters.h"
+#include "riglogic/riglogic/Configuration.h"
+#include "riglogic/types/bpcm/Optimizer.h"
+#include "riglogic/utils/Extd.h"
+
+#include <tdm/Quat.h>
+
+namespace rl4 {
+
+template<typename TValue, typename TFVec256, typename TFVec128>
+class QuaternionJointsBuilder : public JointsBuilder {
+    public:
+        QuaternionJointsBuilder(const Configuration& config_, MemoryResource* memRes_);
+
+        void computeStorageRequirements(const RigMetrics& source) override;
+        void computeStorageRequirements(const JointBehaviorFilter& source) override;
+        void allocateStorage(const JointBehaviorFilter& source) override;
+        void fillStorage(const JointBehaviorFilter& source) override;
+        JointsEvaluator::Pointer build() override;
+
+    private:
+        void setInputIndices(const JointBehaviorFilter& source, std::uint16_t jointGroupIndex);
+        void setOutputIndices(const JointBehaviorFilter& source, std::uint16_t jointGroupIndex);
+        void setValues(const JointBehaviorFilter& source, std::uint16_t jointGroupIndex);
+        void setLODs(const JointBehaviorFilter& source, std::uint16_t jointGroupIndex);
+        void remapOutputIndices(std::uint16_t jointGroupIndex);
+
+        static constexpr std::uint32_t BlockHeight() {
+            return 32u;
+        }
+
+        static constexpr std::uint32_t PadTo() {
+            return 16u;
+        }
+
+        static constexpr std::uint32_t Stride() {
+            return 4u;
+        }
+
+    private:
+        Configuration config;
+        MemoryResource* memRes;
+        Vector<JointGroup<TValue> > jointGroups;
+        dna::RotationUnit rotationUnit;
+};
+
+template<typename TValue, typename TFVec256, typename TFVec128>
+QuaternionJointsBuilder<TValue, TFVec256, TFVec128>::QuaternionJointsBuilder(const Configuration& config_,
+                                                                             MemoryResource* memRes_) :
+    config{config_},
+    memRes{memRes_},
+    jointGroups{memRes_},
+    rotationUnit{} {
+}
+
+template<typename TValue, typename TFVec256, typename TFVec128>
+void QuaternionJointsBuilder<TValue, TFVec256, TFVec128>::computeStorageRequirements(const RigMetrics&  /*unused*/) {
+}
+
+template<typename TValue, typename TFVec256, typename TFVec128>
+void QuaternionJointsBuilder<TValue, TFVec256, TFVec128>::computeStorageRequirements(const JointBehaviorFilter&  /*unused*/) {
+}
+
+template<typename TValue, typename TFVec256, typename TFVec128>
+void QuaternionJointsBuilder<TValue, TFVec256, TFVec128>::allocateStorage(const JointBehaviorFilter& source) {
+    jointGroups.resize(source.getJointGroupCount(), JointGroup<TValue>{memRes});
+    for (auto& group : jointGroups) {
+        group.lods.resize(source.getLODCount());
+    }
+}
+
+template<typename TValue, typename TFVec256, typename TFVec128>
+void QuaternionJointsBuilder<TValue, TFVec256, TFVec128>::setInputIndices(const JointBehaviorFilter& source,
+                                                                          std::uint16_t jointGroupIndex) {
+    JointGroup<TValue>& group = jointGroups[jointGroupIndex];
+    group.inputIndices.resize(source.getColumnCount(jointGroupIndex));
+    source.copyInputIndices(jointGroupIndex, group.inputIndices);
+}
+
+template<typename TValue, typename TFVec256, typename TFVec128>
+void QuaternionJointsBuilder<TValue, TFVec256, TFVec128>::setOutputIndices(const JointBehaviorFilter& source,
+                                                                           std::uint16_t jointGroupIndex) {
+    Vector<std::uint16_t> outputIndices{source.getRowCount(jointGroupIndex), {}, memRes};
+    source.copyOutputIndices(jointGroupIndex, outputIndices);
+    // Given any rotation indices (rx, ry, rz), return only rx indices for all joints in the group
+    #if !defined(__clang__) && defined(__GNUC__)
+        #pragma GCC diagnostic push
+        #pragma GCC diagnostic ignored "-Wattributes"
+    #endif
+    auto deduplicate = [this](Vector<std::uint16_t>& v) {
+            UnorderedSet<std::uint16_t> deduplicator{memRes};
+            v.erase(std::remove_if(v.begin(), v.end(), [&](const std::uint16_t value) {
+                return !deduplicator.insert(value).second;
+            }), v.end());
+        };
+    #if !defined(__clang__) && defined(__GNUC__)
+        #pragma GCC diagnostic pop
+    #endif
+    Vector<std::uint16_t> outputRotationBaseIndices{memRes};
+    outputRotationBaseIndices.reserve(outputIndices.size() / 3ul);
+    std::transform(outputIndices.begin(),
+                   outputIndices.end(),
+                   std::back_inserter(outputRotationBaseIndices),
+                   [](std::uint16_t outputIndex) {
+            return static_cast<std::uint16_t>((outputIndex / 9) * 9);
+        });
+    deduplicate(outputRotationBaseIndices);
+    // Expand output rotation base indices (qx) into (qx, qy, qz, qw) for all rotations
+    JointGroup<TValue>& group = jointGroups[jointGroupIndex];
+    group.outputIndices.reserve(outputRotationBaseIndices.size() * static_cast<std::uint8_t>(RotationType::Quaternions));
+    for (const auto baseIndex : outputRotationBaseIndices) {
+        // Remap output indices from 9-attribute joints to 10-attribute joints rx -> qx
+        const auto jointIndex = static_cast<std::uint16_t>(baseIndex / 9u);
+        const auto remappedBaseIndex = static_cast<std::uint16_t>(jointIndex * 10u);
+        group.outputIndices.push_back(static_cast<std::uint16_t>(remappedBaseIndex + 3));
+        group.outputIndices.push_back(static_cast<std::uint16_t>(remappedBaseIndex + 4));
+        group.outputIndices.push_back(static_cast<std::uint16_t>(remappedBaseIndex + 5));
+        group.outputIndices.push_back(static_cast<std::uint16_t>(remappedBaseIndex + 6));
+    }
+}
+
+template<typename TValue, typename TFVec256, typename TFVec128>
+void QuaternionJointsBuilder<TValue, TFVec256, TFVec128>::setValues(const JointBehaviorFilter& source,
+                                                                    std::uint16_t jointGroupIndex) {
+    // Convert euler angles to quaternions
+    std::function<tdm::frad(float)> angConv;
+    if (rotationUnit == dna::RotationUnit::degrees) {
+        angConv = [](float angle) {
+                return tdm::frad{tdm::fdeg{angle}};
+            };
+    } else {
+        angConv = [](float angle) {
+                return tdm::frad{angle};
+            };
+    }
+
+    const auto rowCount = source.getRowCount(jointGroupIndex);
+    const auto colCount = source.getColumnCount(jointGroupIndex);
+    Vector<float> eulers(static_cast<std::size_t>(rowCount) * static_cast<std::size_t>(colCount), {}, memRes);
+    source.copyValues(jointGroupIndex, eulers);
+
+    Vector<std::uint16_t> outputIndices{source.getRowCount(jointGroupIndex), {}, memRes};
+    source.copyOutputIndices(jointGroupIndex, outputIndices);
+
+    JointGroup<TValue>& group = jointGroups[jointGroupIndex];
+    Vector<float> quaternions(static_cast<std::size_t>(group.outputIndices.size() * colCount), {}, memRes);
+    for (std::size_t col = {}; col < colCount; ++col) {
+        for (std::size_t row = {}, quatIndex = {}; row < rowCount; ++quatIndex) {
+            tdm::frad3 angles;
+            const std::uint16_t jointIndex = static_cast<std::uint16_t>(outputIndices[row] / 9u);
+            while ((row < rowCount) && (jointIndex == static_cast<std::uint16_t>(outputIndices[row] / 9u))) {
+                const auto relAttrIndex = static_cast<std::uint16_t>(outputIndices[row] % 9u);
+                // 0 = rx, 1 = ry, 2 = rz
+                const auto relRotAttrIndex = static_cast<std::uint16_t>(relAttrIndex % 3u);
+                angles[relRotAttrIndex] = angConv(eulers[row * colCount + col]);
+                ++row;
+            }
+            tdm::fquat q{angles, static_cast<tdm::rot_seq>(config.rotationOrder)};
+            quaternions[((quatIndex * 4) + 0ul) * colCount + col] = q.x;
+            quaternions[((quatIndex * 4) + 1ul) * colCount + col] = q.y;
+            quaternions[((quatIndex * 4) + 2ul) * colCount + col] = q.z;
+            quaternions[((quatIndex * 4) + 3ul) * colCount + col] = q.w;
+        }
+    }
+
+    // 8 quaternions x 4 floats per quat = 32
+    const auto newRowCount = static_cast<std::uint32_t>(group.outputIndices.size());
+    const auto paddedRowCount = extd::roundUp(newRowCount, PadTo());
+    group.values.resize(paddedRowCount * colCount);
+    bpcm::Optimizer<TFVec256, BlockHeight(), PadTo(), Stride()>::optimize(group.values.data(),
+                                                                          quaternions.data(), {newRowCount, colCount});
+}
+
+template<typename TValue, typename TFVec256, typename TFVec128>
+void QuaternionJointsBuilder<TValue, TFVec256, TFVec128>::setLODs(const JointBehaviorFilter& source,
+                                                                  std::uint16_t jointGroupIndex) {
+    const auto maxRemappedRotationIndex = [](std::uint16_t absRotAttrIndex) {
+            const auto jointIndex = static_cast<std::uint16_t>(absRotAttrIndex / 9u);
+            const auto newAttrBase = static_cast<std::uint16_t>(jointIndex * 10u);
+            // Only rotation indices are inputs, and since the goal is to find the maximum rotation index,
+            // the last quaternion attribute index is used, which is 6 based on [tx, ty, tz, qx, qy, qz, qw, sx, sy, sz]
+            return static_cast<std::uint16_t>(newAttrBase + 6);
+        };
+
+    Vector<std::uint16_t> outputRotationIndices{source.getRowCount(jointGroupIndex), {}, memRes};
+    source.copyOutputIndices(jointGroupIndex, outputRotationIndices);
+
+    JointGroup<TValue>& group = jointGroups[jointGroupIndex];
+    const auto newRowCount = static_cast<std::uint32_t>(group.outputIndices.size());
+    const auto paddedRowCount = extd::roundUp(newRowCount, PadTo());
+    const auto lodCount = source.getLODCount();
+    for (std::uint16_t lod = {}; lod < lodCount; ++lod) {
+        const std::uint32_t oldLODRowCount = source.getRowCountForLOD(jointGroupIndex, lod);
+        std::uint32_t newLODRowCount = {};
+        if (oldLODRowCount != 0) {
+            const auto qwRotationIndexAtOldLODRowCount = maxRemappedRotationIndex(outputRotationIndices[oldLODRowCount - 1ul]);
+            auto it = std::find(group.outputIndices.begin(), group.outputIndices.end(), qwRotationIndexAtOldLODRowCount);
+            assert(it != group.outputIndices.end());
+            newLODRowCount = static_cast<std::uint16_t>(std::distance(group.outputIndices.begin(), it) + 1);
+        }
+        group.lods[lod] = LODRegion{newLODRowCount, paddedRowCount, BlockHeight(), PadTo()};
+    }
+}
+
+template<typename TValue, typename TFVec256, typename TFVec128>
+void QuaternionJointsBuilder<TValue, TFVec256, TFVec128>::remapOutputIndices(std::uint16_t jointGroupIndex) {
+    JointGroup<TValue>& group = jointGroups[jointGroupIndex];
+    for (auto& outputIndex : group.outputIndices) {
+        const auto jointIndex = static_cast<std::uint16_t>(outputIndex / 10u);
+        const auto relAttrIndex = static_cast<std::uint16_t>(outputIndex % 10u);
+        const auto newAttrBase = static_cast<std::uint16_t>(jointIndex * 9u);
+        // Only rotations are among output indices (no translation or scale)
+        // qx, qy, qz are kept, qw is ignored
+        outputIndex = (relAttrIndex == 6) ? std::uint16_t{} : static_cast<std::uint16_t>(newAttrBase + relAttrIndex);
+    }
+}
+
+template<typename TValue, typename TFVec256, typename TFVec128>
+void QuaternionJointsBuilder<TValue, TFVec256, TFVec128>::fillStorage(const JointBehaviorFilter& source) {
+    rotationUnit = source.getRotationUnit();
+
+    for (std::uint16_t jgi = {}; jgi < static_cast<std::uint16_t>(jointGroups.size()); ++jgi) {
+        const auto rowCount = source.getRowCount(jgi);
+        const auto colCount = source.getColumnCount(jgi);
+        if ((rowCount == 0u) || (colCount == 0u)) {
+            continue;
+        }
+
+        setInputIndices(source, jgi);
+        setOutputIndices(source, jgi);
+        setValues(source, jgi);
+        setLODs(source, jgi);
+        // If the selected RigLogic output is in quaternions, then the output indices are already setup as needed.
+        // But if Euler angles were requested, the output indices need to be mapped back to 9-attribute joint output indices
+        if (config.rotationType == RotationType::EulerAngles) {
+            remapOutputIndices(jgi);
+        }
+    }
+}
+
+template<typename T, typename TFVec256, typename TFVec128>
+typename UniqueInstance<JointGroupQuaternionCalculationStrategy<T> >::PointerType createJointGroupQuaternionStrategy(
+    RotationType rotationType,
+    RotationOrder rotationOrder,
+    dna::RotationUnit rotationUnit,
+    MemoryResource* memRes) {
+
+    if (rotationType == RotationType::Quaternions) {
+        using CalculationStrategy = VectorizedJointGroupQuaternionCalculationStrategy<T, TFVec256, TFVec128, PassthroughAdapter>;
+        return UniqueInstance<CalculationStrategy, JointGroupQuaternionCalculationStrategy<T> >::with(memRes).create();
+    }
+
+    #ifdef RL_BUILD_WITH_XYZ_ROTATION_ORDER
+        if (rotationOrder == RotationOrder::XYZ) {
+            if (rotationUnit == dna::RotationUnit::degrees) {
+                using CalculationStrategy = VectorizedJointGroupQuaternionCalculationStrategy<T, TFVec256, TFVec128,
+                                                                                              QuaternionsToEulerAngles<tdm::fdeg,
+                                                                                                                       tdm::rot_seq::xyz> >;
+                return UniqueInstance<CalculationStrategy, JointGroupQuaternionCalculationStrategy<T> >::with(memRes).create();
+            } else {
+                using CalculationStrategy = VectorizedJointGroupQuaternionCalculationStrategy<T, TFVec256, TFVec128,
+                                                                                              QuaternionsToEulerAngles<tdm::frad,
+                                                                                                                       tdm::rot_seq::xyz> >;
+                return UniqueInstance<CalculationStrategy, JointGroupQuaternionCalculationStrategy<T> >::with(memRes).create();
+            }
+        }
+    #endif  // RL_BUILD_WITH_XYZ_ROTATION_ORDER
+
+    #ifdef RL_BUILD_WITH_XZY_ROTATION_ORDER
+        if (rotationOrder == RotationOrder::XZY) {
+            if (rotationUnit == dna::RotationUnit::degrees) {
+                using CalculationStrategy = VectorizedJointGroupQuaternionCalculationStrategy<T, TFVec256, TFVec128,
+                                                                                              QuaternionsToEulerAngles<tdm::fdeg,
+                                                                                                                       tdm::rot_seq::xzy> >;
+                return UniqueInstance<CalculationStrategy, JointGroupQuaternionCalculationStrategy<T> >::with(memRes).create();
+            } else {
+                using CalculationStrategy = VectorizedJointGroupQuaternionCalculationStrategy<T, TFVec256, TFVec128,
+                                                                                              QuaternionsToEulerAngles<tdm::frad,
+                                                                                                                       tdm::rot_seq::xzy> >;
+                return UniqueInstance<CalculationStrategy, JointGroupQuaternionCalculationStrategy<T> >::with(memRes).create();
+            }
+        }
+    #endif  // RL_BUILD_WITH_XZY_ROTATION_ORDER
+
+    #ifdef RL_BUILD_WITH_YXZ_ROTATION_ORDER
+        if (rotationOrder == RotationOrder::YXZ) {
+            if (rotationUnit == dna::RotationUnit::degrees) {
+                using CalculationStrategy = VectorizedJointGroupQuaternionCalculationStrategy<T, TFVec256, TFVec128,
+                                                                                              QuaternionsToEulerAngles<tdm::fdeg,
+                                                                                                                       tdm::rot_seq::yxz> >;
+                return UniqueInstance<CalculationStrategy, JointGroupQuaternionCalculationStrategy<T> >::with(memRes).create();
+            } else {
+                using CalculationStrategy = VectorizedJointGroupQuaternionCalculationStrategy<T, TFVec256, TFVec128,
+                                                                                              QuaternionsToEulerAngles<tdm::frad,
+                                                                                                                       tdm::rot_seq::yxz> >;
+                return UniqueInstance<CalculationStrategy, JointGroupQuaternionCalculationStrategy<T> >::with(memRes).create();
+            }
+        }
+    #endif  // RL_BUILD_WITH_YXZ_ROTATION_ORDER
+
+    #ifdef RL_BUILD_WITH_YZX_ROTATION_ORDER
+        if (rotationOrder == RotationOrder::YZX) {
+            if (rotationUnit == dna::RotationUnit::degrees) {
+                using CalculationStrategy = VectorizedJointGroupQuaternionCalculationStrategy<T, TFVec256, TFVec128,
+                                                                                              QuaternionsToEulerAngles<tdm::fdeg,
+                                                                                                                       tdm::rot_seq::yzx> >;
+                return UniqueInstance<CalculationStrategy, JointGroupQuaternionCalculationStrategy<T> >::with(memRes).create();
+            } else {
+                using CalculationStrategy = VectorizedJointGroupQuaternionCalculationStrategy<T, TFVec256, TFVec128,
+                                                                                              QuaternionsToEulerAngles<tdm::frad,
+                                                                                                                       tdm::rot_seq::yzx> >;
+                return UniqueInstance<CalculationStrategy, JointGroupQuaternionCalculationStrategy<T> >::with(memRes).create();
+            }
+        }
+    #endif  // RL_BUILD_WITH_YZX_ROTATION_ORDER
+
+    #ifdef RL_BUILD_WITH_ZXY_ROTATION_ORDER
+        if (rotationOrder == RotationOrder::ZXY) {
+            if (rotationUnit == dna::RotationUnit::degrees) {
+                using CalculationStrategy = VectorizedJointGroupQuaternionCalculationStrategy<T, TFVec256, TFVec128,
+                                                                                              QuaternionsToEulerAngles<tdm::fdeg,
+                                                                                                                       tdm::rot_seq::zxy> >;
+                return UniqueInstance<CalculationStrategy, JointGroupQuaternionCalculationStrategy<T> >::with(memRes).create();
+            } else {
+                using CalculationStrategy = VectorizedJointGroupQuaternionCalculationStrategy<T, TFVec256, TFVec128,
+                                                                                              QuaternionsToEulerAngles<tdm::frad,
+                                                                                                                       tdm::rot_seq::zxy> >;
+                return UniqueInstance<CalculationStrategy, JointGroupQuaternionCalculationStrategy<T> >::with(memRes).create();
+            }
+        }
+    #endif  // RL_BUILD_WITH_ZXY_ROTATION_ORDER
+
+    #ifdef RL_BUILD_WITH_ZYX_ROTATION_ORDER
+        if (rotationOrder == RotationOrder::ZYX) {
+            if (rotationUnit == dna::RotationUnit::degrees) {
+                using CalculationStrategy = VectorizedJointGroupQuaternionCalculationStrategy<T, TFVec256, TFVec128,
+                                                                                              QuaternionsToEulerAngles<tdm::fdeg,
+                                                                                                                       tdm::rot_seq::zyx> >;
+                return UniqueInstance<CalculationStrategy, JointGroupQuaternionCalculationStrategy<T> >::with(memRes).create();
+            } else {
+                using CalculationStrategy = VectorizedJointGroupQuaternionCalculationStrategy<T, TFVec256, TFVec128,
+                                                                                              QuaternionsToEulerAngles<tdm::frad,
+                                                                                                                       tdm::rot_seq::zyx> >;
+                return UniqueInstance<CalculationStrategy, JointGroupQuaternionCalculationStrategy<T> >::with(memRes).create();
+            }
+        }
+    #endif  // RL_BUILD_WITH_ZYX_ROTATION_ORDER
+
+    return nullptr;
+}
+
+template<typename TValue, typename TFVec256, typename TFVec128>
+JointsEvaluator::Pointer QuaternionJointsBuilder<TValue, TFVec256, TFVec128>::build() {
+    auto factory = UniqueInstance<QuaternionJointsEvaluator<TValue>, JointsEvaluator>::with(memRes);
+    auto strategy = createJointGroupQuaternionStrategy<TValue, TFVec256, TFVec128>(config.rotationType,
+                                                                                   config.rotationOrder,
+                                                                                   rotationUnit,
+                                                                                   memRes);
+    return factory.create(std::move(strategy), std::move(jointGroups), nullptr, memRes);
+}
+
+}  // namespace rl4
diff --git a/Engine/Plugins/Animation/RigLogic/Source/RigLogicLib/Private/riglogic/joints/cpu/quaternions/QuaternionJointsBuilderFactory.cpp b/Engine/Plugins/Animation/RigLogic/Source/RigLogicLib/Private/riglogic/joints/cpu/quaternions/QuaternionJointsBuilderFactory.cpp
new file mode 100644
--- /dev/null
+++ b/Engine/Plugins/Animation/RigLogic/Source/RigLogicLib/Private/riglogic/joints/cpu/quaternions/QuaternionJointsBuilderFactory.cpp
@@ -0,0 +1,54 @@
+// Copyright Epic Games, Inc. All Rights Reserved.
+
+#include "riglogic/joints/cpu/quaternions/QuaternionJointsBuilderFactory.h"
+
+#include "riglogic/joints/cpu/quaternions/QuaternionJointsBuilder.h"
+#include "riglogic/system/simd/Detect.h"
+
+namespace rl4 {
+
+UniqueInstance<JointsBuilder>::PointerType QuaternionJointsBuilderFactory::create(const Configuration& config,
+                                                                                  MemoryResource* memRes) {
+    #if defined(__clang__)
+        #pragma clang diagnostic push
+        #pragma clang diagnostic ignored "-Wunused-local-typedef"
+    #endif
+    #if !defined(__clang__) && defined(__GNUC__)
+        #pragma GCC diagnostic push
+        #pragma GCC diagnostic ignored "-Wunused-local-typedefs"
+    #endif
+    #ifdef RL_USE_HALF_FLOATS
+        using StorageType = std::uint16_t;
+    #else
+        using StorageType = float;
+    #endif  // RL_USE_HALF_FLOATS
+    #if !defined(__clang__) && defined(__GNUC__)
+        #pragma GCC diagnostic pop
+    #endif
+    #if defined(__clang__)
+        #pragma clang diagnostic pop
+    #endif
+
+    #ifdef RL_BUILD_WITH_SSE
+        if ((config.calculationType == CalculationType::SSE) || (config.calculationType == CalculationType::AnyVector)) {
+            using SSEQuaternionJointsBuilder = QuaternionJointsBuilder<StorageType, trimd::sse::F256, trimd::sse::F128>;
+            return UniqueInstance<SSEQuaternionJointsBuilder, JointsBuilder>::with(memRes).create(config, memRes);
+        }
+    #endif  // RL_BUILD_WITH_SSE
+    #ifdef RL_BUILD_WITH_AVX
+        if ((config.calculationType == CalculationType::AVX) || (config.calculationType == CalculationType::AnyVector)) {
+            using AVXQuaternionJointsBuilder = QuaternionJointsBuilder<StorageType, trimd::avx::F256, trimd::sse::F128>;
+            return UniqueInstance<AVXQuaternionJointsBuilder, JointsBuilder>::with(memRes).create(config, memRes);
+        }
+    #endif  // RL_BUILD_WITH_AVX
+    #ifdef RL_BUILD_WITH_NEON
+        if ((config.calculationType == CalculationType::NEON) || (config.calculationType == CalculationType::AnyVector)) {
+            using NEONQuaternionJointsBuilder = QuaternionJointsBuilder<StorageType, trimd::neon::F256, trimd::neon::F128>;
+            return UniqueInstance<NEONQuaternionJointsBuilder, JointsBuilder>::with(memRes).create(config, memRes);
+        }
+    #endif  // RL_BUILD_WITH_NEON
+    using ScalarQuaternionJointsBuilder = QuaternionJointsBuilder<float, trimd::scalar::F256, trimd::scalar::F128>;
+    return UniqueInstance<ScalarQuaternionJointsBuilder, JointsBuilder>::with(memRes).create(config, memRes);
+}
+
+}  // namespace rl4
diff --git a/Engine/Plugins/Animation/RigLogic/Source/RigLogicLib/Private/riglogic/joints/cpu/quaternions/QuaternionJointsBuilderFactory.h b/Engine/Plugins/Animation/RigLogic/Source/RigLogicLib/Private/riglogic/joints/cpu/quaternions/QuaternionJointsBuilderFactory.h
new file mode 100644
--- /dev/null
+++ b/Engine/Plugins/Animation/RigLogic/Source/RigLogicLib/Private/riglogic/joints/cpu/quaternions/QuaternionJointsBuilderFactory.h
@@ -0,0 +1,16 @@
+// Copyright Epic Games, Inc. All Rights Reserved.
+
+#pragma once
+
+#include "riglogic/TypeDefs.h"
+#include "riglogic/joints/JointsBuilder.h"
+#include "riglogic/riglogic/Configuration.h"
+
+namespace rl4 {
+
+struct QuaternionJointsBuilderFactory {
+    static UniqueInstance<JointsBuilder>::PointerType create(const Configuration& config, MemoryResource* memRes);
+
+};
+
+}  // namespace rl4
diff --git a/Engine/Plugins/Animation/RigLogic/Source/RigLogicLib/Private/riglogic/joints/cpu/quaternions/QuaternionJointsEvaluator.h b/Engine/Plugins/Animation/RigLogic/Source/RigLogicLib/Private/riglogic/joints/cpu/quaternions/QuaternionJointsEvaluator.h
new file mode 100644
--- /dev/null
+++ b/Engine/Plugins/Animation/RigLogic/Source/RigLogicLib/Private/riglogic/joints/cpu/quaternions/QuaternionJointsEvaluator.h
@@ -0,0 +1,89 @@
+// Copyright Epic Games, Inc. All Rights Reserved.
+
+#pragma once
+
+#include "riglogic/TypeDefs.h"
+#include "riglogic/controls/ControlsInputInstance.h"
+#include "riglogic/joints/JointsEvaluator.h"
+#include "riglogic/joints/JointsOutputInstance.h"
+#include "riglogic/joints/cpu/quaternions/CalculationStrategy.h"
+#include "riglogic/joints/cpu/quaternions/JointGroup.h"
+
+#include <cstdint>
+
+namespace rl4 {
+
+template<typename TValue>
+class QuaternionJointsEvaluator : public JointsEvaluator {
+    public:
+        struct Accessor;
+        friend Accessor;
+
+    private:
+        using CalculationStrategyPointer = typename UniqueInstance<JointGroupQuaternionCalculationStrategy<TValue> >::PointerType;
+
+    public:
+        QuaternionJointsEvaluator(CalculationStrategyPointer strategy_,
+                                  Vector<JointGroup<TValue> >&& jointGroups_,
+                                  JointsOutputInstance::Factory instanceFactory_,
+                                  MemoryResource* memRes);
+
+        JointsOutputInstance::Pointer createInstance(MemoryResource* instanceMemRes) const override;
+        void calculate(const ControlsInputInstance* inputs, JointsOutputInstance* outputs, std::uint16_t lod) const override;
+        void calculate(const ControlsInputInstance* inputs,
+                       JointsOutputInstance* outputs,
+                       std::uint16_t lod,
+                       std::uint16_t jointGroupIndex) const override;
+        void load(terse::BinaryInputArchive<BoundedIOStream>& archive) override;
+        void save(terse::BinaryOutputArchive<BoundedIOStream>& archive) override;
+
+    private:
+        CalculationStrategyPointer strategy;
+        Vector<JointGroup<TValue> > jointGroups;
+        JointsOutputInstance::Factory instanceFactory;
+
+};
+
+template<typename TValue>
+QuaternionJointsEvaluator<TValue>::QuaternionJointsEvaluator(CalculationStrategyPointer strategy_,
+                                                             Vector<JointGroup<TValue> >&& jointGroups_,
+                                                             JointsOutputInstance::Factory instanceFactory_,
+                                                             MemoryResource*  /*unused*/) :
+    strategy{std::move(strategy_)},
+    jointGroups{std::move(jointGroups_)},
+    instanceFactory{instanceFactory_} {
+}
+
+template<typename TValue>
+JointsOutputInstance::Pointer QuaternionJointsEvaluator<TValue>::createInstance(MemoryResource* instanceMemRes) const {
+    return instanceFactory(instanceMemRes);
+}
+
+template<typename TValue>
+void QuaternionJointsEvaluator<TValue>::calculate(const ControlsInputInstance* inputs,
+                                                  JointsOutputInstance* outputs,
+                                                  std::uint16_t lod) const {
+    for (std::size_t i = {}; i < jointGroups.size(); ++i) {
+        calculate(inputs, outputs, lod, static_cast<std::uint16_t>(i));
+    }
+}
+
+template<typename TValue>
+void QuaternionJointsEvaluator<TValue>::calculate(const ControlsInputInstance* inputs,
+                                                  JointsOutputInstance* outputs,
+                                                  std::uint16_t lod,
+                                                  std::uint16_t jointGroupIndex) const {
+    strategy->calculate(jointGroups[jointGroupIndex], inputs->getInputBuffer(), outputs->getOutputBuffer(), lod);
+}
+
+template<typename TValue>
+void QuaternionJointsEvaluator<TValue>::load(terse::BinaryInputArchive<BoundedIOStream>& archive) {
+    archive(jointGroups);
+}
+
+template<typename TValue>
+void QuaternionJointsEvaluator<TValue>::save(terse::BinaryOutputArchive<BoundedIOStream>& archive) {
+    archive(jointGroups);
+}
+
+}  // namespace rl4
diff --git a/Engine/Plugins/Animation/RigLogic/Source/RigLogicLib/Private/riglogic/joints/cpu/twistswing/TwistSwingJointsBuilder.h b/Engine/Plugins/Animation/RigLogic/Source/RigLogicLib/Private/riglogic/joints/cpu/twistswing/TwistSwingJointsBuilder.h
new file mode 100644
--- /dev/null
+++ b/Engine/Plugins/Animation/RigLogic/Source/RigLogicLib/Private/riglogic/joints/cpu/twistswing/TwistSwingJointsBuilder.h
@@ -0,0 +1,259 @@
+// Copyright Epic Games, Inc. All Rights Reserved.
+
+#pragma once
+
+#include "riglogic/TypeDefs.h"
+#include "riglogic/joints/JointBehaviorFilter.h"
+#include "riglogic/joints/JointsBuilder.h"
+#include "riglogic/joints/cpu/quaternions/RotationAdapters.h"
+#include "riglogic/joints/cpu/twistswing/TwistSwingJointsEvaluator.h"
+#include "riglogic/riglogic/Configuration.h"
+#include "riglogic/utils/Extd.h"
+
+namespace rl4 {
+
+template<class TContainer, class UContainer>
+static bool containersEqual(const TContainer& lhs, const UContainer& rhs) {
+    if (lhs.size() != rhs.size()) {
+        return false;
+    }
+
+    #if __cplusplus >= 201402L || (defined(_MSC_VER) && _MSC_VER >= 1900)
+        return std::equal(lhs.begin(), lhs.end(), rhs.begin(), rhs.end());
+    #else
+        return std::equal(lhs.begin(), lhs.end(), rhs.begin());
+    #endif
+}
+
+template<typename TValue, typename TFVec256, typename TFVec128>
+class TwistSwingJointsBuilder : public JointsBuilder {
+    public:
+        TwistSwingJointsBuilder(const Configuration& config_, MemoryResource* memRes_);
+
+        void computeStorageRequirements(const RigMetrics& source) override;
+        void computeStorageRequirements(const JointBehaviorFilter& source) override;
+        void allocateStorage(const JointBehaviorFilter& source) override;
+        void fillStorage(const JointBehaviorFilter& source) override;
+        JointsEvaluator::Pointer build() override;
+
+    private:
+        Configuration config;
+        MemoryResource* memRes;
+        Vector<TwistSwingSetup> setups;
+        dna::RotationUnit rotationUnit;
+};
+
+template<typename TValue, typename TFVec256, typename TFVec128>
+TwistSwingJointsBuilder<TValue, TFVec256, TFVec128>::TwistSwingJointsBuilder(const Configuration& config_,
+                                                                             MemoryResource* memRes_) :
+    config{config_},
+    memRes{memRes_},
+    setups{memRes_},
+    rotationUnit{} {
+}
+
+template<typename TValue, typename TFVec256, typename TFVec128>
+void TwistSwingJointsBuilder<TValue, TFVec256, TFVec128>::computeStorageRequirements(const RigMetrics&  /*unused*/) {
+}
+
+template<typename TValue, typename TFVec256, typename TFVec128>
+void TwistSwingJointsBuilder<TValue, TFVec256, TFVec128>::computeStorageRequirements(const JointBehaviorFilter&  /*unused*/) {
+}
+
+template<typename TValue, typename TFVec256, typename TFVec128>
+void TwistSwingJointsBuilder<TValue, TFVec256, TFVec128>::allocateStorage(const JointBehaviorFilter&  /*unused*/) {
+}
+
+template<typename TValue, typename TFVec256, typename TFVec128>
+void TwistSwingJointsBuilder<TValue, TFVec256, TFVec128>::fillStorage(const JointBehaviorFilter& source) {
+    rotationUnit = source.getRotationUnit();
+    const dna::Reader* reader = source.getReader();
+    const auto swingCount = reader->getSwingCount();
+    const auto twistCount = reader->getTwistCount();
+
+    if (!config.loadTwistSwingBehavior || ((twistCount == 0u) && (swingCount == 0u))) {
+        return;
+    }
+
+    setups.reserve(std::max(twistCount, swingCount));
+
+    for (std::uint16_t si = {}; si < swingCount; ++si) {
+        TwistSwingSetup setup{memRes};
+
+        const auto swingTwistAxis = reader->getSwingSetupTwistAxis(si);
+        const auto swingInputIndices = reader->getSwingInputControlIndices(si);
+        const auto swingOutputIndices = reader->getSwingOutputJointIndices(si);
+        const auto swingBlendWeights = reader->getSwingBlendWeights(si);
+
+        assert(swingInputIndices.size() == 4ul);
+        assert(swingBlendWeights.size() == swingOutputIndices.size());
+
+        setup.swingTwistAxis = swingTwistAxis;
+        setup.swingInputIndices.assign(swingInputIndices.begin(), swingInputIndices.end());
+        setup.swingBlendWeights.assign(swingBlendWeights.begin(), swingBlendWeights.end());
+        setup.swingOutputIndices.reserve(swingOutputIndices.size() * 10ul);
+        if (config.rotationType == RotationType::EulerAngles) {
+            for (const auto jointIndex : swingOutputIndices) {
+                const auto absBaseAttrIndex = jointIndex * 9ul;
+                setup.swingOutputIndices.push_back(static_cast<std::uint16_t>(absBaseAttrIndex + 3u));
+                setup.swingOutputIndices.push_back(static_cast<std::uint16_t>(absBaseAttrIndex + 4u));
+                setup.swingOutputIndices.push_back(static_cast<std::uint16_t>(absBaseAttrIndex + 5u));
+                setup.swingOutputIndices.push_back(static_cast<std::uint16_t>(0));
+            }
+        } else {
+            for (const auto jointIndex : swingOutputIndices) {
+                const auto absBaseAttrIndex = jointIndex * 10ul;
+                setup.swingOutputIndices.push_back(static_cast<std::uint16_t>(absBaseAttrIndex + 3u));
+                setup.swingOutputIndices.push_back(static_cast<std::uint16_t>(absBaseAttrIndex + 4u));
+                setup.swingOutputIndices.push_back(static_cast<std::uint16_t>(absBaseAttrIndex + 5u));
+                setup.swingOutputIndices.push_back(static_cast<std::uint16_t>(absBaseAttrIndex + 6u));
+            }
+        }
+
+        setups.push_back(std::move(setup));
+    }
+
+    for (std::uint16_t ti = {}; ti < twistCount; ++ti) {
+        const auto twistInputIndices = reader->getTwistInputControlIndices(ti);
+        const auto twistOutputIndices = reader->getTwistOutputJointIndices(ti);
+        const auto twistBlendWeights = reader->getTwistBlendWeights(ti);
+        const auto twistTwistAxis = reader->getTwistSetupTwistAxis(ti);
+
+        assert(twistInputIndices.size() == 4ul);
+        assert(twistBlendWeights.size() == twistOutputIndices.size());
+
+        auto fillSetup =
+            [this, twistInputIndices, twistOutputIndices, twistBlendWeights, twistTwistAxis](TwistSwingSetup& setup) {
+                setup.twistTwistAxis = twistTwistAxis;
+                setup.twistInputIndices.assign(twistInputIndices.begin(), twistInputIndices.end());
+                setup.twistBlendWeights.assign(twistBlendWeights.begin(), twistBlendWeights.end());
+                setup.twistOutputIndices.reserve(twistOutputIndices.size() * 10ul);
+                if (config.rotationType == RotationType::EulerAngles) {
+                    for (const auto jointIndex : twistOutputIndices) {
+                        const auto absBaseAttrIndex = jointIndex * 9ul;
+                        setup.twistOutputIndices.push_back(static_cast<std::uint16_t>(absBaseAttrIndex + 3u));
+                        setup.twistOutputIndices.push_back(static_cast<std::uint16_t>(absBaseAttrIndex + 4u));
+                        setup.twistOutputIndices.push_back(static_cast<std::uint16_t>(absBaseAttrIndex + 5u));
+                        setup.twistOutputIndices.push_back(static_cast<std::uint16_t>(0));
+                    }
+                } else {
+                    for (const auto jointIndex : twistOutputIndices) {
+                        const auto absBaseAttrIndex = jointIndex * 10ul;
+                        setup.twistOutputIndices.push_back(static_cast<std::uint16_t>(absBaseAttrIndex + 3u));
+                        setup.twistOutputIndices.push_back(static_cast<std::uint16_t>(absBaseAttrIndex + 4u));
+                        setup.twistOutputIndices.push_back(static_cast<std::uint16_t>(absBaseAttrIndex + 5u));
+                        setup.twistOutputIndices.push_back(static_cast<std::uint16_t>(absBaseAttrIndex + 6u));
+                    }
+                }
+            };
+
+        auto it = std::find_if(setups.begin(), setups.end(), [twistInputIndices](const TwistSwingSetup& setup) {
+                return containersEqual(setup.swingInputIndices, twistInputIndices);
+            });
+        if (it == setups.end()) {
+            TwistSwingSetup setup{memRes};
+            fillSetup(setup);
+            setups.push_back(std::move(setup));
+        } else {
+            fillSetup(*it);
+        }
+    }
+
+}
+
+template<typename TValue, typename TFVec256, typename TFVec128>
+JointsEvaluator::Pointer TwistSwingJointsBuilder<TValue, TFVec256, TFVec128>::build() {
+
+    if (config.rotationType == RotationType::Quaternions) {
+        using TSJEvaluator = TwistSwingJointsEvaluator<TValue, TFVec256, TFVec128, PassthroughAdapter>;
+        return UniqueInstance<TSJEvaluator, JointsEvaluator>::with(memRes).create(std::move(setups), nullptr, memRes);
+    }
+
+    #ifdef RL_BUILD_WITH_XYZ_ROTATION_ORDER
+        if (config.rotationOrder == RotationOrder::XYZ) {
+            if (rotationUnit == dna::RotationUnit::degrees) {
+                using TSJEvaluator = TwistSwingJointsEvaluator<TValue, TFVec256, TFVec128, QuaternionsToEulerAngles<tdm::fdeg,
+                                                                                                                    tdm::rot_seq::xyz> >;
+                return UniqueInstance<TSJEvaluator, JointsEvaluator>::with(memRes).create(std::move(setups), nullptr, memRes);
+            } else {
+                using TSJEvaluator = TwistSwingJointsEvaluator<TValue, TFVec256, TFVec128, QuaternionsToEulerAngles<tdm::frad,
+                                                                                                                    tdm::rot_seq::xyz> >;
+                return UniqueInstance<TSJEvaluator, JointsEvaluator>::with(memRes).create(std::move(setups), nullptr, memRes);
+            }
+        }
+    #endif  // RL_BUILD_WITH_XYZ_ROTATION_ORDER
+
+    #ifdef RL_BUILD_WITH_XZY_ROTATION_ORDER
+        if (config.rotationOrder == RotationOrder::XZY) {
+            if (rotationUnit == dna::RotationUnit::degrees) {
+                using TSJEvaluator = TwistSwingJointsEvaluator<TValue, TFVec256, TFVec128, QuaternionsToEulerAngles<tdm::fdeg,
+                                                                                                                    tdm::rot_seq::xzy> >;
+                return UniqueInstance<TSJEvaluator, JointsEvaluator>::with(memRes).create(std::move(setups), nullptr, memRes);
+            } else {
+                using TSJEvaluator = TwistSwingJointsEvaluator<TValue, TFVec256, TFVec128, QuaternionsToEulerAngles<tdm::frad,
+                                                                                                                    tdm::rot_seq::xzy> >;
+                return UniqueInstance<TSJEvaluator, JointsEvaluator>::with(memRes).create(std::move(setups), nullptr, memRes);
+            }
+        }
+    #endif  // RL_BUILD_WITH_XZY_ROTATION_ORDER
+
+    #ifdef RL_BUILD_WITH_YXZ_ROTATION_ORDER
+        if (config.rotationOrder == RotationOrder::YXZ) {
+            if (rotationUnit == dna::RotationUnit::degrees) {
+                using TSJEvaluator = TwistSwingJointsEvaluator<TValue, TFVec256, TFVec128, QuaternionsToEulerAngles<tdm::fdeg,
+                                                                                                                    tdm::rot_seq::yxz> >;
+                return UniqueInstance<TSJEvaluator, JointsEvaluator>::with(memRes).create(std::move(setups), nullptr, memRes);
+            } else {
+                using TSJEvaluator = TwistSwingJointsEvaluator<TValue, TFVec256, TFVec128, QuaternionsToEulerAngles<tdm::frad,
+                                                                                                                    tdm::rot_seq::yxz> >;
+                return UniqueInstance<TSJEvaluator, JointsEvaluator>::with(memRes).create(std::move(setups), nullptr, memRes);
+            }
+        }
+    #endif  // RL_BUILD_WITH_YXZ_ROTATION_ORDER
+
+    #ifdef RL_BUILD_WITH_YZX_ROTATION_ORDER
+        if (config.rotationOrder == RotationOrder::YZX) {
+            if (rotationUnit == dna::RotationUnit::degrees) {
+                using TSJEvaluator = TwistSwingJointsEvaluator<TValue, TFVec256, TFVec128, QuaternionsToEulerAngles<tdm::fdeg,
+                                                                                                                    tdm::rot_seq::yzx> >;
+                return UniqueInstance<TSJEvaluator, JointsEvaluator>::with(memRes).create(std::move(setups), nullptr, memRes);
+            } else {
+                using TSJEvaluator = TwistSwingJointsEvaluator<TValue, TFVec256, TFVec128, QuaternionsToEulerAngles<tdm::frad,
+                                                                                                                    tdm::rot_seq::yzx> >;
+                return UniqueInstance<TSJEvaluator, JointsEvaluator>::with(memRes).create(std::move(setups), nullptr, memRes);
+            }
+        }
+    #endif  // RL_BUILD_WITH_YZX_ROTATION_ORDER
+
+    #ifdef RL_BUILD_WITH_ZXY_ROTATION_ORDER
+        if (config.rotationOrder == RotationOrder::ZXY) {
+            if (rotationUnit == dna::RotationUnit::degrees) {
+                using TSJEvaluator = TwistSwingJointsEvaluator<TValue, TFVec256, TFVec128, QuaternionsToEulerAngles<tdm::fdeg,
+                                                                                                                    tdm::rot_seq::zxy> >;
+                return UniqueInstance<TSJEvaluator, JointsEvaluator>::with(memRes).create(std::move(setups), nullptr, memRes);
+            } else {
+                using TSJEvaluator = TwistSwingJointsEvaluator<TValue, TFVec256, TFVec128, QuaternionsToEulerAngles<tdm::frad,
+                                                                                                                    tdm::rot_seq::zxy> >;
+                return UniqueInstance<TSJEvaluator, JointsEvaluator>::with(memRes).create(std::move(setups), nullptr, memRes);
+            }
+        }
+    #endif  // RL_BUILD_WITH_ZXY_ROTATION_ORDER
+
+    #ifdef RL_BUILD_WITH_ZYX_ROTATION_ORDER
+        if (config.rotationOrder == RotationOrder::ZYX) {
+            if (rotationUnit == dna::RotationUnit::degrees) {
+                using TSJEvaluator = TwistSwingJointsEvaluator<TValue, TFVec256, TFVec128, QuaternionsToEulerAngles<tdm::fdeg,
+                                                                                                                    tdm::rot_seq::zyx> >;
+                return UniqueInstance<TSJEvaluator, JointsEvaluator>::with(memRes).create(std::move(setups), nullptr, memRes);
+            } else {
+                using TSJEvaluator = TwistSwingJointsEvaluator<TValue, TFVec256, TFVec128, QuaternionsToEulerAngles<tdm::frad,
+                                                                                                                    tdm::rot_seq::zyx> >;
+                return UniqueInstance<TSJEvaluator, JointsEvaluator>::with(memRes).create(std::move(setups), nullptr, memRes);
+            }
+        }
+    #endif  // RL_BUILD_WITH_ZYX_ROTATION_ORDER
+
+    return nullptr;
+}
+
+}  // namespace rl4
diff --git a/Engine/Plugins/Animation/RigLogic/Source/RigLogicLib/Private/riglogic/joints/cpu/twistswing/TwistSwingJointsBuilderFactory.cpp b/Engine/Plugins/Animation/RigLogic/Source/RigLogicLib/Private/riglogic/joints/cpu/twistswing/TwistSwingJointsBuilderFactory.cpp
new file mode 100644
--- /dev/null
+++ b/Engine/Plugins/Animation/RigLogic/Source/RigLogicLib/Private/riglogic/joints/cpu/twistswing/TwistSwingJointsBuilderFactory.cpp
@@ -0,0 +1,58 @@
+// Copyright Epic Games, Inc. All Rights Reserved.
+
+#include "riglogic/joints/cpu/twistswing/TwistSwingJointsBuilderFactory.h"
+
+#include "riglogic/joints/cpu/twistswing/TwistSwingJointsBuilder.h"
+#include "riglogic/system/simd/Detect.h"
+
+namespace rl4 {
+
+UniqueInstance<JointsBuilder>::PointerType TwistSwingJointsBuilderFactory::create(const Configuration& config,
+                                                                                  MemoryResource* memRes) {
+    // At the moment, only a Scalar implementation exists for twist and swing evaluation, so there is no need to
+    // compile any other configuration.
+    /*
+    #if defined(__clang__)
+        #pragma clang diagnostic push
+        #pragma clang diagnostic ignored "-Wunused-local-typedef"
+    #endif
+    #if !defined(__clang__) && defined(__GNUC__)
+        #pragma GCC diagnostic push
+        #pragma GCC diagnostic ignored "-Wunused-local-typedefs"
+    #endif
+    #ifdef RL_USE_HALF_FLOATS
+        using StorageType = std::uint16_t;
+    #else
+        using StorageType = float;
+    #endif  // RL_USE_HALF_FLOATS
+    #if !defined(__clang__) && defined(__GNUC__)
+        #pragma GCC diagnostic pop
+    #endif
+    #if defined(__clang__)
+        #pragma clang diagnostic pop
+    #endif
+
+    #ifdef RL_BUILD_WITH_SSE
+        if ((config.calculationType == CalculationType::SSE) || (config.calculationType == CalculationType::AnyVector)) {
+            using SSETwistSwingJointsBuilder = TwistSwingJointsBuilder<StorageType, trimd::sse::F256, trimd::sse::F128>;
+            return UniqueInstance<SSETwistSwingJointsBuilder, JointsBuilder>::with(memRes).create(config, memRes);
+        }
+    #endif  // RL_BUILD_WITH_SSE
+    #ifdef RL_BUILD_WITH_AVX
+        if ((config.calculationType == CalculationType::AVX) || (config.calculationType == CalculationType::AnyVector)) {
+            using AVXTwistSwingJointsBuilder = TwistSwingJointsBuilder<StorageType, trimd::avx::F256, trimd::sse::F128>;
+            return UniqueInstance<AVXTwistSwingJointsBuilder, JointsBuilder>::with(memRes).create(config, memRes);
+        }
+    #endif  // RL_BUILD_WITH_AVX
+    #ifdef RL_BUILD_WITH_NEON
+        if ((config.calculationType == CalculationType::NEON) || (config.calculationType == CalculationType::AnyVector)) {
+            using NEONTwistSwingJointsBuilder = TwistSwingJointsBuilder<StorageType, trimd::neon::F256, trimd::neon::F128>;
+            return UniqueInstance<NEONTwistSwingJointsBuilder, JointsBuilder>::with(memRes).create(config, memRes);
+        }
+    #endif  // RL_BUILD_WITH_NEON
+    */
+    using ScalarTwistSwingJointsBuilder = TwistSwingJointsBuilder<float, trimd::scalar::F256, trimd::scalar::F128>;
+    return UniqueInstance<ScalarTwistSwingJointsBuilder, JointsBuilder>::with(memRes).create(config, memRes);
+}
+
+}  // namespace rl4
diff --git a/Engine/Plugins/Animation/RigLogic/Source/RigLogicLib/Private/riglogic/joints/cpu/twistswing/TwistSwingJointsBuilderFactory.h b/Engine/Plugins/Animation/RigLogic/Source/RigLogicLib/Private/riglogic/joints/cpu/twistswing/TwistSwingJointsBuilderFactory.h
new file mode 100644
--- /dev/null
+++ b/Engine/Plugins/Animation/RigLogic/Source/RigLogicLib/Private/riglogic/joints/cpu/twistswing/TwistSwingJointsBuilderFactory.h
@@ -0,0 +1,16 @@
+// Copyright Epic Games, Inc. All Rights Reserved.
+
+#pragma once
+
+#include "riglogic/TypeDefs.h"
+#include "riglogic/joints/JointsBuilder.h"
+#include "riglogic/riglogic/Configuration.h"
+
+namespace rl4 {
+
+struct TwistSwingJointsBuilderFactory {
+    static UniqueInstance<JointsBuilder>::PointerType create(const Configuration& config, MemoryResource* memRes);
+
+};
+
+}  // namespace rl4
diff --git a/Engine/Plugins/Animation/RigLogic/Source/RigLogicLib/Private/riglogic/ml/MachineLearnedBehavior.h b/Engine/Plugins/Animation/RigLogic/Source/RigLogicLib/Private/riglogic/ml/MachineLearnedBehavior.h
--- a/Engine/Plugins/Animation/RigLogic/Source/RigLogicLib/Private/riglogic/ml/MachineLearnedBehavior.h
+++ b/Engine/Plugins/Animation/RigLogic/Source/RigLogicLib/Private/riglogic/ml/MachineLearnedBehavior.h
@@ -33,22 +33,14 @@ class MachineLearnedBehavior {
 
         template<class Archive>
         void load(Archive& archive) {
-            #ifdef RL_BUILD_WITH_ML_EVALUATOR
-                evaluator->load(archive);
-                archive >> neuralNetworkIndicesPerMeshRegion;
-            #else
-                UNUSED(archive);
-            #endif  // RL_BUILD_WITH_ML_EVALUATOR
+            evaluator->load(archive);
+            archive >> neuralNetworkIndicesPerMeshRegion;
         }
 
         template<class Archive>
         void save(Archive& archive) {
-            #ifdef RL_BUILD_WITH_ML_EVALUATOR
-                evaluator->save(archive);
-                archive << neuralNetworkIndicesPerMeshRegion;
-            #else
-                UNUSED(archive);
-            #endif  // RL_BUILD_WITH_ML_EVALUATOR
+            evaluator->save(archive);
+            archive << neuralNetworkIndicesPerMeshRegion;
         }
 
         std::uint16_t getMeshCount() const;
diff --git a/Engine/Plugins/Animation/RigLogic/Source/RigLogicLib/Private/riglogic/ml/MachineLearnedBehaviorFactory.cpp b/Engine/Plugins/Animation/RigLogic/Source/RigLogicLib/Private/riglogic/ml/MachineLearnedBehaviorFactory.cpp
--- a/Engine/Plugins/Animation/RigLogic/Source/RigLogicLib/Private/riglogic/ml/MachineLearnedBehaviorFactory.cpp
+++ b/Engine/Plugins/Animation/RigLogic/Source/RigLogicLib/Private/riglogic/ml/MachineLearnedBehaviorFactory.cpp
@@ -6,9 +6,7 @@
 #include "riglogic/utils/Extd.h"
 #include "riglogic/ml/MachineLearnedBehaviorEvaluator.h"
 #include "riglogic/ml/MachineLearnedBehaviorNullEvaluator.h"
-#ifdef RL_BUILD_WITH_ML_EVALUATOR
-    #include "riglogic/ml/cpu/Factory.h"
-#endif  // RL_BUILD_WITH_ML_EVALUATOR
+#include "riglogic/ml/cpu/CPUMachineLearnedBehaviorFactory.h"
 #include "riglogic/riglogic/Configuration.h"
 #include "riglogic/riglogic/RigMetrics.h"
 #include "riglogic/system/simd/SIMD.h"
@@ -17,33 +15,44 @@
 
 namespace rl4 {
 
-// *INDENT-OFF*
-#ifdef RL_BUILD_WITH_ML_EVALUATOR
+MachineLearnedBehaviorEvaluator::Pointer createMLEvaluator(const Configuration& config,
+                                                           const dna::MachineLearnedBehaviorReader* reader,
+                                                           MemoryResource* memRes) {
+    #if defined(__clang__)
+        #pragma clang diagnostic push
+        #pragma clang diagnostic ignored "-Wunused-local-typedef"
+    #endif
+    #if !defined(__clang__) && defined(__GNUC__)
+        #pragma GCC diagnostic push
+        #pragma GCC diagnostic ignored "-Wunused-local-typedefs"
+    #endif
+    #ifdef RL_USE_HALF_FLOATS
+        using StorageValueType = std::uint16_t;
+    #else
+        using StorageValueType = float;
+    #endif  // RL_USE_HALF_FLOATS
+    #if !defined(__clang__) && defined(__GNUC__)
+        #pragma GCC diagnostic pop
+    #endif
+    #if defined(__clang__)
+        #pragma clang diagnostic pop
+    #endif
 
-#ifdef RL_USE_HALF_FLOATS
-    using StorageValueType = std::uint16_t;
-#else
-    using StorageValueType = float;
-#endif  // RL_USE_HALF_FLOATS
-
-MachineLearnedBehaviorEvaluator::Pointer createEvaluator(const Configuration& config,
-                                                         const dna::MachineLearnedBehaviorReader* reader,
-                                                         MemoryResource* memRes) {
     // Work around unused parameter warning when building without SSE and AVX
     static_cast<void>(config);
     #ifdef RL_BUILD_WITH_SSE
-        if (config.calculationType == CalculationType::SSE) {
+        if ((config.calculationType == CalculationType::SSE) || (config.calculationType == CalculationType::AnyVector)) {
             return ml::cpu::Factory<StorageValueType, trimd::sse::F256, trimd::sse::F128>::create(reader, memRes);
         }
     #endif  // RL_BUILD_WITH_SSE
     #ifdef RL_BUILD_WITH_AVX
-        if (config.calculationType == CalculationType::AVX) {
+        if ((config.calculationType == CalculationType::AVX) || (config.calculationType == CalculationType::AnyVector)) {
             // Use 256-bit AVX registers and whatever 128-bit width type is available
             return ml::cpu::Factory<StorageValueType, trimd::avx::F256, trimd::sse::F128>::create(reader, memRes);
         }
     #endif  // RL_BUILD_WITH_AVX
     #ifdef RL_BUILD_WITH_NEON
-        if (config.calculationType == CalculationType::NEON) {
+        if ((config.calculationType == CalculationType::NEON) || (config.calculationType == CalculationType::AnyVector)) {
             return ml::cpu::Factory<StorageValueType, trimd::neon::F256, trimd::neon::F128>::create(reader, memRes);
         }
     #endif  // RL_BUILD_WITH_NEON
@@ -70,34 +79,19 @@ MachineLearnedBehavior::Pointer MachineLearnedBehaviorFactory::create(const Conf
         }
     }
 
-    return moduleFactory.create(createEvaluator(config, reader, memRes), std::move(neuralNetworkIndicesPerMeshRegion));
+    return moduleFactory.create(createMLEvaluator(config, reader, memRes), std::move(neuralNetworkIndicesPerMeshRegion));
 }
 
-MachineLearnedBehavior::Pointer MachineLearnedBehaviorFactory::create(const Configuration& config, MemoryResource* memRes) {
+MachineLearnedBehavior::Pointer MachineLearnedBehaviorFactory::create(const Configuration& config,
+                                                                      const RigMetrics& metrics,
+                                                                      MemoryResource* memRes) {
     auto moduleFactory = UniqueInstance<MachineLearnedBehavior>::with(memRes);
-    if (!config.loadMachineLearnedBehavior) {
+    if (!config.loadMachineLearnedBehavior || (metrics.neuralNetworkCount == 0u)) {
         auto evaluator =
             UniqueInstance<MachineLearnedBehaviorNullEvaluator, MachineLearnedBehaviorEvaluator>::with(memRes).create();
         return moduleFactory.create(std::move(evaluator), memRes);
     }
-    return moduleFactory.create(createEvaluator(config, nullptr, memRes), memRes);
+    return moduleFactory.create(createMLEvaluator(config, nullptr, memRes), memRes);
 }
 
-#else
-
-MachineLearnedBehavior::Pointer MachineLearnedBehaviorFactory::create(const Configuration&  /*unused*/,
-                                                                      const dna::MachineLearnedBehaviorReader*  /*unused*/,
-                                                                      MemoryResource* memRes) {
-    auto moduleFactory = UniqueInstance<MachineLearnedBehavior>::with(memRes);
-    auto evaluator = UniqueInstance<MachineLearnedBehaviorNullEvaluator, MachineLearnedBehaviorEvaluator>::with(memRes).create();
-    return moduleFactory.create(std::move(evaluator), memRes);
-}
-
-MachineLearnedBehavior::Pointer MachineLearnedBehaviorFactory::create(const Configuration& config, MemoryResource* memRes) {
-    return create(config, nullptr, memRes);
-}
-
-#endif  // RL_BUILD_WITH_ML_EVALUATOR
-// *INDENT-ON*
-
 }  // namespace rl4
diff --git a/Engine/Plugins/Animation/RigLogic/Source/RigLogicLib/Private/riglogic/ml/MachineLearnedBehaviorFactory.h b/Engine/Plugins/Animation/RigLogic/Source/RigLogicLib/Private/riglogic/ml/MachineLearnedBehaviorFactory.h
--- a/Engine/Plugins/Animation/RigLogic/Source/RigLogicLib/Private/riglogic/ml/MachineLearnedBehaviorFactory.h
+++ b/Engine/Plugins/Animation/RigLogic/Source/RigLogicLib/Private/riglogic/ml/MachineLearnedBehaviorFactory.h
@@ -14,7 +14,7 @@ struct MachineLearnedBehaviorFactory {
     static MachineLearnedBehavior::Pointer create(const Configuration& config,
                                                   const dna::MachineLearnedBehaviorReader* reader,
                                                   MemoryResource* memRes);
-    static MachineLearnedBehavior::Pointer create(const Configuration& config, MemoryResource* memRes);
+    static MachineLearnedBehavior::Pointer create(const Configuration& config, const RigMetrics& metrics, MemoryResource* memRes);
 
 };
 
