Update to new RigLogic implementing body rig support

diff --git a/Engine/Plugins/Animation/RigLogic/Source/RigLogicLib/Private/riglogic/ml/cpu/CPUMachineLearnedBehaviorEvaluator.h b/Engine/Plugins/Animation/RigLogic/Source/RigLogicLib/Private/riglogic/ml/cpu/CPUMachineLearnedBehaviorEvaluator.h
new file mode 100644
--- /dev/null
+++ b/Engine/Plugins/Animation/RigLogic/Source/RigLogicLib/Private/riglogic/ml/cpu/CPUMachineLearnedBehaviorEvaluator.h
@@ -0,0 +1,98 @@
+// Copyright Epic Games, Inc. All Rights Reserved.
+
+#pragma once
+
+#include "riglogic/TypeDefs.h"
+#include "riglogic/controls/ControlsInputInstance.h"
+#include "riglogic/ml/MachineLearnedBehaviorEvaluator.h"
+#include "riglogic/ml/cpu/CPUMachineLearnedBehaviorOutputInstance.h"
+#include "riglogic/ml/cpu/Inference.h"
+#include "riglogic/ml/cpu/NeuralNet.h"
+#include "riglogic/types/LODSpec.h"
+
+#include <cstddef>
+
+namespace rl4 {
+
+namespace ml {
+
+namespace cpu {
+
+template<typename T, typename TF256, typename TF128>
+class Evaluator : public MachineLearnedBehaviorEvaluator {
+    public:
+        using NeuralNetVectorType = Vector<NeuralNetInference<T, TF256, TF128> >;
+
+        struct Accessor;
+        friend Accessor;
+
+    public:
+        Evaluator(LODSpec<std::uint32_t>&& lods_,
+                  NeuralNetVectorType&& neuralNets_,
+                  Vector<std::uint32_t>&& maxLayerOutputCounts_,
+                  OutputInstance::Factory instanceFactory_) :
+            lods{std::move(lods_)},
+            neuralNets{std::move(neuralNets_)},
+            maxLayerOutputCounts{std::move(maxLayerOutputCounts_)},
+            instanceFactory{instanceFactory_} {
+        }
+
+        MachineLearnedBehaviorOutputInstance::Pointer createInstance(MemoryResource* instanceMemRes) const override {
+            return instanceFactory(ConstArrayView<std::uint32_t>{maxLayerOutputCounts}, instanceMemRes);
+        }
+
+        void calculate(ControlsInputInstance* inputs, MachineLearnedBehaviorOutputInstance* intermediateOutputs,
+                       std::uint16_t lod) const override {
+            assert(lod < lods.indicesPerLOD.size());
+            const auto& netIndices = lods.indicesPerLOD[lod];
+            const auto inputBuffer = inputs->getInputBuffer();
+            const auto masks = intermediateOutputs->getMaskBuffer();
+            for (const auto neuralNetIndex : netIndices) {
+                assert(neuralNetIndex < masks.size());
+                assert(neuralNetIndex < neuralNets.size());
+                const float weight = masks[neuralNetIndex];
+                auto layerBuffer1 = static_cast<OutputInstance*>(intermediateOutputs)->getOutputBuffer1(neuralNetIndex);
+                auto layerBuffer2 = static_cast<OutputInstance*>(intermediateOutputs)->getOutputBuffer2(neuralNetIndex);
+                // inputBuffer is also outputBuffer
+                neuralNets[neuralNetIndex].calculate(inputBuffer, layerBuffer1, layerBuffer2, inputBuffer, weight);
+            }
+        }
+
+        void calculate(ControlsInputInstance* inputs,
+                       MachineLearnedBehaviorOutputInstance* intermediateOutputs,
+                       std::uint16_t lod,
+                       std::uint16_t neuralNetIndex) const override {
+            assert(lod < lods.indicesPerLOD.size());
+            static_cast<void>(lod);
+
+            const auto masks = intermediateOutputs->getMaskBuffer();
+            assert(neuralNetIndex < masks.size());
+            assert(neuralNetIndex < neuralNets.size());
+            const float weight = masks[neuralNetIndex];
+            auto inputBuffer = inputs->getInputBuffer();
+            auto layerBuffer1 = static_cast<OutputInstance*>(intermediateOutputs)->getOutputBuffer1(neuralNetIndex);
+            auto layerBuffer2 = static_cast<OutputInstance*>(intermediateOutputs)->getOutputBuffer2(neuralNetIndex);
+            // inputBuffer is also outputBuffer
+            neuralNets[neuralNetIndex].calculate(inputBuffer, layerBuffer1, layerBuffer2, inputBuffer, weight);
+        }
+
+        void load(terse::BinaryInputArchive<BoundedIOStream>& archive) override {
+            archive(lods, neuralNets, maxLayerOutputCounts);
+        }
+
+        void save(terse::BinaryOutputArchive<BoundedIOStream>& archive) override {
+            archive(lods, neuralNets, maxLayerOutputCounts);
+        }
+
+    private:
+        LODSpec<std::uint32_t> lods;
+        NeuralNetVectorType neuralNets;
+        Vector<std::uint32_t> maxLayerOutputCounts;
+        OutputInstance::Factory instanceFactory;
+};
+
+}  // namespace cpu
+
+}  // namespace ml
+
+}  // namespace rl4
diff --git a/Engine/Plugins/Animation/RigLogic/Source/RigLogicLib/Private/riglogic/ml/cpu/CPUMachineLearnedBehaviorFactory.h b/Engine/Plugins/Animation/RigLogic/Source/RigLogicLib/Private/riglogic/ml/cpu/CPUMachineLearnedBehaviorFactory.h
new file mode 100644
--- /dev/null
+++ b/Engine/Plugins/Animation/RigLogic/Source/RigLogicLib/Private/riglogic/ml/cpu/CPUMachineLearnedBehaviorFactory.h
@@ -0,0 +1,160 @@
+// Copyright Epic Games, Inc. All Rights Reserved.
+
+#pragma once
+
+#include "riglogic/TypeDefs.h"
+#include "riglogic/ml/MachineLearnedBehaviorEvaluator.h"
+#include "riglogic/ml/cpu/CPUMachineLearnedBehaviorEvaluator.h"
+#include "riglogic/ml/cpu/CPUMachineLearnedBehaviorOutputInstance.h"
+#include "riglogic/ml/cpu/Inference.h"
+#include "riglogic/ml/cpu/NeuralNet.h"
+#include "riglogic/types/bpcm/Optimizer.h"
+#include "riglogic/types/LODSpec.h"
+#include "riglogic/utils/Extd.h"
+
+namespace rl4 {
+
+namespace ml {
+
+namespace cpu {
+
+template<typename T, typename TF256, typename TF128>
+class Factory {
+    public:
+        static MachineLearnedBehaviorEvaluator::Pointer create(const dna::MachineLearnedBehaviorReader* reader,
+                                                               MemoryResource* memRes) {
+            Vector<NeuralNetInference<T, TF256, TF128> > neuralNets{memRes};
+            Vector<std::uint32_t> maxLayerOutputCountPerNet{memRes};
+            auto instanceFactory = [](ConstArrayView<std::uint32_t> maxLayerOutputCounts, MemoryResource* instanceMemRes) {
+                    using OutputInstancePointer = UniqueInstance<OutputInstance, MachineLearnedBehaviorOutputInstance>;
+                    return OutputInstancePointer::with(instanceMemRes).create(maxLayerOutputCounts, instanceMemRes);
+                };
+            auto factory = UniqueInstance<Evaluator<T, TF256, TF128>, MachineLearnedBehaviorEvaluator>::with(memRes);
+
+            if (reader == nullptr) {
+                return factory.create(LODSpec<std::uint32_t>{memRes},
+                                      std::move(neuralNets),
+                                      std::move(maxLayerOutputCountPerNet),
+                                      instanceFactory);
+            }
+
+            auto lods = computeLODs(reader, memRes);
+            maxLayerOutputCountPerNet.resize(lods.count);
+            neuralNets.reserve(lods.count);
+
+            for (std::uint16_t neuralNetIdx = {}; neuralNetIdx < lods.count; ++neuralNetIdx) {
+                const auto layerCount = reader->getNeuralNetworkLayerCount(neuralNetIdx);
+                if (layerCount != 0u) {
+                    auto net = createNeuralNet(reader,
+                                               neuralNetIdx,
+                                               layerCount,
+                                               &maxLayerOutputCountPerNet[neuralNetIdx],
+                                               memRes);
+                    neuralNets.emplace_back(std::move(net), memRes);
+                }
+            }
+
+            return factory.create(std::move(lods),
+                                  std::move(neuralNets),
+                                  std::move(maxLayerOutputCountPerNet),
+                                  instanceFactory);
+        }
+
+    private:
+        static NeuralNet<T> createNeuralNet(const dna::MachineLearnedBehaviorReader* reader,
+                                            std::uint16_t neuralNetIdx,
+                                            std::uint16_t layerCount,
+                                            std::uint32_t* maxLayerOutputCount,
+                                            MemoryResource* memRes) {
+            const auto inputIndices = reader->getNeuralNetworkInputIndices(neuralNetIdx);
+            const auto outputIndices = reader->getNeuralNetworkOutputIndices(neuralNetIdx);
+            std::uint32_t inputCount = static_cast<std::uint32_t>(inputIndices.size());
+
+            NeuralNet<T> neuralNet{memRes};
+            * maxLayerOutputCount = inputCount;
+            neuralNet.layers.reserve(layerCount);
+            for (std::uint16_t layerIdx = {}; layerIdx < layerCount; ++layerIdx) {
+                const auto weights = reader->getNeuralNetworkLayerWeights(neuralNetIdx, layerIdx);
+                const auto biases = reader->getNeuralNetworkLayerBiases(neuralNetIdx, layerIdx);
+                const auto activationFunction = reader->getNeuralNetworkLayerActivationFunction(neuralNetIdx, layerIdx);
+                const auto activationFunctionParams = reader->getNeuralNetworkLayerActivationFunctionParameters(neuralNetIdx,
+                                                                                                                layerIdx);
+                const auto outputCount = static_cast<std::uint32_t>(biases.size());
+                const auto layer = createLayer(inputCount,
+                                               outputCount,
+                                               weights,
+                                               biases,
+                                               activationFunction,
+                                               activationFunctionParams,
+                                               memRes);
+                // Keep track of the layer with the largest number of outputs
+                // In the next layer, the current output count becomes the input count
+                * maxLayerOutputCount = std::max(*maxLayerOutputCount, layer.weights.padded.rows);
+                neuralNet.layers.push_back(std::move(layer));
+                inputCount = outputCount;
+            }
+
+            neuralNet.inputIndices.assign(inputIndices.begin(), inputIndices.end());
+            neuralNet.outputIndices.assign(outputIndices.begin(), outputIndices.end());
+            return neuralNet;
+        }
+
+        static NeuralNetLayer<T> createLayer(std::uint32_t inputCount,
+                                             std::uint32_t outputCount,
+                                             ConstArrayView<float> weights,
+                                             ConstArrayView<float> biases,
+                                             dna::ActivationFunction activationFunction,
+                                             ConstArrayView<float> activationFunctionParams,
+                                             MemoryResource* memRes) {
+            NeuralNetLayer<T> layer{memRes};
+
+            layer.weights.original = {outputCount, inputCount};
+            const std::uint32_t padding =
+                extd::roundUp(layer.weights.original.rows,
+                              static_cast<std::uint32_t>(TF128::size())) - layer.weights.original.rows;
+            layer.weights.padded = {layer.weights.original.rows + padding, layer.weights.original.cols};
+            layer.weights.rows = PaddedBlockView{
+                layer.weights.original.rows,
+                layer.weights.padded.rows,
+                static_cast<std::uint32_t>(TF256::size()), static_cast<std::uint32_t>(TF128::size())
+            };
+            layer.weights.cols = PaddedBlockView{
+                layer.weights.padded.cols,
+                layer.weights.padded.cols - (layer.weights.padded.cols % 4u),
+                layer.weights.padded.cols - (layer.weights.padded.cols % 8u)
+            };
+            layer.weights.values.resize(layer.weights.padded.size());
+            bpcm::Optimizer<TF128, TF256::size(), TF128::size(), 1u>::optimize(layer.weights.values.data(),
+                                                                               weights.data(),
+                                                                               layer.weights.original);
+
+            layer.biases.resize(layer.weights.padded.rows);
+            bpcm::Optimizer<TF128, TF256::size(), TF128::size(), 1u>::optimize(layer.biases.data(),
+                                                                               biases.data(),
+                                                                               Extent{static_cast<std::uint32_t>(biases.size()),
+                                                                                      1u});
+
+            layer.activationFunction = activationFunction;
+            layer.activationFunctionParameters.assign(activationFunctionParams.begin(), activationFunctionParams.end());
+            return layer;
+        }
+
+        static LODSpec<std::uint32_t> computeLODs(const dna::MachineLearnedBehaviorReader* reader, MemoryResource* memRes) {
+            LODSpec<std::uint32_t> lods{memRes};
+            const auto lodCount = reader->getLODCount();
+            lods.indicesPerLOD.resize(lodCount);
+            for (std::uint16_t lod = {}; lod < lodCount; ++lod) {
+                const auto netIndices = reader->getNeuralNetworkIndicesForLOD(lod);
+                lods.indicesPerLOD[lod].assign(netIndices.begin(), netIndices.end());
+            }
+            lods.count = reader->getNeuralNetworkCount();
+            return lods;
+        }
+
+};
+
+}  // namespace cpu
+
+}  // namespace ml
+
+}  // namespace rl4
diff --git a/Engine/Plugins/Animation/RigLogic/Source/RigLogicLib/Private/riglogic/ml/cpu/CPUMachineLearnedBehaviorOutputInstance.cpp b/Engine/Plugins/Animation/RigLogic/Source/RigLogicLib/Private/riglogic/ml/cpu/CPUMachineLearnedBehaviorOutputInstance.cpp
new file mode 100644
--- /dev/null
+++ b/Engine/Plugins/Animation/RigLogic/Source/RigLogicLib/Private/riglogic/ml/cpu/CPUMachineLearnedBehaviorOutputInstance.cpp
@@ -0,0 +1,49 @@
+// Copyright Epic Games, Inc. All Rights Reserved.
+
+#include "riglogic/ml/cpu/CPUMachineLearnedBehaviorOutputInstance.h"
+
+#include <cassert>
+#include <cstdint>
+
+namespace rl4 {
+
+namespace ml {
+
+namespace cpu {
+
+OutputInstance::OutputInstance(ConstArrayView<std::uint32_t> maxLayerOutputCounts, MemoryResource* memRes) :
+    outputBuffer1{memRes},
+    outputBuffer2{memRes},
+    maskBuffer{maxLayerOutputCounts.size(), 1.0f, memRes} {
+
+    outputBuffer1.resize(maxLayerOutputCounts.size());
+    outputBuffer2.resize(maxLayerOutputCounts.size());
+    for (std::size_t neuralNetIdx = 0ul; neuralNetIdx < maxLayerOutputCounts.size(); ++neuralNetIdx) {
+        outputBuffer1[neuralNetIdx].resize(maxLayerOutputCounts[neuralNetIdx]);
+        outputBuffer2[neuralNetIdx].resize(maxLayerOutputCounts[neuralNetIdx]);
+    }
+}
+
+ArrayView<float> OutputInstance::getOutputBuffer1(std::uint32_t neuralNetIndex) {
+    assert(neuralNetIndex < outputBuffer1.size());
+    return ArrayView<float>{outputBuffer1[neuralNetIndex]};
+}
+
+ArrayView<float> OutputInstance::getOutputBuffer2(std::uint32_t neuralNetIndex) {
+    assert(neuralNetIndex < outputBuffer2.size());
+    return ArrayView<float>{outputBuffer2[neuralNetIndex]};
+}
+
+ArrayView<float> OutputInstance::getMaskBuffer() {
+    return ArrayView<float>{maskBuffer};
+}
+
+ConstArrayView<float> OutputInstance::getMaskBuffer() const {
+    return ConstArrayView<float>{maskBuffer};
+}
+
+}  // namespace cpu
+
+}  // namespace ml
+
+}  // namespace rl4
diff --git a/Engine/Plugins/Animation/RigLogic/Source/RigLogicLib/Private/riglogic/ml/cpu/CPUMachineLearnedBehaviorOutputInstance.h b/Engine/Plugins/Animation/RigLogic/Source/RigLogicLib/Private/riglogic/ml/cpu/CPUMachineLearnedBehaviorOutputInstance.h
new file mode 100644
--- /dev/null
+++ b/Engine/Plugins/Animation/RigLogic/Source/RigLogicLib/Private/riglogic/ml/cpu/CPUMachineLearnedBehaviorOutputInstance.h
@@ -0,0 +1,45 @@
+// Copyright Epic Games, Inc. All Rights Reserved.
+
+#pragma once
+
+#include "riglogic/TypeDefs.h"
+#include "riglogic/ml/MachineLearnedBehaviorOutputInstance.h"
+
+#ifdef _MSC_VER
+    #pragma warning(push)
+    #pragma warning(disable : 4365 4987)
+#endif
+#include <functional>
+#ifdef _MSC_VER
+    #pragma warning(pop)
+#endif
+
+namespace rl4 {
+
+namespace ml {
+
+namespace cpu {
+
+class OutputInstance : public MachineLearnedBehaviorOutputInstance {
+    public:
+        using Factory = std::function<Pointer(ConstArrayView<std::uint32_t>, MemoryResource*)>;
+
+    public:
+        OutputInstance(ConstArrayView<std::uint32_t> maxLayerOutputCounts, MemoryResource* memRes);
+        ArrayView<float> getOutputBuffer1(std::uint32_t neuralNetIndex);
+        ArrayView<float> getOutputBuffer2(std::uint32_t neuralNetIndex);
+        ArrayView<float> getMaskBuffer() override;
+        ConstArrayView<float> getMaskBuffer() const override;
+
+    private:
+        Vector<AlignedVector<float> > outputBuffer1;
+        Vector<AlignedVector<float> > outputBuffer2;
+        Vector<float> maskBuffer;
+
+};
+
+}  // namespace cpu
+
+}  // namespace ml
+
+}  // namespace rl4
diff --git a/Engine/Plugins/Animation/RigLogic/Source/RigLogicLib/Private/riglogic/ml/cpu/Inference.h b/Engine/Plugins/Animation/RigLogic/Source/RigLogicLib/Private/riglogic/ml/cpu/Inference.h
--- a/Engine/Plugins/Animation/RigLogic/Source/RigLogicLib/Private/riglogic/ml/cpu/Inference.h
+++ b/Engine/Plugins/Animation/RigLogic/Source/RigLogicLib/Private/riglogic/ml/cpu/Inference.h
@@ -2,9 +2,6 @@
 
 #pragma once
 
-// *INDENT-OFF*
-#ifdef RL_BUILD_WITH_ML_EVALUATOR
-
 #include "riglogic/TypeDefs.h"
 #include "riglogic/ml/cpu/NeuralNet.h"
 #include "riglogic/ml/cpu/layers/LeakyReLULayerEvaluator.h"
@@ -111,6 +108,3 @@ struct NeuralNetInference {
 }  // namespace ml
 
 }  // namespace rl4
-
-#endif  // RL_BUILD_WITH_ML_EVALUATOR
-// *INDENT-ON*
diff --git a/Engine/Plugins/Animation/RigLogic/Source/RigLogicLib/Private/riglogic/ml/cpu/NeuralNet.h b/Engine/Plugins/Animation/RigLogic/Source/RigLogicLib/Private/riglogic/ml/cpu/NeuralNet.h
--- a/Engine/Plugins/Animation/RigLogic/Source/RigLogicLib/Private/riglogic/ml/cpu/NeuralNet.h
+++ b/Engine/Plugins/Animation/RigLogic/Source/RigLogicLib/Private/riglogic/ml/cpu/NeuralNet.h
@@ -2,11 +2,8 @@
 
 #pragma once
 
-// *INDENT-OFF*
-#ifdef RL_BUILD_WITH_ML_EVALUATOR
-
 #include "riglogic/TypeDefs.h"
-#include "riglogic/types/AlignedBlockView.h"
+#include "riglogic/types/PaddedBlockView.h"
 #include "riglogic/types/Extent.h"
 
 #include <cstdint>
@@ -21,8 +18,8 @@ template<typename T>
 struct WeightMatrix {
     Extent original;
     Extent padded;
-    AlignedBlockView rows;
-    AlignedBlockView cols;
+    PaddedBlockView rows;
+    PaddedBlockView cols;
     AlignedVector<T> values;
 
     explicit WeightMatrix(MemoryResource* memRes) :
@@ -89,6 +86,3 @@ struct NeuralNet {
 }  // namespace ml
 
 }  // namespace rl4
-
-#endif  // RL_BUILD_WITH_ML_EVALUATOR
-// *INDENT-ON*
diff --git a/Engine/Plugins/Animation/RigLogic/Source/RigLogicLib/Private/riglogic/ml/cpu/layers/LayerEvaluator.h b/Engine/Plugins/Animation/RigLogic/Source/RigLogicLib/Private/riglogic/ml/cpu/layers/LayerEvaluator.h
--- a/Engine/Plugins/Animation/RigLogic/Source/RigLogicLib/Private/riglogic/ml/cpu/layers/LayerEvaluator.h
+++ b/Engine/Plugins/Animation/RigLogic/Source/RigLogicLib/Private/riglogic/ml/cpu/layers/LayerEvaluator.h
@@ -2,11 +2,9 @@
 
 #pragma once
 
-// *INDENT-OFF*
-#ifdef RL_BUILD_WITH_ML_EVALUATOR
-
 #include "riglogic/TypeDefs.h"
 #include "riglogic/ml/cpu/NeuralNet.h"
+#include "riglogic/utils/Macros.h"
 
 namespace rl4 {
 
@@ -25,11 +23,207 @@ class LayerEvaluator {
         virtual void calculate(const NeuralNetLayer<T>& layer, ConstArrayView<float> inputs, ArrayView<float> outputs) const = 0;
 };
 
-}  // namespace cpu
+template<typename TF256, typename T>
+static FORCE_INLINE void processBlocks8x1(const float* inputVectorEndAlignedTo4,
+                                          const float* inputVectorEnd,
+                                          const T*& weights,
+                                          TF256& sum) {
+    TF256 remainder{};
+    for (const float* inputVector = inputVectorEndAlignedTo4;
+         inputVector < inputVectorEnd;
+         ++inputVector, weights += TF256::size()) {
+        const TF256 input{*inputVector};
+        const TF256 blk = TF256::fromAlignedSource(weights);
+        remainder += (blk * input);
+    }
+    sum += remainder;
+}
 
-}  // namespace ml
+template<typename TF256, typename TActivationFunction, typename T>
+static FORCE_INLINE void processBlocks8x4(const float* inputVectorStart,
+                                          const float* inputVectorEndAlignedTo4,
+                                          const float* inputVectorEnd,
+                                          const T*& weights,
+                                          const T* biases,
+                                          const float* activationParams,
+                                          float* outbuf) {
+    TF256 sum1{};
+    TF256 sum2{};
+    TF256 sum3{};
+    TF256 sum4{};
+    for (const float* inputVector = inputVectorStart;
+         inputVector < inputVectorEndAlignedTo4;
+         inputVector += 4ul, weights += (TF256::size() * 4ul)) {
+        const TF256 input1{inputVector[0]};
+        const TF256 input2{inputVector[1]};
+        const TF256 input3{inputVector[2]};
+        const TF256 input4{inputVector[3]};
+        const TF256 blk1 = TF256::fromAlignedSource(weights + TF256::size() * 0);
+        const TF256 blk2 = TF256::fromAlignedSource(weights + TF256::size() * 1);
+        const TF256 blk3 = TF256::fromAlignedSource(weights + TF256::size() * 2);
+        const TF256 blk4 = TF256::fromAlignedSource(weights + TF256::size() * 3);
+        /*
+        #ifndef RL_NO_MANUAL_PREFETCH
+            #ifdef RL_USE_NON_TEMPORAL_PREFETCH
+                TF256::prefetchNTA(values + (lookAheadOffset * 6ul));
+                TF256::prefetchNTA(values + (lookAheadOffset * 7ul));
+                TF256::prefetchNTA(values + (lookAheadOffset * 8ul));
+                TF256::prefetchNTA(values + (lookAheadOffset * 9ul));
+                TF256::prefetchNTA(values + (lookAheadOffset * 10ul));
+                TF256::prefetchNTA(values + (lookAheadOffset * 11ul));
+            #else
+                TF256::prefetchT0(values + (lookAheadOffset * 6ul));
+                TF256::prefetchT0(values + (lookAheadOffset * 7ul));
+                TF256::prefetchT0(values + (lookAheadOffset * 8ul));
+                TF256::prefetchT0(values + (lookAheadOffset * 9ul));
+                TF256::prefetchT0(values + (lookAheadOffset * 10ul));
+                TF256::prefetchT0(values + (lookAheadOffset * 11ul));
+            #endif
+        #endif
+        */
+        sum1 += (blk1 * input1);
+        sum2 += (blk2 * input2);
+        sum3 += (blk3 * input3);
+        sum4 += (blk4 * input4);
+    }
+    // Process 8x1 horizontal remainder portion after 8x4 blocks are consumed
+    processBlocks8x1(inputVectorEndAlignedTo4, inputVectorEnd, weights, sum1);
+
+    const TF256 bias = TF256::fromAlignedSource(biases);
+
+    sum1 += sum3;
+    sum2 += sum4;
+    sum1 += sum2;
+
+    sum1 += bias;
+
+    TActivationFunction{} (sum1, activationParams);
+
+    sum1.alignedStore(outbuf);
+}
+
+template<typename TF128, typename T>
+static FORCE_INLINE void processBlocks4x1(const float* inputVectorEndAlignedTo8,
+                                          const float* inputVectorEnd,
+                                          const T*& weights,
+                                          TF128& sum1) {
+    TF128 remainder{};
+    for (const float* inputVector = inputVectorEndAlignedTo8;
+         inputVector < inputVectorEnd;
+         ++inputVector, weights += TF128::size()) {
+        const TF128 input{*inputVector};
+        const TF128 blk = TF128::fromAlignedSource(weights);
+        remainder += (blk * input);
+    }
+    sum1 += remainder;
+}
+
+template<typename TF128, typename TActivationFunction, typename T>
+static FORCE_INLINE void processBlocks4x8(const float* inputVectorStart,
+                                          const float* inputVectorEndAlignedTo8,
+                                          const float* inputVectorEnd,
+                                          const T*& weights,
+                                          const T* biases,
+                                          const float* activationParams,
+                                          float* outbuf) {
+    TF128 sum1{};
+    TF128 sum2{};
+    TF128 sum3{};
+    TF128 sum4{};
+    TF128 sum5{};
+    TF128 sum6{};
+    TF128 sum7{};
+    TF128 sum8{};
+    for (const float* inputVector = inputVectorStart;
+         inputVector < inputVectorEndAlignedTo8;
+         inputVector += 8ul, weights += (TF128::size() * 8ul)) {
+        const TF128 input1{inputVector[0]};
+        const TF128 input2{inputVector[1]};
+        const TF128 input3{inputVector[2]};
+        const TF128 input4{inputVector[3]};
+        const TF128 input5{inputVector[4]};
+        const TF128 input6{inputVector[5]};
+        const TF128 input7{inputVector[6]};
+        const TF128 input8{inputVector[7]};
+        const TF128 blk1 = TF128::fromAlignedSource(weights);
+        const TF128 blk2 = TF128::fromAlignedSource(weights + TF128::size());
+        const TF128 blk3 = TF128::fromAlignedSource(weights + TF128::size() * 2);
+        const TF128 blk4 = TF128::fromAlignedSource(weights + TF128::size() * 3);
+        const TF128 blk5 = TF128::fromAlignedSource(weights + TF128::size() * 4);
+        const TF128 blk6 = TF128::fromAlignedSource(weights + TF128::size() * 5);
+        const TF128 blk7 = TF128::fromAlignedSource(weights + TF128::size() * 6);
+        const TF128 blk8 = TF128::fromAlignedSource(weights + TF128::size() * 7);
+        sum1 += (blk1 * input1);
+        sum2 += (blk2 * input2);
+        sum3 += (blk3 * input3);
+        sum4 += (blk4 * input4);
+        sum5 += (blk5 * input5);
+        sum6 += (blk6 * input6);
+        sum7 += (blk7 * input7);
+        sum8 += (blk8 * input8);
+    }
+    // Process 4x1 horizontal remainder portion after 4x8 blocks are consumed
+    processBlocks4x1(inputVectorEndAlignedTo8, inputVectorEnd, weights, sum1);
+
+    const TF128 bias1 = TF128::fromAlignedSource(biases);
+
+    sum1 += sum2;
+    sum3 += sum4;
+    sum5 += sum6;
+    sum7 += sum8;
+    sum1 += sum3;
+    sum5 += sum7;
+    sum1 += sum5;
+
+    sum1 += bias1;
+
+    TActivationFunction{} (sum1, activationParams);
+
+    sum1.alignedStore(outbuf);
+}
+
+template<typename TF256, typename TF128, template<class ...> class TActivationFunction, typename T>
+static FORCE_INLINE void calculateBlock4(const NeuralNetLayer<T>& layer, ConstArrayView<float> inputs, ArrayView<float> outputs) {
+    float* outputVector = outputs.data();
+    const float* outputVectorEndPaddedToLastFullBlock = outputVector + layer.weights.rows.sizePaddedToLastFullBlock;
+    float* outputVectorEnd = outputVector + layer.weights.rows.size;
+
+    const float* inputVector = inputs.data();
+    const float* inputVectorEndPaddedTo4 = inputVector + layer.weights.cols.sizePaddedToLastFullBlock;
+    const float* inputVectorEndPaddedTo8 = inputVector + layer.weights.cols.sizePaddedToSecondLastFullBlock;
+    const float* inputVectorEnd = inputVector + layer.weights.cols.size;
+
+    const T* weights = layer.weights.values.data();
+    const T* biases = layer.biases.data();
+    const float* activationParams = layer.activationFunctionParameters.data();
+
+    for (; outputVector < outputVectorEndPaddedToLastFullBlock; outputVector += TF256::size(), biases += TF256::size()) {
+        alignas(TF256::alignment()) float outbuf[TF256::size()];
+        processBlocks8x4<TF256, TActivationFunction<TF256> >(inputVector,
+                                                             inputVectorEndPaddedTo4,
+                                                             inputVectorEnd,
+                                                             weights,
+                                                             biases,
+                                                             activationParams,
+                                                             outbuf);
+        std::memcpy(outputVector, outbuf, sizeof(outbuf));
+    }
+
+    for (; outputVector < outputVectorEnd; outputVector += TF128::size(), biases += TF128::size()) {
+        alignas(TF128::alignment()) float outbuf[TF128::size()];
+        processBlocks4x8<TF128, TActivationFunction<TF128> >(inputVector,
+                                                             inputVectorEndPaddedTo8,
+                                                             inputVectorEnd,
+                                                             weights,
+                                                             biases,
+                                                             activationParams,
+                                                             outbuf);
+        std::memcpy(outputVector, outbuf, sizeof(outbuf));
+    }
+}
 
 }  // namespace rl4
 
-#endif  // RL_BUILD_WITH_ML_EVALUATOR
-// *INDENT-ON*
+}  // namespace ml
+
+}  // namespace cpu
diff --git a/Engine/Plugins/Animation/RigLogic/Source/RigLogicLib/Private/riglogic/ml/cpu/layers/LeakyReLULayerEvaluator.h b/Engine/Plugins/Animation/RigLogic/Source/RigLogicLib/Private/riglogic/ml/cpu/layers/LeakyReLULayerEvaluator.h
--- a/Engine/Plugins/Animation/RigLogic/Source/RigLogicLib/Private/riglogic/ml/cpu/layers/LeakyReLULayerEvaluator.h
+++ b/Engine/Plugins/Animation/RigLogic/Source/RigLogicLib/Private/riglogic/ml/cpu/layers/LeakyReLULayerEvaluator.h
@@ -2,12 +2,7 @@
 
 #pragma once
 
-// *INDENT-OFF*
-#ifdef RL_BUILD_WITH_ML_EVALUATOR
-
 #include "riglogic/TypeDefs.h"
-#include "riglogic/ml/cpu/Block4.h"
-#include "riglogic/ml/cpu/Consts.h"
 #include "riglogic/ml/cpu/layers/LayerEvaluator.h"
 #include "riglogic/system/simd/SIMD.h"
 #include "riglogic/utils/Macros.h"
@@ -45,6 +40,3 @@ class LeakyReLULayerEvaluator : public LayerEvaluator<T> {
 }  // namespace ml
 
 }  // namespace rl4
-
-#endif  // RL_BUILD_WITH_ML_EVALUATOR
-// *INDENT-ON*
diff --git a/Engine/Plugins/Animation/RigLogic/Source/RigLogicLib/Private/riglogic/ml/cpu/layers/LinearLayerEvaluator.h b/Engine/Plugins/Animation/RigLogic/Source/RigLogicLib/Private/riglogic/ml/cpu/layers/LinearLayerEvaluator.h
--- a/Engine/Plugins/Animation/RigLogic/Source/RigLogicLib/Private/riglogic/ml/cpu/layers/LinearLayerEvaluator.h
+++ b/Engine/Plugins/Animation/RigLogic/Source/RigLogicLib/Private/riglogic/ml/cpu/layers/LinearLayerEvaluator.h
@@ -2,12 +2,7 @@
 
 #pragma once
 
-// *INDENT-OFF*
-#ifdef RL_BUILD_WITH_ML_EVALUATOR
-
 #include "riglogic/TypeDefs.h"
-#include "riglogic/ml/cpu/Block4.h"
-#include "riglogic/ml/cpu/Consts.h"
 #include "riglogic/ml/cpu/layers/LayerEvaluator.h"
 #include "riglogic/utils/Macros.h"
 
@@ -39,6 +34,3 @@ class LinearLayerEvaluator : public LayerEvaluator<T> {
 }  // namespace ml
 
 }  // namespace rl4
-
-#endif  // RL_BUILD_WITH_ML_EVALUATOR
-// *INDENT-ON*
diff --git a/Engine/Plugins/Animation/RigLogic/Source/RigLogicLib/Private/riglogic/ml/cpu/layers/Utils.h b/Engine/Plugins/Animation/RigLogic/Source/RigLogicLib/Private/riglogic/ml/cpu/layers/Utils.h
--- a/Engine/Plugins/Animation/RigLogic/Source/RigLogicLib/Private/riglogic/ml/cpu/layers/Utils.h
+++ b/Engine/Plugins/Animation/RigLogic/Source/RigLogicLib/Private/riglogic/ml/cpu/layers/Utils.h
@@ -2,9 +2,6 @@
 
 #pragma once
 
-// *INDENT-OFF*
-#ifdef RL_BUILD_WITH_ML_EVALUATOR
-
 #include <type_traits>
 
 namespace rl4 {
@@ -23,6 +20,3 @@ struct HasSize {
 }  // namespace ml
 
 }  // namespace rl4
-
-#endif  // RL_BUILD_WITH_ML_EVALUATOR
-// *INDENT-ON*
diff --git a/Engine/Plugins/Animation/RigLogic/Source/RigLogicLib/Private/riglogic/psdmatrix/PSDMatrix.cpp b/Engine/Plugins/Animation/RigLogic/Source/RigLogicLib/Private/riglogic/psdmatrix/PSDMatrix.cpp
--- a/Engine/Plugins/Animation/RigLogic/Source/RigLogicLib/Private/riglogic/psdmatrix/PSDMatrix.cpp
+++ b/Engine/Plugins/Animation/RigLogic/Source/RigLogicLib/Private/riglogic/psdmatrix/PSDMatrix.cpp
@@ -3,6 +3,7 @@
 #include "riglogic/psdmatrix/PSDMatrix.h"
 
 #include "riglogic/TypeDefs.h"
+#include "riglogic/utils/Extd.h"
 
 #ifdef _MSC_VER
     #pragma warning(push)
@@ -18,7 +19,8 @@ namespace rl4 {
 
 namespace {
 
-constexpr float maxControlValue = 1.0f;
+constexpr float minPSDValue = 0.0f;
+constexpr float maxPSDValue = 1.0f;
 constexpr std::size_t blockSize = 4ul;
 
 }  // namespace
@@ -61,10 +63,10 @@ void PSDMatrix::calculate(float* inputs, std::uint16_t rawControlCount) const {
         const float weight1 = values[i + 1ul];
         const float weight2 = values[i + 2ul];
         const float weight3 = values[i + 3ul];
-        const float input0 = inputs[col0];
-        const float input1 = inputs[col1];
-        const float input2 = inputs[col2];
-        const float input3 = inputs[col3];
+        const float input0 = extd::clamp(inputs[col0], minPSDValue, maxPSDValue);
+        const float input1 = extd::clamp(inputs[col1], minPSDValue, maxPSDValue);
+        const float input2 = extd::clamp(inputs[col2], minPSDValue, maxPSDValue);
+        const float input3 = extd::clamp(inputs[col3], minPSDValue, maxPSDValue);
         inputs[row0] *= weight0 * input0;
         inputs[row1] *= weight1 * input1;
         inputs[row2] *= weight2 * input2;
@@ -75,13 +77,13 @@ void PSDMatrix::calculate(float* inputs, std::uint16_t rawControlCount) const {
         const std::uint16_t row = rowIndices[i];
         const std::uint16_t col = columnIndices[i];
         const float weight = values[i];
-        const float input = inputs[col];
+        const float input = extd::clamp(inputs[col], minPSDValue, maxPSDValue);
         inputs[row] *= weight * input;
     }
 
     auto psds = inputs + rawControlCount;
     for (std::size_t i = 0ul; i < distinctPSDs; ++i) {
-        psds[i] = std::min(maxControlValue, psds[i]);
+        psds[i] = std::min(maxPSDValue, psds[i]);
     }
 }
 
diff --git a/Engine/Plugins/Animation/RigLogic/Source/RigLogicLib/Private/riglogic/rbf/RBFBehaviorFactory.cpp b/Engine/Plugins/Animation/RigLogic/Source/RigLogicLib/Private/riglogic/rbf/RBFBehaviorFactory.cpp
new file mode 100644
--- /dev/null
+++ b/Engine/Plugins/Animation/RigLogic/Source/RigLogicLib/Private/riglogic/rbf/RBFBehaviorFactory.cpp
@@ -0,0 +1,75 @@
+// Copyright Epic Games, Inc. All Rights Reserved.
+
+#include "riglogic/rbf/RBFBehaviorFactory.h"
+
+#include "riglogic/rbf/RBFBehaviorEvaluator.h"
+#include "riglogic/rbf/RBFBehaviorNullEvaluator.h"
+#include "riglogic/rbf/cpu/CPURBFBehaviorFactory.h"
+#include "riglogic/riglogic/Configuration.h"
+#include "riglogic/riglogic/RigMetrics.h"
+#include "riglogic/system/simd/SIMD.h"
+
+#include <cstdint>
+
+namespace rl4 {
+
+RBFBehaviorEvaluator::Pointer createRBFEvaluator(const Configuration& config, const dna::Reader* reader, MemoryResource* memRes) {
+    #if defined(__clang__)
+        #pragma clang diagnostic push
+        #pragma clang diagnostic ignored "-Wunused-local-typedef"
+    #endif
+    #if !defined(__clang__) && defined(__GNUC__)
+        #pragma GCC diagnostic push
+        #pragma GCC diagnostic ignored "-Wunused-local-typedefs"
+    #endif
+    using StorageValueType = float;
+    #if !defined(__clang__) && defined(__GNUC__)
+        #pragma GCC diagnostic pop
+    #endif
+    #if defined(__clang__)
+        #pragma clang diagnostic pop
+    #endif
+
+    // Work around unused parameter warning when building without SSE and AVX
+    static_cast<void>(config);
+    #ifdef RL_BUILD_WITH_SSE
+        if ((config.calculationType == CalculationType::SSE) || (config.calculationType == CalculationType::AnyVector)) {
+            return rbf::cpu::Factory<StorageValueType, trimd::sse::F256, trimd::sse::F128>::create(reader, memRes);
+        }
+    #endif  // RL_BUILD_WITH_SSE
+    #ifdef RL_BUILD_WITH_AVX
+        if ((config.calculationType == CalculationType::AVX) || (config.calculationType == CalculationType::AnyVector)) {
+            // Use 256-bit AVX registers and whatever 128-bit width type is available
+            return rbf::cpu::Factory<StorageValueType, trimd::avx::F256, trimd::sse::F128>::create(reader, memRes);
+        }
+    #endif  // RL_BUILD_WITH_AVX
+    #ifdef RL_BUILD_WITH_NEON
+        if ((config.calculationType == CalculationType::NEON) || (config.calculationType == CalculationType::AnyVector)) {
+            return rbf::cpu::Factory<StorageValueType, trimd::neon::F256, trimd::neon::F128>::create(reader, memRes);
+        }
+    #endif  // RL_BUILD_WITH_NEON
+    return rbf::cpu::Factory<float, trimd::scalar::F256, trimd::scalar::F128>::create(reader, memRes);
+}
+
+RBFBehavior::Pointer RBFBehaviorFactory::create(const Configuration& config, const dna::Reader* reader, MemoryResource* memRes) {
+    auto moduleFactory = UniqueInstance<RBFBehavior>::with(memRes);
+    if (!config.loadRBFBehavior || (reader->getRBFSolverCount() == 0u) || (reader->getRBFPoseCount() == 0u)) {
+        auto evaluator =
+            UniqueInstance<RBFBehaviorNullEvaluator, RBFBehaviorEvaluator>::with(memRes).create();
+        return moduleFactory.create(std::move(evaluator));
+    }
+
+    return moduleFactory.create(createRBFEvaluator(config, reader, memRes));
+}
+
+RBFBehavior::Pointer RBFBehaviorFactory::create(const Configuration& config, const RigMetrics& metrics, MemoryResource* memRes) {
+    auto moduleFactory = UniqueInstance<RBFBehavior>::with(memRes);
+    if (!config.loadRBFBehavior || (metrics.rbfControlCount == 0u) || ((metrics.rbfSolverCount) == 0u)) {
+        auto evaluator =
+            UniqueInstance<RBFBehaviorNullEvaluator, RBFBehaviorEvaluator>::with(memRes).create();
+        return moduleFactory.create(std::move(evaluator));
+    }
+    return moduleFactory.create(createRBFEvaluator(config, nullptr, memRes));
+}
+
+}  // namespace rl4
diff --git a/Engine/Plugins/Animation/RigLogic/Source/RigLogicLib/Private/riglogic/rbf/RBFBehaviorFactory.h b/Engine/Plugins/Animation/RigLogic/Source/RigLogicLib/Private/riglogic/rbf/RBFBehaviorFactory.h
new file mode 100644
--- /dev/null
+++ b/Engine/Plugins/Animation/RigLogic/Source/RigLogicLib/Private/riglogic/rbf/RBFBehaviorFactory.h
@@ -0,0 +1,19 @@
+// Copyright Epic Games, Inc. All Rights Reserved.
+
+#pragma once
+
+#include "riglogic/TypeDefs.h"
+#include "riglogic/rbf/RBFBehavior.h"
+
+namespace rl4 {
+
+struct Configuration;
+struct RigMetrics;
+
+struct RBFBehaviorFactory {
+    static RBFBehavior::Pointer create(const Configuration& config, const dna::Reader* reader, MemoryResource* memRes);
+    static RBFBehavior::Pointer create(const Configuration& config, const RigMetrics& metrics, MemoryResource* memRes);
+
+};
+
+}  // namespace rl4
diff --git a/Engine/Plugins/Animation/RigLogic/Source/RigLogicLib/Private/riglogic/rbf/RBFBehaviorNullEvaluator.cpp b/Engine/Plugins/Animation/RigLogic/Source/RigLogicLib/Private/riglogic/rbf/RBFBehaviorNullEvaluator.cpp
new file mode 100644
--- /dev/null
+++ b/Engine/Plugins/Animation/RigLogic/Source/RigLogicLib/Private/riglogic/rbf/RBFBehaviorNullEvaluator.cpp
@@ -0,0 +1,31 @@
+// Copyright Epic Games, Inc. All Rights Reserved.
+
+#include "riglogic/rbf/RBFBehaviorNullEvaluator.h"
+
+#include "riglogic/rbf/RBFBehaviorNullOutputInstance.h"
+
+namespace rl4 {
+
+RBFBehaviorOutputInstance::Pointer RBFBehaviorNullEvaluator::createInstance(MemoryResource* instanceMemRes)
+const {
+    return UniqueInstance<RBFBehaviorNullOutputInstance,
+                          RBFBehaviorOutputInstance>::with(instanceMemRes).create();
+}
+
+void RBFBehaviorNullEvaluator::calculate(ControlsInputInstance*  /*unused*/, RBFBehaviorOutputInstance*  /*unused*/,
+                                         std::uint16_t  /*unused*/) const {
+}
+
+void RBFBehaviorNullEvaluator::calculate(ControlsInputInstance*  /*unused*/,
+                                         RBFBehaviorOutputInstance*  /*unused*/,
+                                         std::uint16_t  /*unused*/,
+                                         std::uint16_t  /*unused*/) const {
+}
+
+void RBFBehaviorNullEvaluator::load(terse::BinaryInputArchive<BoundedIOStream>&  /*unused*/) {
+}
+
+void RBFBehaviorNullEvaluator::save(terse::BinaryOutputArchive<BoundedIOStream>&  /*unused*/) {
+}
+
+}  // namespace rl4
diff --git a/Engine/Plugins/Animation/RigLogic/Source/RigLogicLib/Private/riglogic/rbf/RBFBehaviorNullEvaluator.h b/Engine/Plugins/Animation/RigLogic/Source/RigLogicLib/Private/riglogic/rbf/RBFBehaviorNullEvaluator.h
new file mode 100644
--- /dev/null
+++ b/Engine/Plugins/Animation/RigLogic/Source/RigLogicLib/Private/riglogic/rbf/RBFBehaviorNullEvaluator.h
@@ -0,0 +1,27 @@
+// Copyright Epic Games, Inc. All Rights Reserved.
+
+#pragma once
+
+#include "riglogic/TypeDefs.h"
+#include "riglogic/rbf/RBFBehaviorEvaluator.h"
+#include "riglogic/rbf/RBFBehaviorOutputInstance.h"
+
+namespace rl4 {
+
+class ControlsInputInstance;
+
+class RBFBehaviorNullEvaluator : public RBFBehaviorEvaluator {
+    public:
+        RBFBehaviorOutputInstance::Pointer createInstance(MemoryResource* instanceMemRes) const override;
+        void calculate(ControlsInputInstance*  /*unused*/, RBFBehaviorOutputInstance*  /*unused*/,
+                       std::uint16_t  /*unused*/) const override;
+        void calculate(ControlsInputInstance*  /*unused*/,
+                       RBFBehaviorOutputInstance*  /*unused*/,
+                       std::uint16_t  /*unused*/,
+                       std::uint16_t  /*unused*/) const override;
+        void load(terse::BinaryInputArchive<BoundedIOStream>&  /*unused*/) override;
+        void save(terse::BinaryOutputArchive<BoundedIOStream>&  /*unused*/) override;
+
+};
+
+}  // namespace rl4
diff --git a/Engine/Plugins/Animation/RigLogic/Source/RigLogicLib/Private/riglogic/rbf/RBFBehaviorNullOutputInstance.cpp b/Engine/Plugins/Animation/RigLogic/Source/RigLogicLib/Private/riglogic/rbf/RBFBehaviorNullOutputInstance.cpp
new file mode 100644
--- /dev/null
+++ b/Engine/Plugins/Animation/RigLogic/Source/RigLogicLib/Private/riglogic/rbf/RBFBehaviorNullOutputInstance.cpp
@@ -0,0 +1,19 @@
+// Copyright Epic Games, Inc. All Rights Reserved.
+
+#include "riglogic/rbf/RBFBehaviorNullOutputInstance.h"
+
+namespace rl4 {
+
+ArrayView<float> RBFBehaviorNullOutputInstance::getInputBuffer() {
+    return {};
+}
+
+ArrayView<float> RBFBehaviorNullOutputInstance::getIntermediateWeightsBuffer() {
+    return {};
+}
+
+ArrayView<float> RBFBehaviorNullOutputInstance::getOutputWeightsBuffer() {
+    return {};
+}
+
+}  // namespace rl4
diff --git a/Engine/Plugins/Animation/RigLogic/Source/RigLogicLib/Private/riglogic/rbf/RBFBehaviorNullOutputInstance.h b/Engine/Plugins/Animation/RigLogic/Source/RigLogicLib/Private/riglogic/rbf/RBFBehaviorNullOutputInstance.h
new file mode 100644
--- /dev/null
+++ b/Engine/Plugins/Animation/RigLogic/Source/RigLogicLib/Private/riglogic/rbf/RBFBehaviorNullOutputInstance.h
@@ -0,0 +1,17 @@
+// Copyright Epic Games, Inc. All Rights Reserved.
+
+#pragma once
+
+#include "riglogic/rbf/RBFBehaviorOutputInstance.h"
+
+namespace rl4 {
+
+class RBFBehaviorNullOutputInstance : public RBFBehaviorOutputInstance {
+    public:
+        ArrayView<float> getInputBuffer() override;
+        ArrayView<float> getIntermediateWeightsBuffer() override;
+        ArrayView<float> getOutputWeightsBuffer() override;
+
+};
+
+}  // namespace rl4
diff --git a/Engine/Plugins/Animation/RigLogic/Source/RigLogicLib/Private/riglogic/rbf/RBFBehaviorOutputInstance.h b/Engine/Plugins/Animation/RigLogic/Source/RigLogicLib/Private/riglogic/rbf/RBFBehaviorOutputInstance.h
new file mode 100644
--- /dev/null
+++ b/Engine/Plugins/Animation/RigLogic/Source/RigLogicLib/Private/riglogic/rbf/RBFBehaviorOutputInstance.h
@@ -0,0 +1,21 @@
+// Copyright Epic Games, Inc. All Rights Reserved.
+
+#pragma once
+
+#include "riglogic/TypeDefs.h"
+
+namespace rl4 {
+
+class RBFBehaviorOutputInstance {
+    public:
+        using Pointer = UniqueInstance<RBFBehaviorOutputInstance>::PointerType;
+
+    public:
+        virtual ~RBFBehaviorOutputInstance();
+        virtual ArrayView<float> getInputBuffer() = 0;
+        virtual ArrayView<float> getIntermediateWeightsBuffer() = 0;
+        virtual ArrayView<float> getOutputWeightsBuffer() = 0;
+
+};
+
+}  // namespace rl4
diff --git a/Engine/Plugins/Animation/RigLogic/Source/RigLogicLib/Private/riglogic/rbf/cpu/AdditiveRBFSolver.cpp b/Engine/Plugins/Animation/RigLogic/Source/RigLogicLib/Private/riglogic/rbf/cpu/AdditiveRBFSolver.cpp
new file mode 100644
--- /dev/null
+++ b/Engine/Plugins/Animation/RigLogic/Source/RigLogicLib/Private/riglogic/rbf/cpu/AdditiveRBFSolver.cpp
@@ -0,0 +1,24 @@
+// Copyright Epic Games, Inc. All Rights Reserved.
+
+#include "riglogic/rbf/cpu/AdditiveRBFSolver.h"
+
+namespace rl4 {
+
+AdditiveRBFSolver::AdditiveRBFSolver(MemoryResource* memRes) : RBFSolver(memRes) {
+}
+
+AdditiveRBFSolver::AdditiveRBFSolver(const RBFSolverRecipe& recipe, MemoryResource* memRes) : RBFSolver(recipe, memRes) {
+}
+
+RBFSolverType AdditiveRBFSolver::getSolverType() const {
+    return RBFSolverType::Additive;
+}
+
+void AdditiveRBFSolver::solve(ArrayView<float> input, ArrayView<float>  /*unused*/, ArrayView<float> outputWeights) const {
+    convertInput(input);
+    getDistanceWeight(targets, input, outputWeights, radius);
+
+    normalizeAndCutOff(outputWeights);
+}
+
+}  // namespace rl4
diff --git a/Engine/Plugins/Animation/RigLogic/Source/RigLogicLib/Private/riglogic/rbf/cpu/CPURBFBehaviorFactory.h b/Engine/Plugins/Animation/RigLogic/Source/RigLogicLib/Private/riglogic/rbf/cpu/CPURBFBehaviorFactory.h
new file mode 100644
--- /dev/null
+++ b/Engine/Plugins/Animation/RigLogic/Source/RigLogicLib/Private/riglogic/rbf/cpu/CPURBFBehaviorFactory.h
@@ -0,0 +1,140 @@
+// Copyright Epic Games, Inc. All Rights Reserved.
+
+#pragma once
+
+#include "riglogic/TypeDefs.h"
+#include "riglogic/rbf/RBFBehaviorEvaluator.h"
+#include "riglogic/rbf/cpu/CPURBFBehaviorEvaluator.h"
+#include "riglogic/rbf/cpu/CPURBFBehaviorOutputInstance.h"
+#include "riglogic/rbf/cpu/RBFSolver.h"
+#include "riglogic/types/Aliases.h"
+#include "riglogic/types/bpcm/Optimizer.h"
+#include "riglogic/types/LODSpec.h"
+#include "riglogic/utils/Extd.h"
+
+#include <tdm/Ang.h>
+
+namespace rl4 {
+
+namespace rbf {
+
+namespace cpu {
+
+template<typename T, typename TF256, typename TF128>
+class Factory {
+    public:
+        static RBFBehaviorEvaluator::Pointer create(const dna::Reader* reader, MemoryResource* memRes) {
+            Vector<RBFSolver::Pointer> solvers{memRes};
+            std::uint16_t maximumInputCount{};
+            std::uint16_t maximumTargetCount{};
+            Matrix<std::uint16_t> solverRawControlInputIndices{memRes};
+            Matrix<std::uint16_t> solverRawControlOutputIndices{memRes};
+
+            auto instanceFactory = [](std::uint16_t maxInputCount, std::uint16_t maxTargetCount, MemoryResource* instanceMemRes) {
+                    using OutputInstancePointer = UniqueInstance<OutputInstance, RBFBehaviorOutputInstance>;
+                    return OutputInstancePointer::with(instanceMemRes).create(maxInputCount, maxTargetCount, instanceMemRes);
+                };
+            auto factory = UniqueInstance<Evaluator<T, TF256, TF128>, RBFBehaviorEvaluator>::with(memRes);
+
+            if (reader == nullptr) {
+                return factory.create(LODSpec<std::uint16_t>{memRes},
+                                      std::move(solvers),
+                                      std::move(solverRawControlInputIndices),
+                                      std::move(solverRawControlOutputIndices),
+                                      maximumInputCount,
+                                      maximumTargetCount,
+                                      std::move(instanceFactory));
+            }
+
+            auto lods = computeLODs(reader, memRes);
+            solvers.reserve(lods.count);
+            solverRawControlInputIndices.resize(lods.count);
+            solverRawControlOutputIndices.resize(lods.count);
+
+            const std::uint16_t rbfControlsOffset = static_cast<std::uint16_t>(
+                reader->getRawControlCount() +
+                reader->getPSDCount() +
+                reader->getMLControlCount());
+
+            RBFSolverRecipe recipe{};
+            Vector<float> targetScales{memRes};
+            Vector<float> targetValues{memRes};
+            const auto rotationUnit = reader->getRotationUnit();
+            std::function<float(float)> angConv;
+            if (rotationUnit == dna::RotationUnit::degrees) {
+                angConv = [](float angle) {
+                        return tdm::frad{tdm::fdeg{angle}}.value;
+                    };
+            } else {
+                angConv = [](float angle) {
+                        return tdm::frad{angle}.value;
+                    };
+            }
+            for (std::uint16_t solverIndex = {}; solverIndex < lods.count; ++solverIndex) {
+                recipe.solverType = reader->getRBFSolverType(solverIndex);
+                recipe.distanceMethod = reader->getRBFSolverDistanceMethod(solverIndex);
+                recipe.weightFunction = reader->getRBFSolverFunctionType(solverIndex);
+                recipe.normalizeMethod = reader->getRBFSolverNormalizeMethod(solverIndex);
+                recipe.twistAxis = reader->getRBFSolverTwistAxis(solverIndex);
+                recipe.isAutomaticRadius = reader->getRBFSolverAutomaticRadius(solverIndex) == dna::AutomaticRadius::On;
+                recipe.radius = angConv(reader->getRBFSolverRadius(solverIndex));
+                recipe.weightThreshold = reader->getRBFSolverWeightThreshold(solverIndex);
+
+                auto rawControlIndices = reader->getRBFSolverRawControlIndices(solverIndex);
+                solverRawControlInputIndices[solverIndex].assign(rawControlIndices.begin(), rawControlIndices.end());
+                if (rawControlIndices.size() > maximumInputCount) {
+                    maximumInputCount = static_cast<std::uint16_t>(rawControlIndices.size());
+                }
+                recipe.rawControlCount = static_cast<std::uint16_t>(rawControlIndices.size());
+
+                auto poseIndices = reader->getRBFSolverPoseIndices(solverIndex);
+                solverRawControlOutputIndices[solverIndex].reserve(poseIndices.size());
+                for (const std::uint16_t pi : poseIndices) {
+                    solverRawControlOutputIndices[solverIndex].push_back(static_cast<std::uint16_t>(rbfControlsOffset + pi));
+                }
+
+                const auto targetCount = static_cast<std::uint16_t>(poseIndices.size());
+                if (targetCount > maximumTargetCount) {
+                    maximumTargetCount = targetCount;
+                    targetScales.resize(maximumTargetCount);
+                }
+
+                for (std::uint16_t i = 0u; i < targetCount; ++i) {
+                    targetScales[i] = reader->getRBFPoseScale(poseIndices[i]);
+                }
+                recipe.targetValues = reader->getRBFSolverRawControlValues(solverIndex);
+                recipe.targetScales = ConstArrayView<float>{targetScales.data(), targetCount};
+
+                auto solver = RBFSolver::create(recipe, memRes);
+                solvers.emplace_back(std::move(solver));
+            }
+
+            return factory.create(std::move(lods),
+                                  std::move(solvers),
+                                  std::move(solverRawControlInputIndices),
+                                  std::move(solverRawControlOutputIndices),
+                                  maximumInputCount,
+                                  maximumTargetCount,
+                                  std::move(instanceFactory));
+        }
+
+    private:
+        static LODSpec<std::uint16_t> computeLODs(const dna::RBFBehaviorReader* reader, MemoryResource* memRes) {
+            LODSpec<std::uint16_t> lods{memRes};
+            const auto lodCount = reader->getLODCount();
+            lods.indicesPerLOD.resize(lodCount);
+            for (std::uint16_t lod = {}; lod < lodCount; ++lod) {
+                const auto solverIndices = reader->getRBFSolverIndicesForLOD(lod);
+                lods.indicesPerLOD[lod].assign(solverIndices.begin(), solverIndices.end());
+            }
+            lods.count = reader->getRBFSolverCount();
+            return lods;
+        }
+
+};
+
+}  // namespace cpu
+
+}  // namespace rbf
+
+}  // namespace rl4
diff --git a/Engine/Plugins/Animation/RigLogic/Source/RigLogicLib/Private/riglogic/rbf/cpu/CPURBFBehaviorOutputInstance.cpp b/Engine/Plugins/Animation/RigLogic/Source/RigLogicLib/Private/riglogic/rbf/cpu/CPURBFBehaviorOutputInstance.cpp
new file mode 100644
--- /dev/null
+++ b/Engine/Plugins/Animation/RigLogic/Source/RigLogicLib/Private/riglogic/rbf/cpu/CPURBFBehaviorOutputInstance.cpp
@@ -0,0 +1,33 @@
+// Copyright Epic Games, Inc. All Rights Reserved.
+
+#include "riglogic/rbf/cpu/CPURBFBehaviorOutputInstance.h"
+
+namespace rl4 {
+
+namespace rbf {
+
+namespace cpu {
+
+OutputInstance::OutputInstance(std::uint16_t maximumInputCount, std::uint16_t maximumTargetCount, MemoryResource* memRes) :
+    inputBuffer{maximumInputCount, {}, memRes},
+    intermediateWeightsBuffer{maximumTargetCount, {}, memRes},
+    outputWeightsBuffer{maximumTargetCount, {}, memRes} {
+}
+
+ArrayView<float> OutputInstance::getInputBuffer() {
+    return inputBuffer;
+}
+
+ArrayView<float> OutputInstance::getIntermediateWeightsBuffer() {
+    return intermediateWeightsBuffer;
+}
+
+ArrayView<float> OutputInstance::getOutputWeightsBuffer() {
+    return outputWeightsBuffer;
+}
+
+}  // namespace cpu
+
+}  // namespace rbf
+
+}  // namespace rl4
diff --git a/Engine/Plugins/Animation/RigLogic/Source/RigLogicLib/Private/riglogic/rbf/cpu/CPURBFBehaviorOutputInstance.h b/Engine/Plugins/Animation/RigLogic/Source/RigLogicLib/Private/riglogic/rbf/cpu/CPURBFBehaviorOutputInstance.h
new file mode 100644
--- /dev/null
+++ b/Engine/Plugins/Animation/RigLogic/Source/RigLogicLib/Private/riglogic/rbf/cpu/CPURBFBehaviorOutputInstance.h
@@ -0,0 +1,44 @@
+// Copyright Epic Games, Inc. All Rights Reserved.
+
+#pragma once
+
+#include "riglogic/TypeDefs.h"
+#include "riglogic/rbf/RBFBehaviorOutputInstance.h"
+
+#ifdef _MSC_VER
+    #pragma warning(push)
+    #pragma warning(disable : 4365 4987)
+#endif
+#include <functional>
+#ifdef _MSC_VER
+    #pragma warning(pop)
+#endif
+
+namespace rl4 {
+
+namespace rbf {
+
+namespace cpu {
+
+class OutputInstance : public RBFBehaviorOutputInstance {
+    public:
+        using Factory = std::function<Pointer(std::uint16_t, std::uint16_t, MemoryResource*)>;
+
+    public:
+        OutputInstance(std::uint16_t maximumInputCount, std::uint16_t maximumTargetCount, MemoryResource* memRes);
+        ArrayView<float> getInputBuffer() override;
+        ArrayView<float> getIntermediateWeightsBuffer() override;
+        ArrayView<float> getOutputWeightsBuffer() override;
+
+    private:
+        AlignedVector<float> inputBuffer;
+        AlignedVector<float> intermediateWeightsBuffer;
+        AlignedVector<float> outputWeightsBuffer;
+
+};
+
+}  // namespace cpu
+
+}  // namespace rbf
+
+}  // namespace rl4
diff --git a/Engine/Plugins/Animation/RigLogic/Source/RigLogicLib/Private/riglogic/riglogic/ConfigurationSerializer.h b/Engine/Plugins/Animation/RigLogic/Source/RigLogicLib/Private/riglogic/riglogic/ConfigurationSerializer.h
--- a/Engine/Plugins/Animation/RigLogic/Source/RigLogicLib/Private/riglogic/riglogic/ConfigurationSerializer.h
+++ b/Engine/Plugins/Animation/RigLogic/Source/RigLogicLib/Private/riglogic/riglogic/ConfigurationSerializer.h
@@ -12,7 +12,13 @@ void serialize(Archive& archive, Configuration& config) {
             config.loadJoints,
             config.loadBlendShapes,
             config.loadAnimatedMaps,
-            config.loadMachineLearnedBehavior);
+            config.loadMachineLearnedBehavior,
+            config.loadRBFBehavior,
+            config.loadTwistSwingBehavior,
+            config.translationType,
+            config.rotationType,
+            config.rotationOrder,
+            config.scaleType);
 }
 
 }  // namespace rl4
diff --git a/Engine/Plugins/Animation/RigLogic/Source/RigLogicLib/Private/riglogic/riglogic/RigInstanceImpl.cpp b/Engine/Plugins/Animation/RigLogic/Source/RigLogicLib/Private/riglogic/riglogic/RigInstanceImpl.cpp
--- a/Engine/Plugins/Animation/RigLogic/Source/RigLogicLib/Private/riglogic/riglogic/RigInstanceImpl.cpp
+++ b/Engine/Plugins/Animation/RigLogic/Source/RigLogicLib/Private/riglogic/riglogic/RigInstanceImpl.cpp
@@ -20,10 +20,8 @@
 #include <cassert>
 #include <cstdint>
 #include <numeric>
-#ifdef _MSC_VER
-#if (_MSC_VER >= 1900)
-#include <span>
-#endif
+#if defined(_MSC_VER) && !defined(__clang__) && (_MSC_VER >= 1900) && (__cplusplus >= 202002L)
+    #include <span>
 #endif
 #ifdef _MSC_VER
     #pragma warning(pop)
@@ -33,9 +31,6 @@ namespace rl4 {
 
 namespace {
 
-const float controlsMin = 0.0f;
-const float controlsMax = 1.0f;
-
 inline std::uint16_t getMaxLODLevel(std::uint16_t lodCount) {
     return (lodCount > 0u ? static_cast<std::uint16_t>(lodCount - 1) : static_cast<std::uint16_t>(0));
 }
@@ -66,9 +61,11 @@ RigInstanceImpl::RigInstanceImpl(const RigMetrics& metrics, RigLogicImpl* rigLog
     rawControlCount{metrics.rawControlCount},
     psdControlCount{metrics.psdControlCount},
     mlControlCount{metrics.mlControlCount},
+    rbfControlCount{metrics.rbfControlCount},
     neuralNetworkCount{metrics.neuralNetworkCount},
     controlsInstance{rigLogic->createControlsInstance(memRes)},
     machineLearnedBehaviorInstance{rigLogic->createMachineLearnedBehaviorInstance(memRes)},
+    rbfBehaviorInstance{rigLogic->createRBFBehaviorInstance(memRes)},
     jointsInstance{rigLogic->createJointsInstance(memRes)},
     blendShapesInstance{rigLogic->createBlendShapesInstance(memRes)},
     animatedMapsInstance{rigLogic->createAnimatedMapsInstance(memRes)} {
@@ -93,10 +90,16 @@ ConstArrayView<float> RigInstanceImpl::getGUIControlValues() const {
 
 void RigInstanceImpl::setGUIControlValues(const float* values) {
     auto guiControlBuffer = controlsInstance->getGUIControlBuffer();
-    #if (defined(_MSC_VER) && _MSC_VER >= 1900)
-        std::copy(values,
-                  values + guiControlCount,
-                  std::span{guiControlBuffer.data(), guiControlBuffer.size()}.begin());
+    #if defined(_MSC_VER) && !defined(__clang__)
+        #if (_MSC_VER >= 1900) && (__cplusplus >= 202002L)
+            std::copy(values,
+                      values + guiControlCount,
+                      std::span{guiControlBuffer.data(), guiControlBuffer.size()}.begin());
+        #else
+            std::copy(values,
+                      values + guiControlCount,
+                      stdext::make_checked_array_iterator(guiControlBuffer.data(), guiControlBuffer.size()));
+        #endif
     #else
         std::copy(values, values + guiControlCount, guiControlBuffer.begin());
     #endif
@@ -112,7 +115,7 @@ float RigInstanceImpl::getRawControl(std::uint16_t index) const {
 
 void RigInstanceImpl::setRawControl(std::uint16_t index, float value) {
     auto inputBuffer = controlsInstance->getInputBuffer();
-    inputBuffer[index] = extd::clamp(value, controlsMin, controlsMax);
+    inputBuffer[index] = value;
 }
 
 ConstArrayView<float> RigInstanceImpl::getRawControlValues() const {
@@ -120,18 +123,19 @@ ConstArrayView<float> RigInstanceImpl::getRawControlValues() const {
 }
 
 void RigInstanceImpl::setRawControlValues(const float* values) {
-    auto clamper = [](float v) {
-            return extd::clamp(v, controlsMin, controlsMax);
-        };
     auto inputBuffer = controlsInstance->getInputBuffer();
-
-    #if (defined(_MSC_VER) && _MSC_VER >= 1900)
-        std::transform(values,
-                       values + rawControlCount,
-                       std::span{inputBuffer.data(), inputBuffer.size()}.begin(),
-                       clamper);
+    #if defined(_MSC_VER) && !defined(__clang__)
+        #if (_MSC_VER >= 1900) && (__cplusplus >= 202002L)
+            std::copy(values,
+                      values + rawControlCount,
+                      std::span{inputBuffer.data(), inputBuffer.size()}.begin());
+        #else
+            std::copy(values,
+                      values + rawControlCount,
+                      stdext::make_checked_array_iterator(inputBuffer.data(), inputBuffer.size()));
+        #endif
     #else
-        std::transform(values, values + rawControlCount, inputBuffer.begin(), clamper);
+        std::copy(values, values + rawControlCount, inputBuffer.begin());
     #endif
 }
 
@@ -160,6 +164,20 @@ ConstArrayView<float> RigInstanceImpl::getMLControlValues() const {
     return controlsInstance->getInputBuffer().subview(mlControlsOffset, mlControlCount);
 }
 
+std::uint16_t RigInstanceImpl::getRBFControlCount() const {
+    return rbfControlCount;
+}
+
+float RigInstanceImpl::getRBFControl(std::uint16_t index) const {
+    return getRBFControlValues()[index];
+}
+
+ConstArrayView<float> RigInstanceImpl::getRBFControlValues() const {
+    const auto rbfControlsOffset = static_cast<std::size_t>(rawControlCount) + static_cast<std::size_t>(psdControlCount) +
+        static_cast<std::size_t>(mlControlCount);
+    return controlsInstance->getInputBuffer().subview(rbfControlsOffset, rbfControlCount);
+}
+
 std::uint16_t RigInstanceImpl::getNeuralNetworkCount() const {
     return neuralNetworkCount;
 }
@@ -182,25 +200,17 @@ std::uint16_t RigInstanceImpl::getLOD() const {
 
 void RigInstanceImpl::setLOD(std::uint16_t level) {
     if (level != lodLevel) {
-        auto jointOutputs = jointsInstance->getOutputBuffer();
-        auto blendShapeOutputs = blendShapesInstance->getOutputBuffer();
-        auto animatedMapOutputs = animatedMapsInstance->getOutputBuffer();
-        std::fill(jointOutputs.begin(), jointOutputs.end(), 0.0f);
-        std::fill(blendShapeOutputs.begin(), blendShapeOutputs.end(), 0.0f);
-        std::fill(animatedMapOutputs.begin(), animatedMapOutputs.end(), 0.0f);
+        jointsInstance->resetOutputBuffer();
+        blendShapesInstance->resetOutputBuffer();
+        animatedMapsInstance->resetOutputBuffer();
     }
     lodLevel = extd::clamp(level, static_cast<std::uint16_t>(0), lodMaxLevel);
 }
 
-ConstArrayView<float> RigInstanceImpl::getRawJointOutputs() const {
+ConstArrayView<float> RigInstanceImpl::getJointOutputs() const {
     return jointsInstance->getOutputBuffer();
 }
 
-TransformationArrayView RigInstanceImpl::getJointOutputs() const {
-    auto jointOutputs = jointsInstance->getOutputBuffer();
-    return TransformationArrayView{jointOutputs.data(), jointOutputs.size()};
-}
-
 ConstArrayView<float> RigInstanceImpl::getBlendShapeOutputs() const {
     return blendShapesInstance->getOutputBuffer();
 }
@@ -217,6 +227,10 @@ MachineLearnedBehaviorOutputInstance* RigInstanceImpl::getMachineLearnedBehavior
     return machineLearnedBehaviorInstance.get();
 }
 
+RBFBehaviorOutputInstance* RigInstanceImpl::getRBFBehaviorOutputInstance() {
+    return rbfBehaviorInstance.get();
+}
+
 JointsOutputInstance* RigInstanceImpl::getJointsOutputInstance() {
     return jointsInstance.get();
 }
diff --git a/Engine/Plugins/Animation/RigLogic/Source/RigLogicLib/Private/riglogic/riglogic/RigInstanceImpl.h b/Engine/Plugins/Animation/RigLogic/Source/RigLogicLib/Private/riglogic/riglogic/RigInstanceImpl.h
--- a/Engine/Plugins/Animation/RigLogic/Source/RigLogicLib/Private/riglogic/riglogic/RigInstanceImpl.h
+++ b/Engine/Plugins/Animation/RigLogic/Source/RigLogicLib/Private/riglogic/riglogic/RigInstanceImpl.h
@@ -8,6 +8,7 @@
 #include "riglogic/controls/ControlsInputInstance.h"
 #include "riglogic/joints/JointsOutputInstance.h"
 #include "riglogic/ml/MachineLearnedBehaviorOutputInstance.h"
+#include "riglogic/rbf/RBFBehaviorOutputInstance.h"
 #include "riglogic/riglogic/RigInstance.h"
 #include "riglogic/riglogic/RigMetrics.h"
 
@@ -45,16 +46,20 @@ class RigInstanceImpl : public RigInstance {
         float getNeuralNetworkMask(std::uint16_t neuralNetIndex) const override;
         void setNeuralNetworkMask(std::uint16_t neuralNetIndex, float value) override;
 
+        std::uint16_t getRBFControlCount() const override;
+        float getRBFControl(std::uint16_t index) const override;
+        ConstArrayView<float> getRBFControlValues() const override;
+
         std::uint16_t getLOD() const override;
         void setLOD(std::uint16_t level) override;
 
-        ConstArrayView<float> getRawJointOutputs() const override;
-        TransformationArrayView getJointOutputs() const override;
+        ConstArrayView<float> getJointOutputs() const override;
         ConstArrayView<float> getBlendShapeOutputs() const override;
         ConstArrayView<float> getAnimatedMapOutputs() const override;
 
         ControlsInputInstance* getControlsInputInstance();
         MachineLearnedBehaviorOutputInstance* getMachineLearnedBehaviorOutputInstance();
+        RBFBehaviorOutputInstance* getRBFBehaviorOutputInstance();
         JointsOutputInstance* getJointsOutputInstance();
         BlendShapesOutputInstance* getBlendShapesOutputInstance();
         AnimatedMapsOutputInstance* getAnimatedMapOutputInstance();
@@ -70,10 +75,12 @@ class RigInstanceImpl : public RigInstance {
         std::uint16_t rawControlCount;
         std::uint16_t psdControlCount;
         std::uint16_t mlControlCount;
+        std::uint16_t rbfControlCount;
         std::uint16_t neuralNetworkCount;
 
         ControlsInputInstance::Pointer controlsInstance;
         MachineLearnedBehaviorOutputInstance::Pointer machineLearnedBehaviorInstance;
+        RBFBehaviorOutputInstance::Pointer rbfBehaviorInstance;
         JointsOutputInstance::Pointer jointsInstance;
         BlendShapesOutputInstance::Pointer blendShapesInstance;
         AnimatedMapsOutputInstance::Pointer animatedMapsInstance;
diff --git a/Engine/Plugins/Animation/RigLogic/Source/RigLogicLib/Private/riglogic/riglogic/RigLogicImpl.cpp b/Engine/Plugins/Animation/RigLogic/Source/RigLogicLib/Private/riglogic/riglogic/RigLogicImpl.cpp
--- a/Engine/Plugins/Animation/RigLogic/Source/RigLogicLib/Private/riglogic/riglogic/RigLogicImpl.cpp
+++ b/Engine/Plugins/Animation/RigLogic/Source/RigLogicLib/Private/riglogic/riglogic/RigLogicImpl.cpp
@@ -8,6 +8,8 @@
 #include "riglogic/controls/ControlsFactory.h"
 #include "riglogic/joints/JointsFactory.h"
 #include "riglogic/ml/MachineLearnedBehaviorFactory.h"
+#include "riglogic/rbf/RBFBehaviorFactory.h"
+#include "riglogic/rbf/RBFBehaviorOutputInstance.h"
 #include "riglogic/riglogic/ConfigurationSerializer.h"
 #include "riglogic/riglogic/RigInstanceImpl.h"
 #include "riglogic/riglogic/RigMetrics.h"
@@ -32,29 +34,33 @@ static RigInstanceImpl* castInstance(RigInstance* instance) {
     return static_cast<RigInstanceImpl*>(instance);
 }
 
-static RigMetrics::Pointer computeRigMetrics(const dna::Reader* reader, MemoryResource* memRes) {
+static RigMetrics::Pointer computeRigMetrics(const dna::Reader* reader, const Configuration& config, MemoryResource* memRes) {
     RigMetrics::Pointer metrics = UniqueInstance<RigMetrics>::with(memRes).create(memRes);
     metrics->lodCount = reader->getLODCount();
     metrics->guiControlCount = reader->getGUIControlCount();
     metrics->rawControlCount = reader->getRawControlCount();
     metrics->psdControlCount = reader->getPSDCount();
-    metrics->jointAttributeCount = reader->getJointRowCount();
+    const auto numAttrsPerJoint = (static_cast<std::uint8_t>(config.translationType) +
+                                   static_cast<std::uint8_t>(config.rotationType) +
+                                   static_cast<std::uint8_t>(config.scaleType));
+    metrics->jointAttributeCount = static_cast<std::uint16_t>(reader->getJointCount() * numAttrsPerJoint);
     metrics->blendShapeCount = reader->getBlendShapeChannelCount();
     metrics->animatedMapCount = reader->getAnimatedMapCount();
-    #ifdef RL_BUILD_WITH_ML_EVALUATOR
-        metrics->mlControlCount = reader->getMLControlCount();
-        metrics->neuralNetworkCount = reader->getNeuralNetworkCount();
-    #endif  // RL_BUILD_WITH_ML_EVALUATOR
+    metrics->mlControlCount = reader->getMLControlCount();
+    metrics->neuralNetworkCount = reader->getNeuralNetworkCount();
+    metrics->rbfSolverCount = reader->getRBFSolverCount();
+    metrics->rbfControlCount = reader->getRBFPoseCount();
     return metrics;
 }
 
 RigLogic::~RigLogic() = default;
 
-RigLogic* RigLogic::create(const dna::Reader* reader, Configuration config, MemoryResource* memRes) {
-    auto metrics = computeRigMetrics(reader, memRes);
+RigLogic* RigLogic::create(const dna::Reader* reader, const Configuration& config, MemoryResource* memRes) {
+    auto metrics = computeRigMetrics(reader, config, memRes);
 
     auto controls = ControlsFactory::create(config, reader, memRes);
     auto machineLearnedBlendShapes = MachineLearnedBehaviorFactory::create(config, reader, memRes);
+    auto rbfBehavior = RBFBehaviorFactory::create(config, reader, memRes);
     auto joints = JointsFactory::create(config, reader, memRes);
     auto blendShapes = BlendShapesFactory::create(config, reader, memRes);
     auto animatedMaps = AnimatedMapsFactory::create(config, reader, memRes);
@@ -64,6 +70,7 @@ RigLogic* RigLogic::create(const dna::Reader* reader, Configuration config, Memo
                            std::move(metrics),
                            std::move(controls),
                            std::move(machineLearnedBlendShapes),
+                           std::move(rbfBehavior),
                            std::move(joints),
                            std::move(blendShapes),
                            std::move(animatedMaps),
@@ -89,7 +96,8 @@ RigLogic* RigLogic::restore(BoundedIOStream* source, MemoryResource* memRes) {
     archive >> *metrics;
 
     auto controls = ControlsFactory::create(config, *metrics, memRes);
-    auto machineLearnedBlendShapes = MachineLearnedBehaviorFactory::create(config, memRes);
+    auto machineLearnedBehavior = MachineLearnedBehaviorFactory::create(config, *metrics, memRes);
+    auto rbfBehavior = RBFBehaviorFactory::create(config, *metrics, memRes);
     auto joints = JointsFactory::create(config, *metrics, memRes);
     auto blendShapes = BlendShapesFactory::create(config, *metrics, memRes);
     auto animatedMaps = AnimatedMapsFactory::create(config, *metrics, memRes);
@@ -97,21 +105,23 @@ RigLogic* RigLogic::restore(BoundedIOStream* source, MemoryResource* memRes) {
     terse::VirtualSerializerProxy<AnimatedMaps> animatedMapsProxy{animatedMaps.get()};
     terse::VirtualSerializerProxy<BlendShapes> blendShapesProxy{blendShapes.get()};
 
-    archive >> *controls >> *machineLearnedBlendShapes >> *joints >> blendShapesProxy >> animatedMapsProxy;
+    archive >> *controls >> *machineLearnedBehavior >> *rbfBehavior >> *joints >> blendShapesProxy >> animatedMapsProxy;
     return alloc.newObject(config,
                            std::move(metrics),
                            std::move(controls),
-                           std::move(machineLearnedBlendShapes),
+                           std::move(machineLearnedBehavior),
+                           std::move(rbfBehavior),
                            std::move(joints),
                            std::move(blendShapes),
                            std::move(animatedMaps),
                            memRes);
 }
 
-RigLogicImpl::RigLogicImpl(Configuration config_,
+RigLogicImpl::RigLogicImpl(const Configuration& config_,
                            RigMetrics::Pointer metrics_,
                            Controls::Pointer controls_,
                            MachineLearnedBehavior::Pointer machineLearnedBehavior_,
+                           RBFBehavior::Pointer rbfBehavior_,
                            Joints::Pointer joints_,
                            BlendShapes::Pointer blendShapes_,
                            AnimatedMaps::Pointer animatedMaps_,
@@ -121,6 +131,7 @@ RigLogicImpl::RigLogicImpl(Configuration config_,
     metrics{std::move(metrics_)},
     controls{std::move(controls_)},
     machineLearnedBehavior{std::move(machineLearnedBehavior_)},
+    rbfBehavior{std::move(rbfBehavior_)},
     joints{std::move(joints_)},
     blendShapes{std::move(blendShapes_)},
     animatedMaps{std::move(animatedMaps_)} {
@@ -130,7 +141,16 @@ void RigLogicImpl::dump(BoundedIOStream* destination) const {
     terse::BinaryOutputArchive<BoundedIOStream> archive{destination};
     terse::VirtualSerializerProxy<AnimatedMaps> animatedMapsProxy{animatedMaps.get()};
     terse::VirtualSerializerProxy<BlendShapes> blendShapesProxy{blendShapes.get()};
-    archive << config << *metrics << *controls << *machineLearnedBehavior << *joints << blendShapesProxy << animatedMapsProxy;
+    // *INDENT-OFF*
+    archive << config
+            << *metrics
+            << *controls
+            << *machineLearnedBehavior
+            << *rbfBehavior
+            << *joints
+            << blendShapesProxy
+            << animatedMapsProxy;
+    // *INDENT-ON*
 }
 
 const Configuration& RigLogicImpl::getConfiguration() const {
@@ -145,11 +165,7 @@ std::uint16_t RigLogicImpl::getLODCount() const {
     return metrics->lodCount;
 }
 
-ConstArrayView<float> RigLogicImpl::getRawNeutralJointValues() const {
-    return joints->getRawNeutralValues();
-}
-
-TransformationArrayView RigLogicImpl::getNeutralJointValues() const {
+ConstArrayView<float> RigLogicImpl::getNeutralJointValues() const {
     return joints->getNeutralValues();
 }
 
@@ -165,6 +181,10 @@ std::uint16_t RigLogicImpl::getNeuralNetworkCount() const {
     return metrics->neuralNetworkCount;
 }
 
+std::uint16_t RigLogicImpl::getRBFSolverCount() const {
+    return metrics->rbfSolverCount;
+}
+
 std::uint16_t RigLogicImpl::getMeshCount() const {
     return machineLearnedBehavior->getMeshCount();
 }
@@ -186,6 +206,11 @@ const {
     return machineLearnedBehavior->createInstance(instanceMemRes);
 }
 
+RBFBehaviorOutputInstance::Pointer RigLogicImpl::createRBFBehaviorInstance(MemoryResource* instanceMemRes)
+const {
+    return rbfBehavior->createInstance(instanceMemRes);
+}
+
 JointsOutputInstance::Pointer RigLogicImpl::createJointsInstance(MemoryResource* instanceMemRes) const {
     return joints->createInstance(instanceMemRes);
 }
@@ -228,6 +253,21 @@ void RigLogicImpl::calculateMachineLearnedBehaviorControls(RigInstance* instance
                                       neuralNetIndex);
 }
 
+void RigLogicImpl::calculateRBFControls(RigInstance* instance) const {
+    auto pRigInstance = castInstance(instance);
+    rbfBehavior->calculate(pRigInstance->getControlsInputInstance(),
+                           pRigInstance->getRBFBehaviorOutputInstance(),
+                           pRigInstance->getLOD());
+}
+
+void RigLogicImpl::calculateRBFControls(RigInstance* instance, std::uint16_t solverIndex) const {
+    auto pRigInstance = castInstance(instance);
+    rbfBehavior->calculate(pRigInstance->getControlsInputInstance(),
+                           pRigInstance->getRBFBehaviorOutputInstance(),
+                           pRigInstance->getLOD(),
+                           solverIndex);
+}
+
 void RigLogicImpl::calculateJoints(RigInstance* instance) const {
     auto pRigInstance = castInstance(instance);
     joints->calculate(pRigInstance->getControlsInputInstance(), pRigInstance->getJointsOutputInstance(), pRigInstance->getLOD());
@@ -257,6 +297,7 @@ void RigLogicImpl::calculateAnimatedMaps(RigInstance* instance) const {
 
 void RigLogicImpl::calculate(RigInstance* instance) const {
     calculateMachineLearnedBehaviorControls(instance);
+    calculateRBFControls(instance);
     calculateControls(instance);
     calculateJoints(instance);
     calculateBlendShapes(instance);
diff --git a/Engine/Plugins/Animation/RigLogic/Source/RigLogicLib/Private/riglogic/riglogic/RigLogicImpl.h b/Engine/Plugins/Animation/RigLogic/Source/RigLogicLib/Private/riglogic/riglogic/RigLogicImpl.h
--- a/Engine/Plugins/Animation/RigLogic/Source/RigLogicLib/Private/riglogic/riglogic/RigLogicImpl.h
+++ b/Engine/Plugins/Animation/RigLogic/Source/RigLogicLib/Private/riglogic/riglogic/RigLogicImpl.h
@@ -13,6 +13,7 @@
 #include "riglogic/joints/JointsOutputInstance.h"
 #include "riglogic/ml/MachineLearnedBehavior.h"
 #include "riglogic/ml/MachineLearnedBehaviorOutputInstance.h"
+#include "riglogic/rbf/RBFBehavior.h"
 #include "riglogic/riglogic/Configuration.h"
 #include "riglogic/riglogic/RigLogic.h"
 #include "riglogic/riglogic/RigMetrics.h"
@@ -21,10 +22,11 @@ namespace rl4 {
 
 class RigLogicImpl : public RigLogic {
     public:
-        RigLogicImpl(Configuration config_,
+        RigLogicImpl(const Configuration& config_,
                      RigMetrics::Pointer metrics_,
                      Controls::Pointer controls_,
                      MachineLearnedBehavior::Pointer machineLearnedBehavior_,
+                     RBFBehavior::Pointer rbfBehavior_,
                      Joints::Pointer joints_,
                      BlendShapes::Pointer blendShapes_,
                      AnimatedMaps::Pointer animatedMaps_,
@@ -34,17 +36,18 @@ class RigLogicImpl : public RigLogic {
         const Configuration& getConfiguration() const;
         const RigMetrics& getRigMetrics() const;
         std::uint16_t getLODCount() const override;
-        ConstArrayView<float> getRawNeutralJointValues() const override;
-        TransformationArrayView getNeutralJointValues() const override;
+        ConstArrayView<float> getNeutralJointValues() const override;
         ConstArrayView<std::uint16_t> getJointVariableAttributeIndices(std::uint16_t lod) const override;
         std::uint16_t getJointGroupCount() const override;
         std::uint16_t getNeuralNetworkCount() const override;
+        std::uint16_t getRBFSolverCount() const override;
         std::uint16_t getMeshCount() const override;
         std::uint16_t getMeshRegionCount(std::uint16_t meshIndex) const override;
         ConstArrayView<std::uint16_t> getNeuralNetworkIndices(std::uint16_t meshIndex, std::uint16_t regionIndex) const override;
 
         ControlsInputInstance::Pointer createControlsInstance(MemoryResource* instanceMemRes) const;
         MachineLearnedBehaviorOutputInstance::Pointer createMachineLearnedBehaviorInstance(MemoryResource* instanceMemRes) const;
+        RBFBehaviorOutputInstance::Pointer createRBFBehaviorInstance(MemoryResource* instanceMemRes) const;
         JointsOutputInstance::Pointer createJointsInstance(MemoryResource* instanceMemRes) const;
         BlendShapesOutputInstance::Pointer createBlendShapesInstance(MemoryResource* instanceMemRes) const;
         AnimatedMapsOutputInstance::Pointer createAnimatedMapsInstance(MemoryResource* instanceMemRes) const;
@@ -54,6 +57,8 @@ class RigLogicImpl : public RigLogic {
         void calculateControls(RigInstance* instance) const override;
         void calculateMachineLearnedBehaviorControls(RigInstance* instance) const override;
         void calculateMachineLearnedBehaviorControls(RigInstance* instance, std::uint16_t neuralNetIndex) const override;
+        void calculateRBFControls(RigInstance* instance) const override;
+        void calculateRBFControls(RigInstance* instance, std::uint16_t solverIndex) const override;
         void calculateJoints(RigInstance* instance) const override;
         void calculateJoints(RigInstance* instance, std::uint16_t jointGroupIndex) const override;
         void calculateBlendShapes(RigInstance* instance) const override;
@@ -68,6 +73,7 @@ class RigLogicImpl : public RigLogic {
         RigMetrics::Pointer metrics;
         Controls::Pointer controls;
         MachineLearnedBehavior::Pointer machineLearnedBehavior;
+        RBFBehavior::Pointer rbfBehavior;
         Joints::Pointer joints;
         BlendShapes::Pointer blendShapes;
         AnimatedMaps::Pointer animatedMaps;
diff --git a/Engine/Plugins/Animation/RigLogic/Source/RigLogicLib/Private/riglogic/riglogic/RigMetrics.h b/Engine/Plugins/Animation/RigLogic/Source/RigLogicLib/Private/riglogic/riglogic/RigMetrics.h
--- a/Engine/Plugins/Animation/RigLogic/Source/RigLogicLib/Private/riglogic/riglogic/RigMetrics.h
+++ b/Engine/Plugins/Animation/RigLogic/Source/RigLogicLib/Private/riglogic/riglogic/RigMetrics.h
@@ -16,10 +16,12 @@ struct RigMetrics {
     std::uint16_t rawControlCount;
     std::uint16_t psdControlCount;
     std::uint16_t mlControlCount;
+    std::uint16_t rbfControlCount;
     std::uint16_t jointAttributeCount;
     std::uint16_t blendShapeCount;
     std::uint16_t animatedMapCount;
     std::uint16_t neuralNetworkCount;
+    std::uint16_t rbfSolverCount;
 
     explicit RigMetrics(MemoryResource*  /*unused*/) :
         lodCount{},
@@ -27,10 +29,12 @@ struct RigMetrics {
         rawControlCount{},
         psdControlCount{},
         mlControlCount{},
+        rbfControlCount{},
         jointAttributeCount{},
         blendShapeCount{},
         animatedMapCount{},
-        neuralNetworkCount{} {
+        neuralNetworkCount{},
+        rbfSolverCount{} {
     }
 
     template<class Archive>
@@ -40,10 +44,12 @@ struct RigMetrics {
                 rawControlCount,
                 psdControlCount,
                 mlControlCount,
+                rbfControlCount,
                 jointAttributeCount,
                 blendShapeCount,
                 animatedMapCount,
-                neuralNetworkCount);
+                neuralNetworkCount,
+                rbfSolverCount);
     }
 
 };
diff --git a/Engine/Plugins/Animation/RigLogic/Source/RigLogicLib/Private/riglogic/types/bpcm/Optimizer.h b/Engine/Plugins/Animation/RigLogic/Source/RigLogicLib/Private/riglogic/types/bpcm/Optimizer.h
--- a/Engine/Plugins/Animation/RigLogic/Source/RigLogicLib/Private/riglogic/types/bpcm/Optimizer.h
+++ b/Engine/Plugins/Animation/RigLogic/Source/RigLogicLib/Private/riglogic/types/bpcm/Optimizer.h
@@ -12,18 +12,22 @@ namespace rl4 {
 
 namespace bpcm {
 
-template<typename TFVec, std::uint32_t BlockHeight, std::uint32_t PadTo>
+template<typename TFVec, std::uint32_t BlockHeight, std::uint32_t PadTo, std::uint32_t Stride>
 struct Optimizer {
+    static_assert(BlockHeight % TFVec::size() == 0, "BlockHeight must be a multiple of TFVec::size()");
+    static_assert(BlockHeight % Stride == 0, "BlockHeight must be a multiple of Stride");
 
     static std::uint32_t optimize(float* dest, const float* source, Extent dimensions) {
         const std::uint32_t remainder = dimensions.rows % BlockHeight;
         const std::uint32_t target = dimensions.rows - remainder;
-        std::uint32_t offset = 0u;
+        std::uint32_t offset = {};
         // Copy the portion of the source matrix that is partitionable into BlockHeight chunks
-        for (std::uint32_t row = 0u; row < target; row += BlockHeight) {
-            for (std::uint32_t col = 0u; col < dimensions.cols; ++col) {
-                for (std::uint32_t blkIdx = 0u; blkIdx < BlockHeight; ++blkIdx, ++offset) {
-                    dest[offset] = source[(row + blkIdx) * dimensions.cols + col];
+        for (std::uint32_t row = {}; row < target; row += BlockHeight) {
+            for (std::uint32_t col = {}; col < dimensions.cols; ++col) {
+                for (std::uint32_t base = {}; base < Stride; ++base) {
+                    for (std::uint32_t blkIdx = base; blkIdx < BlockHeight; blkIdx += Stride, ++offset) {
+                        dest[offset] = source[(row + blkIdx) * dimensions.cols + col];
+                    }
                 }
             }
         }
@@ -31,11 +35,13 @@ struct Optimizer {
         const std::uint32_t paddedBlockHeight = extd::roundUp(remainder, PadTo);
         for (std::uint32_t row = target; row < dimensions.rows; row += remainder) {
             for (std::uint32_t col = 0u; col < dimensions.cols; ++col) {
-                for (std::uint32_t blkIdx = 0u; blkIdx < remainder; ++blkIdx, ++offset) {
-                    dest[offset] = source[(row + blkIdx) * dimensions.cols + col];
+                for (std::uint32_t base = {}; base < Stride; ++base) {
+                    for (std::uint32_t blkIdx = base; blkIdx < remainder; blkIdx += Stride, ++offset) {
+                        dest[offset] = source[(row + blkIdx) * dimensions.cols + col];
+                    }
+                    // Skip over padding (already filled with zeros) in destination storage
+                    offset += (paddedBlockHeight - remainder) / Stride;
                 }
-                // Skip over padding (already filled with zeros) in destination storage
-                offset += (paddedBlockHeight - remainder);
             }
         }
 
@@ -46,18 +52,20 @@ struct Optimizer {
         static constexpr std::size_t alignment = TFVec::alignment();
         const std::uint32_t remainder = dimensions.rows % BlockHeight;
         const std::uint32_t target = dimensions.rows - remainder;
-        std::uint32_t offset = 0u;
+        std::uint32_t offset = {};
         // Copy the portion of the source matrix that is partitionable into BlockHeight chunks
-        for (std::uint32_t row = 0u; row < target; row += BlockHeight) {
-            for (std::uint32_t col = 0u; col < dimensions.cols; ++col) {
+        for (std::uint32_t row = {}; row < target; row += BlockHeight) {
+            for (std::uint32_t col = {}; col < dimensions.cols; ++col) {
                 alignas(alignment) float input[BlockHeight];
                 alignas(alignment) std::uint16_t output[BlockHeight];
 
-                for (std::uint32_t blkIdx = 0u; blkIdx < BlockHeight; ++blkIdx) {
-                    input[blkIdx] = source[(row + blkIdx) * dimensions.cols + col];
+                for (std::uint32_t base = {}, inIdx = {}; base < Stride; ++base) {
+                    for (std::uint32_t blkIdx = base; blkIdx < BlockHeight; blkIdx += Stride) {
+                        input[inIdx++] = source[(row + blkIdx) * dimensions.cols + col];
+                    }
                 }
 
-                for (std::uint32_t blkOffset = 0u;
+                for (std::uint32_t blkOffset = {};
                      blkOffset < BlockHeight;
                      blkOffset += static_cast<std::uint32_t>(TFVec::size())) {
                     #if !defined(__clang__) && defined(__GNUC__) && (__GNUC__ < 7)
@@ -67,7 +75,7 @@ struct Optimizer {
                     #endif
                 }
 
-                for (std::uint32_t blkIdx = 0u; blkIdx < BlockHeight; ++blkIdx, ++offset) {
+                for (std::uint32_t blkIdx = {}; blkIdx < BlockHeight; ++blkIdx, ++offset) {
                     dest[offset] = output[blkIdx];
                 }
             }
@@ -75,15 +83,18 @@ struct Optimizer {
         // Copy the non-BlockHeight sized remainder portion of the source matrix
         const std::uint32_t paddedBlockHeight = extd::roundUp(remainder, PadTo);
         for (std::uint32_t row = target; row < dimensions.rows; row += remainder) {
-            for (std::uint32_t col = 0u; col < dimensions.cols; ++col) {
+            for (std::uint32_t col = {}; col < dimensions.cols; ++col) {
                 alignas(alignment) float input[BlockHeight] = {};
                 alignas(alignment) std::uint16_t output[BlockHeight] = {};
 
-                for (std::uint32_t blkIdx = 0u; blkIdx < remainder; ++blkIdx) {
-                    input[blkIdx] = source[(row + blkIdx) * dimensions.cols + col];
+                for (std::uint32_t base = {}, inIdx = {}; base < Stride; ++base) {
+                    for (std::uint32_t blkIdx = base; blkIdx < remainder; blkIdx += Stride) {
+                        input[inIdx++] = source[(row + blkIdx) * dimensions.cols + col];
+                    }
+                    inIdx += (paddedBlockHeight - remainder) / Stride;
                 }
 
-                for (std::uint32_t blkOffset = 0u;
+                for (std::uint32_t blkOffset = {};
                      blkOffset < BlockHeight;
                      blkOffset += static_cast<std::uint32_t>(TFVec::size())) {
                     #if !defined(__clang__) && defined(__GNUC__) && (__GNUC__ < 7)
@@ -99,7 +110,7 @@ struct Optimizer {
                 ASSUME_TRUE(paddedBlockHeight <= BlockHeight);
                 // Copy even the padded zero values as the representation of half-floats might
                 // differ from the storage's zero-initialized state
-                for (std::uint32_t blkIdx = 0u; blkIdx < paddedBlockHeight; ++blkIdx, ++offset) {
+                for (std::uint32_t blkIdx = {}; blkIdx < paddedBlockHeight; ++blkIdx, ++offset) {
                     dest[offset] = output[blkIdx];
                 }
             }
diff --git a/Engine/Plugins/Animation/RigLogic/Source/RigLogicLib/Private/riglogic/utils/Extd.h b/Engine/Plugins/Animation/RigLogic/Source/RigLogicLib/Private/riglogic/utils/Extd.h
--- a/Engine/Plugins/Animation/RigLogic/Source/RigLogicLib/Private/riglogic/utils/Extd.h
+++ b/Engine/Plugins/Animation/RigLogic/Source/RigLogicLib/Private/riglogic/utils/Extd.h
@@ -28,7 +28,7 @@ inline T clamp(T value, T low, T high) {
 
 template<typename T>
 inline T roundUp(T number, T multiple) {
-    return ((number + multiple - 1) / multiple) * multiple;
+    return static_cast<T>(((number + multiple - 1) / multiple) * multiple);
 }
 
 template<typename T>
@@ -108,6 +108,22 @@ inline impl::LUTFilter<TLookUpTable, impl::LUTStrategy::ByPosition> byPosition(c
     return impl::LUTFilter<TLookUpTable, impl::LUTStrategy::ByPosition>{lookUpTable};
 }
 
+template<class TContainer, class TComparator>
+inline typename TContainer::value_type minOf(const TContainer& container, TComparator&& comparator) {
+    using ValueType = typename TContainer::value_type;
+    const auto it = std::min_element(container.begin(), container.end(), comparator);
+    return (it == container.end() ? ValueType{} : *it);
+}
+
+template<class TContainer>
+inline typename TContainer::value_type minOf(const TContainer& container) {
+    using ValueType = typename TContainer::value_type;
+    const auto compare = [](const ValueType& lhs, const ValueType& rhs) {
+            return lhs < rhs;
+        };
+    return minOf(container, compare);
+}
+
 template<class TContainer, class TComparator>
 inline typename TContainer::value_type maxOf(const TContainer& container, TComparator&& comparator) {
     using ValueType = typename TContainer::value_type;
diff --git a/Engine/Plugins/Animation/RigLogic/Source/RigLogicLib/Private/trio/streams/FileStreamImpl.cpp b/Engine/Plugins/Animation/RigLogic/Source/RigLogicLib/Private/trio/streams/FileStreamImpl.cpp
--- a/Engine/Plugins/Animation/RigLogic/Source/RigLogicLib/Private/trio/streams/FileStreamImpl.cpp
+++ b/Engine/Plugins/Animation/RigLogic/Source/RigLogicLib/Private/trio/streams/FileStreamImpl.cpp
@@ -177,9 +177,9 @@ std::size_t FileStreamImpl::write(const char* source, std::size_t size) {
         return 0ul;
     }
 
-    #if defined(_MSC_VER) && (_MSC_VER < 1930)
-    // MSVC stdlib or compiler bug workaround
-    file.seekp(file.tellp());
+    #if defined(TRIO_MSVC_READ_AFTER_WRITE_WORKAROUND) && defined(_MSC_VER) && (_MSC_VER < 1930)
+        // MSVC stdlib or compiler bug workaround
+        file.seekp(file.tellp());
     #endif  // _MSC_VER
 
     file.write(source, static_cast<std::streamsize>(size));
@@ -199,9 +199,9 @@ std::size_t FileStreamImpl::write(Readable* source, std::size_t size) {
         return 0ul;
     }
 
-    #if defined(_MSC_VER) && (_MSC_VER < 1930)
-    // MSVC stdlib or compiler bug workaround
-    file.seekp(file.tellp());
+    #if defined(TRIO_MSVC_READ_AFTER_WRITE_WORKAROUND) && defined(_MSC_VER) && (_MSC_VER < 1930)
+        // MSVC stdlib or compiler bug workaround
+        file.seekp(file.tellp());
     #endif  // _MSC_VER
 
     char buffer[bufferSize];
diff --git a/Engine/Plugins/Animation/RigLogic/Source/RigLogicLib/Public/arrayview/ArrayView.h b/Engine/Plugins/Animation/RigLogic/Source/RigLogicLib/Public/arrayview/ArrayView.h
--- a/Engine/Plugins/Animation/RigLogic/Source/RigLogicLib/Public/arrayview/ArrayView.h
+++ b/Engine/Plugins/Animation/RigLogic/Source/RigLogicLib/Public/arrayview/ArrayView.h
@@ -79,9 +79,6 @@ class ArrayView {
             sz{size} {
         }
 
-        ArrayView(std::nullptr_t  /*unused*/, size_type  /*unused*/) : ArrayView{nullptr, {}} {
-        }
-
         template<typename U, typename std::enable_if<IsCompatible<T, U>::value, int>::type = 0>
         ArrayView(ArrayView<U>& src) : ArrayView{src.data(), src.size()} {
         }
diff --git a/Engine/Plugins/Animation/RigLogic/Source/RigLogicLib/Public/arrayview/version/Version.h b/Engine/Plugins/Animation/RigLogic/Source/RigLogicLib/Public/arrayview/version/Version.h
--- a/Engine/Plugins/Animation/RigLogic/Source/RigLogicLib/Public/arrayview/version/Version.h
+++ b/Engine/Plugins/Animation/RigLogic/Source/RigLogicLib/Public/arrayview/version/Version.h
@@ -4,5 +4,5 @@
 
 #define AV_MAJOR_VERSION 1
 #define AV_MINOR_VERSION 1
-#define AV_PATCH_VERSION 3
-#define AV_VERSION_STRING "1.1.3"
+#define AV_PATCH_VERSION 5
+#define AV_VERSION_STRING "1.1.5"
diff --git a/Engine/Plugins/Animation/RigLogic/Source/RigLogicLib/Public/dna/DataLayer.h b/Engine/Plugins/Animation/RigLogic/Source/RigLogicLib/Public/dna/DataLayer.h
--- a/Engine/Plugins/Animation/RigLogic/Source/RigLogicLib/Public/dna/DataLayer.h
+++ b/Engine/Plugins/Animation/RigLogic/Source/RigLogicLib/Public/dna/DataLayer.h
@@ -13,7 +13,10 @@ enum class DataLayer : std::uint32_t {
     Geometry = 8 | Definition,  // Implicitly loads Descriptor and Definition
     GeometryWithoutBlendShapes = 16 | Definition,  // Implicitly loads Descriptor and Definition
     MachineLearnedBehavior = 32 | Definition,  // Implicitly loads Definition
-    All = Behavior | Geometry | MachineLearnedBehavior
+    RBFBehavior = 64 | Behavior,  // Implicitly loads Behavior
+    JointBehaviorMetadata = 128 | Definition,  // Implicitly loads Definition
+    TwistSwingBehavior = 256 | Definition,  // Implicitly loads Definition
+    All = RBFBehavior | Geometry | MachineLearnedBehavior | JointBehaviorMetadata | TwistSwingBehavior
 };
 
 inline DataLayer operator|(DataLayer lhs, DataLayer rhs) {
diff --git a/Engine/Plugins/Animation/RigLogic/Source/RigLogicLib/Public/dna/Reader.h b/Engine/Plugins/Animation/RigLogic/Source/RigLogicLib/Public/dna/Reader.h
--- a/Engine/Plugins/Animation/RigLogic/Source/RigLogicLib/Public/dna/Reader.h
+++ b/Engine/Plugins/Animation/RigLogic/Source/RigLogicLib/Public/dna/Reader.h
@@ -4,9 +4,11 @@
 
 #include "dna/Defs.h"
 #include "dna/DataLayer.h"
-#include "dna/layers/BehaviorReader.h"
 #include "dna/layers/GeometryReader.h"
+#include "dna/layers/JointBehaviorMetadataReader.h"
 #include "dna/layers/MachineLearnedBehaviorReader.h"
+#include "dna/layers/RBFBehaviorReader.h"
+#include "dna/layers/TwistSwingBehaviorReader.h"
 
 namespace dna {
 
@@ -19,7 +21,8 @@ namespace dna {
         selectively loaded, it might be convenient to slice-off interfaces which layers were
         not loaded.
 */
-class DNAAPI Reader : public BehaviorReader, public GeometryReader, public MachineLearnedBehaviorReader {
+class DNAAPI Reader : public RBFBehaviorReader, public GeometryReader, public MachineLearnedBehaviorReader,
+    public JointBehaviorMetadataReader, public TwistSwingBehaviorReader  {
     public:
         ~Reader() override;
         /**
diff --git a/Engine/Plugins/Animation/RigLogic/Source/RigLogicLib/Public/dna/Writer.h b/Engine/Plugins/Animation/RigLogic/Source/RigLogicLib/Public/dna/Writer.h
--- a/Engine/Plugins/Animation/RigLogic/Source/RigLogicLib/Public/dna/Writer.h
+++ b/Engine/Plugins/Animation/RigLogic/Source/RigLogicLib/Public/dna/Writer.h
@@ -4,10 +4,11 @@
 
 #include "dna/DataLayer.h"
 #include "dna/Defs.h"
-#include "dna/layers/BehaviorWriter.h"
 #include "dna/layers/GeometryWriter.h"
+#include "dna/layers/JointBehaviorMetadataWriter.h"
+#include "dna/layers/RBFBehaviorWriter.h"
 #include "dna/layers/MachineLearnedBehaviorWriter.h"
-#include "dna/types/Aliases.h"
+#include "dna/layers/TwistSwingBehaviorWriter.h"
 
 namespace dna {
 
@@ -21,7 +22,8 @@ class Reader;
         structure of the Reader hierarchy, as it's not possible to selectively write only
         specific layers.
 */
-class DNAAPI Writer : public BehaviorWriter, public GeometryWriter, public MachineLearnedBehaviorWriter {
+class DNAAPI Writer : public RBFBehaviorWriter, public GeometryWriter, public MachineLearnedBehaviorWriter,
+    public JointBehaviorMetadataWriter, public TwistSwingBehaviorWriter {
     public:
         ~Writer() override;
         /**
diff --git a/Engine/Plugins/Animation/RigLogic/Source/RigLogicLib/Public/dna/version/Version.h b/Engine/Plugins/Animation/RigLogic/Source/RigLogicLib/Public/dna/version/Version.h
--- a/Engine/Plugins/Animation/RigLogic/Source/RigLogicLib/Public/dna/version/Version.h
+++ b/Engine/Plugins/Animation/RigLogic/Source/RigLogicLib/Public/dna/version/Version.h
@@ -3,6 +3,6 @@
 #pragma once
 
 #define DNA_MAJOR_VERSION 9
-#define DNA_MINOR_VERSION 1
-#define DNA_PATCH_VERSION 8
-#define DNA_VERSION_STRING "9.1.8"
+#define DNA_MINOR_VERSION 3
+#define DNA_PATCH_VERSION 5
+#define DNA_VERSION_STRING "9.3.5"
diff --git a/Engine/Plugins/Animation/RigLogic/Source/RigLogicLib/Public/pma/version/Version.h b/Engine/Plugins/Animation/RigLogic/Source/RigLogicLib/Public/pma/version/Version.h
--- a/Engine/Plugins/Animation/RigLogic/Source/RigLogicLib/Public/pma/version/Version.h
+++ b/Engine/Plugins/Animation/RigLogic/Source/RigLogicLib/Public/pma/version/Version.h
@@ -4,5 +4,5 @@
 
 #define PMA_MAJOR_VERSION 1
 #define PMA_MINOR_VERSION 3
-#define PMA_PATCH_VERSION 9
-#define PMA_VERSION_STRING "1.3.9"
+#define PMA_PATCH_VERSION 11
+#define PMA_VERSION_STRING "1.3.11"
