Update to new RigLogic implementing body rig support

diff --git a/Engine/Plugins/Animation/RigLogic/Source/RigLogicLib/Public/riglogic/riglogic/Configuration.h b/Engine/Plugins/Animation/RigLogic/Source/RigLogicLib/Public/riglogic/riglogic/Configuration.h
--- a/Engine/Plugins/Animation/RigLogic/Source/RigLogicLib/Public/riglogic/riglogic/Configuration.h
+++ b/Engine/Plugins/Animation/RigLogic/Source/RigLogicLib/Public/riglogic/riglogic/Configuration.h
@@ -2,18 +2,55 @@
 
 #pragma once
 
+#include <cstdint>
+
 namespace rl4 {
 
 /**
     @brief Implementation type for RigLogic calculations.
 */
-enum class CalculationType {
+enum class CalculationType : std::uint8_t {
     Scalar,  ///< scalar CPU algorithm
     SSE,  ///< vectorized (SSE) CPU algorithm
     AVX,  ///< vectorized (AVX) CPU algorithm (RigLogic must be built with AVX support,
           ///< otherwise it falls back to using the Scalar version)
-    NEON  ///< vectorized (NEON) CPU algorithm (RigLogic must be built with NEON support,
-          ///< otherwise it falls back to using the Scalar version)
+    NEON,  ///< vectorized (NEON) CPU algorithm (RigLogic must be built with NEON support,
+           ///< otherwise it falls back to using the Scalar version)
+    AnyVector  ///< Pick any available vectorization
+};
+
+/**
+    @brief Translation type to be used by RigLogic.
+*/
+enum class TranslationType : std::uint8_t {
+    Vector = 3,
+};
+
+/**
+    @brief Rotation type to be used by RigLogic.
+*/
+enum class RotationType : std::uint8_t {
+    EulerAngles = 3,
+    Quaternions = 4
+};
+
+/**
+    @brief Rotation order to use when converting Euler angles to/from quaternions.
+*/
+enum class RotationOrder : std::uint8_t {
+    XYZ,
+    XZY,
+    YXZ,
+    YZX,
+    ZXY,
+    ZYX
+};
+
+/**
+    @brief Scale type to be used by RigLogic.
+*/
+enum class ScaleType : std::uint8_t {
+    Vector = 3,
 };
 
 struct Configuration {
@@ -22,6 +59,12 @@ struct Configuration {
     bool loadBlendShapes = true;
     bool loadAnimatedMaps = true;
     bool loadMachineLearnedBehavior = true;
+    bool loadRBFBehavior = true;
+    bool loadTwistSwingBehavior = true;
+    TranslationType translationType = TranslationType::Vector;
+    RotationType rotationType = RotationType::EulerAngles;
+    RotationOrder rotationOrder = RotationOrder::XYZ;
+    ScaleType scaleType = ScaleType::Vector;
 };
 
 }  // namespace rl4
diff --git a/Engine/Plugins/Animation/RigLogic/Source/RigLogicLib/Public/riglogic/riglogic/RigInstance.h b/Engine/Plugins/Animation/RigLogic/Source/RigLogicLib/Public/riglogic/riglogic/RigInstance.h
--- a/Engine/Plugins/Animation/RigLogic/Source/RigLogicLib/Public/riglogic/riglogic/RigInstance.h
+++ b/Engine/Plugins/Animation/RigLogic/Source/RigLogicLib/Public/riglogic/riglogic/RigInstance.h
@@ -3,7 +3,6 @@
 #pragma once
 
 #include "riglogic/Defs.h"
-#include "riglogic/transformation/Transformation.h"
 #include "riglogic/types/Aliases.h"
 
 #include <cstdint>
@@ -74,22 +73,17 @@ class RLAPI RigInstance {
         virtual float getNeuralNetworkMask(std::uint16_t neuralNetIndex) const = 0;
         virtual void setNeuralNetworkMask(std::uint16_t neuralNetIndex, float value) = 0;
 
+        virtual std::uint16_t getRBFControlCount() const = 0;
+        virtual float getRBFControl(std::uint16_t index) const = 0;
+        virtual ConstArrayView<float> getRBFControlValues() const = 0;
+
         /**
             @brief Calculated values for joint transformations.
             @note
                 This is just a primitive array of floats, providing access to all the values of all transformations.
             @return View over the array of values.
         */
-        virtual ConstArrayView<float> getRawJointOutputs() const = 0;
-        /**
-            @brief Calculated values for joint transformations.
-            @note
-                A more user-friendly representation that groups values belonging to separate transformations into
-                single units, while providing accessors to the actual values they represent.
-            @return View over the array of transformations.
-            @see Transformation
-        */
-        virtual TransformationArrayView getJointOutputs() const = 0;
+        virtual ConstArrayView<float> getJointOutputs() const = 0;
         /**
             @brief Calculated values for blend shape deformations.
             @return View over the array of floats.
diff --git a/Engine/Plugins/Animation/RigLogic/Source/RigLogicLib/Public/riglogic/riglogic/RigLogic.h b/Engine/Plugins/Animation/RigLogic/Source/RigLogicLib/Public/riglogic/riglogic/RigLogic.h
--- a/Engine/Plugins/Animation/RigLogic/Source/RigLogicLib/Public/riglogic/riglogic/RigLogic.h
+++ b/Engine/Plugins/Animation/RigLogic/Source/RigLogicLib/Public/riglogic/riglogic/RigLogic.h
@@ -4,7 +4,6 @@
 
 #include "riglogic/Defs.h"
 #include "riglogic/riglogic/Configuration.h"
-#include "riglogic/transformation/Transformation.h"
 #include "riglogic/types/Aliases.h"
 
 #include <cstdint>
@@ -43,7 +42,7 @@ class RLAPI RigLogic {
                 User is responsible for releasing the returned pointer by calling destroy.
             @see destroy
         */
-        static RigLogic* create(const dna::Reader* reader, Configuration config = {}, MemoryResource* memRes = nullptr);
+        static RigLogic* create(const dna::Reader* reader, const Configuration& config = {}, MemoryResource* memRes = nullptr);
         /**
             @brief Method for freeing RigLogic.
             @param instance
@@ -88,16 +87,7 @@ class RLAPI RigLogic {
                 These is just a primitive array of floats, providing access to all the values of all transformations.
             @return View over the array of values.
         */
-        virtual ConstArrayView<float> getRawNeutralJointValues() const = 0;
-        /**
-            @brief Neutral values for joint transformations.
-            @note
-                A more user-friendly representation that groups values belonging to separate transformations into
-                single units, while providing accessors to the actual values they represent.
-            @return View over the array of transformations.
-            @see Transformation
-        */
-        virtual TransformationArrayView getNeutralJointValues() const = 0;
+        virtual ConstArrayView<float> getNeutralJointValues() const = 0;
         /**
             @brief All joint output indices concatenated into a single chunk per each LOD.
         */
@@ -112,6 +102,11 @@ class RLAPI RigLogic {
             @see calculateMachineLearnedBehavior
         */
         virtual std::uint16_t getNeuralNetworkCount() const = 0;
+        /**
+            @brief Number of RBF solvers for driving RBF behavior.
+            @see calculateRBFBehavior
+        */
+        virtual std::uint16_t getRBFSolverCount() const = 0;
         /**
             @brief Number of meshes.
         */
@@ -173,6 +168,28 @@ class RLAPI RigLogic {
             @see calculate
         */
         virtual void calculateMachineLearnedBehaviorControls(RigInstance* instance, std::uint16_t neuralNetIndex) const = 0;
+        /**
+            @brief Calculate controls driving the RBF behavior of the of the rig.
+            @note
+                This is considered as an advanced usage use case.
+            @param instance
+                The rig instance whose controls for driving machine learned behavior are to be calculated.
+            @see calculate
+        */
+        virtual void calculateRBFControls(RigInstance* instance) const = 0;
+        /**
+            @brief Calculate controls driving the RBF behavior of the rig.
+            @note
+                This is considered as an advanced usage use case.
+            @param instance
+                The rig instance whose controls for RBF behavior are to be calculated.
+            @param solverIndex
+                The RBF solver whose outputs need to be calculated.
+            @warning
+                The index must be less than the value returned by getRBFSolverCount.
+            @see calculate
+        */
+        virtual void calculateRBFControls(RigInstance* instance, std::uint16_t solverIndex) const = 0;
         /**
             @brief Calculate only the joint outputs of the rig.
             @note
diff --git a/Engine/Plugins/Animation/RigLogic/Source/RigLogicLib/Public/riglogic/version/Version.h b/Engine/Plugins/Animation/RigLogic/Source/RigLogicLib/Public/riglogic/version/Version.h
--- a/Engine/Plugins/Animation/RigLogic/Source/RigLogicLib/Public/riglogic/version/Version.h
+++ b/Engine/Plugins/Animation/RigLogic/Source/RigLogicLib/Public/riglogic/version/Version.h
@@ -2,7 +2,7 @@
 
 #pragma once
 
-#define RL_MAJOR_VERSION 9
-#define RL_MINOR_VERSION 2
+#define RL_MAJOR_VERSION 10
+#define RL_MINOR_VERSION 1
 #define RL_PATCH_VERSION 2
-#define RL_VERSION_STRING "9.2.2"
+#define RL_VERSION_STRING "10.1.2"
diff --git a/Engine/Plugins/Animation/RigLogic/Source/RigLogicLib/Public/status/version/Version.h b/Engine/Plugins/Animation/RigLogic/Source/RigLogicLib/Public/status/version/Version.h
--- a/Engine/Plugins/Animation/RigLogic/Source/RigLogicLib/Public/status/version/Version.h
+++ b/Engine/Plugins/Animation/RigLogic/Source/RigLogicLib/Public/status/version/Version.h
@@ -4,5 +4,5 @@
 
 #define SC_MAJOR_VERSION 1
 #define SC_MINOR_VERSION 2
-#define SC_PATCH_VERSION 3
-#define SC_VERSION_STRING "1.2.3"
+#define SC_PATCH_VERSION 5
+#define SC_VERSION_STRING "1.2.5"
diff --git a/Engine/Plugins/Animation/RigLogic/Source/RigLogicLib/Public/tdm/AxisAngle.h b/Engine/Plugins/Animation/RigLogic/Source/RigLogicLib/Public/tdm/AxisAngle.h
new file mode 100644
--- /dev/null
+++ b/Engine/Plugins/Animation/RigLogic/Source/RigLogicLib/Public/tdm/AxisAngle.h
@@ -0,0 +1,36 @@
+// Copyright Epic Games, Inc. All Rights Reserved.
+
+#pragma once
+
+#include "tdm/Types.h"
+#include "tdm/Ang.h"
+#include "tdm/Mat.h"
+#include "tdm/Computations.h"
+#include "tdm/Transforms.h"
+
+namespace tdm {
+
+template<typename T>
+struct axis_angle {
+    vec3<T> axis;
+    rad<T> angle;
+
+    axis_angle(vec3<T> axis_, rad<T> angle_) : axis{axis_}, angle{angle_} {
+    }
+
+    axis_angle(rad<T> x, rad<T> y, rad<T> z, handedness h = handedness::right) {
+        const mat4<T> m = rotate(x, y, z, h);
+        angle = rad<T>{std::acos((trace(m.template submat<3, 3>(0, 0)) - static_cast<T>(1.0)) / static_cast<T>(2.0))};
+        const T factor = static_cast<T>(1.0) / (static_cast<T>(2.0) * std::sin(angle.value));
+        axis[0] = factor * (m(2, 1) - m(1, 2));
+        axis[1] = factor * (m(0, 2) - m(2, 0));
+        axis[2] = factor * (m(1, 0) - m(0, 1));
+    }
+
+    axis_angle(const rad3<T>& angles, handedness h = handedness::right) : axis_angle{angles[0], angles[1], angles[2], h} {
+    }
+
+};
+
+
+}  // namespace tdm
diff --git a/Engine/Plugins/Animation/RigLogic/Source/RigLogicLib/Public/tdm/Transforms.h b/Engine/Plugins/Animation/RigLogic/Source/RigLogicLib/Public/tdm/Transforms.h
new file mode 100644
--- /dev/null
+++ b/Engine/Plugins/Animation/RigLogic/Source/RigLogicLib/Public/tdm/Transforms.h
@@ -0,0 +1,189 @@
+// Copyright Epic Games, Inc. All Rights Reserved.
+
+#pragma once
+
+#include "tdm/Computations.h"
+#include "tdm/Types.h"
+
+namespace tdm {
+
+inline constexpr double pi() {
+    return 3.14159265358979323846;
+}
+
+template<typename T>
+inline typename std::enable_if<std::is_floating_point<T>::value, T>::type degrees(T radians) {
+    static constexpr const auto c = static_cast<T>(180.0 / pi());
+    return c * radians;
+}
+
+template<typename T>
+inline typename std::enable_if<std::is_floating_point<T>::value, T>::type radians(T degrees) {
+    static constexpr const auto c = static_cast<T>(pi() / 180.0);
+    return c * degrees;
+}
+
+template<typename T>
+inline typename std::enable_if<std::is_floating_point<T>::value, T>::type fastasin(T value) {
+    constexpr float FASTASIN_HALF_PI = 1.5707963050f;
+    // Clamp input to [-1,1].
+    const bool nonnegative = (value >= 0.0f);
+    value = std::abs(value);
+    float omx = 1.0f - value;
+    if (omx < 0.0f) {
+        omx = 0.0f;
+    }
+    const float root = std::sqrt(omx);
+    // 7-degree minimax approximation
+    float result =
+        ((((((-0.0012624911f
+              * value + 0.0066700901f)
+             * value - 0.0170881256f)
+            * value + 0.0308918810f)
+           * value - 0.0501743046f)
+          * value + 0.0889789874f)
+         * value - 0.2145988016f)
+        * value + FASTASIN_HALF_PI;
+
+    result *= root;  // acos(|x|)
+    // acos(x) = pi - acos(-x) when x < 0, asin(x) = pi/2 - acos(x)
+    return (nonnegative ? FASTASIN_HALF_PI - result : result - FASTASIN_HALF_PI);
+}
+
+namespace affine {
+
+template<dim_t L, typename T>
+inline mat<L, L, T> scale(const vec<L, T>& factors) {
+    return mat<L, L, T>::diagonal(factors);
+}
+
+template<dim_t L, typename T>
+inline mat<L, L, T> scale(const mat<L, L, T>& m, const vec<L, T>& factors) {
+    return m * scale(factors);
+}
+
+template<dim_t L, typename T>
+inline mat<L, L, T> scale(T factor) {
+    return scale(vec<L, T>{factor});
+}
+
+template<dim_t L, typename T>
+inline mat<L, L, T> scale(const mat<L, L, T>& m, T factor) {
+    return scale(m, vec<L, T>{factor});
+}
+
+}  // namespace affine
+
+inline namespace projective {
+
+template<typename T>
+inline mat4<T> rotate(const vec3<T>& axis, rad<T> angle, handedness h = handedness::right) {
+    const rad<T> c{std::cos(angle.value)};
+    const rad<T> s{std::sin(angle.value) * static_cast<T>(h)};
+    const rad<T> one_minus_c = rad<T>{static_cast<T>(1)} - c;
+    const vec3<T> n = normalize(axis);
+    return mat4<T>{n[0] * n[0] * one_minus_c.value + c.value,
+                   n[1] * n[0] * one_minus_c.value - n[2] * s.value,
+                   n[2] * n[0] * one_minus_c.value + n[1] * s.value,
+                   static_cast<T>(0),
+                   n[0] * n[1] * one_minus_c.value + n[2] * s.value,
+                   n[1] * n[1] * one_minus_c.value + c.value,
+                   n[2] * n[1] * one_minus_c.value - n[0] * s.value,
+                   static_cast<T>(0),
+                   n[0] * n[2] * one_minus_c.value - n[1] * s.value,
+                   n[1] * n[2] * one_minus_c.value + n[0] * s.value,
+                   n[2] * n[2] * one_minus_c.value + c.value,
+                   static_cast<T>(0),
+                   static_cast<T>(0),
+                   static_cast<T>(0),
+                   static_cast<T>(0),
+                   static_cast<T>(1)};
+}
+
+template<typename T>
+inline mat4<T> rotate(const mat4<T>& m, const vec3<T>& axis, rad<T> angle, handedness h = handedness::right) {
+    return m * rotate(axis, angle, h);
+}
+
+template<typename T>
+inline mat4<T> rotate(rad<T> x, rad<T> y, rad<T> z, handedness h = handedness::right) {
+    const rad<T> sx{std::sin(x.value) * static_cast<T>(h)};
+    const rad<T> sy{std::sin(y.value) * static_cast<T>(h)};
+    const rad<T> sz{std::sin(z.value) * static_cast<T>(h)};
+    const rad<T> cx{std::cos(x.value)};
+    const rad<T> cy{std::cos(y.value)};
+    const rad<T> cz{std::cos(z.value)};
+    auto rx = mat4<T>::identity();
+    rx(1, 1) = cx.value;
+    rx(1, 2) = sx.value;
+    rx(2, 1) = -sx.value;
+    rx(2, 2) = cx.value;
+    auto ry = mat4<T>::identity();
+    ry(0, 0) = cy.value;
+    ry(0, 2) = -sy.value;
+    ry(2, 0) = sy.value;
+    ry(2, 2) = cy.value;
+    auto rz = mat4<T>::identity();
+    rz(0, 0) = cz.value;
+    rz(0, 1) = sz.value;
+    rz(1, 0) = -sz.value;
+    rz(1, 1) = cz.value;
+    return rx * ry * rz;
+}
+
+template<typename T>
+inline mat4<T> rotate(const mat4<T>& m, rad<T> x, rad<T> y, rad<T> z, handedness h = handedness::right) {
+    return m * rotate(x, y, z, h);
+}
+
+template<typename T>
+inline mat4<T> rotate(const rad3<T>& rotation, handedness h = handedness::right) {
+    return rotate(rotation[0], rotation[1], rotation[2], h);
+}
+
+template<typename T>
+inline mat4<T> rotate(const mat4<T>& m, const rad3<T>& rotation, handedness h = handedness::right) {
+    return m * rotate(rotation[0], rotation[1], rotation[2], h);
+}
+
+template<dim_t L, typename T>
+inline mat<L + 1, L + 1, T> scale(const vec<L, T>& factors) {
+    vec<L + 1, T> diagonal{static_cast<T>(1)};
+    factors.apply([&diagonal](const T& value, dim_t i) {
+                diagonal[i] = value;
+            });
+    return mat<L + 1, L + 1, T>::diagonal(diagonal);
+}
+
+template<dim_t L, typename T>
+inline mat<L + 1, L + 1, T> scale(const mat<L + 1, L + 1, T>& m, const vec<L, T>& factors) {
+    return m * scale(factors);
+}
+
+template<dim_t L, typename T>
+inline mat<L + 1, L + 1, T> scale(T factor) {
+    return scale(vec<L, T>{factor});
+}
+
+template<dim_t L, typename T>
+inline mat<L, L, T> scale(const mat<L, L, T>& m, T factor) {
+    return scale(m, vec<L - 1, T>{factor});
+}
+
+template<dim_t L, typename T>
+inline mat<L + 1, L + 1, T> translate(const vec<L, T>& position) {
+    auto m = mat<L + 1, L + 1, T>::identity();
+    position.apply([&m](const T& value, dim_t i) {
+                m(L, i) = value;
+            });
+    return m;
+}
+
+template<dim_t L, typename T>
+inline mat<L + 1, L + 1, T> translate(const mat<L + 1, L + 1, T>& m, const vec<L, T>& position) {
+    return m * translate(position);
+}
+
+}  // namespace projective
+
+}  // namespace tdm
diff --git a/Engine/Plugins/Animation/RigLogic/Source/RigLogicLib/Public/tdm/version/Version.h b/Engine/Plugins/Animation/RigLogic/Source/RigLogicLib/Public/tdm/version/Version.h
new file mode 100644
--- /dev/null
+++ b/Engine/Plugins/Animation/RigLogic/Source/RigLogicLib/Public/tdm/version/Version.h
@@ -0,0 +1,8 @@
+// Copyright Epic Games, Inc. All Rights Reserved.
+
+#pragma once
+
+#define TDM_MAJOR_VERSION 4
+#define TDM_MINOR_VERSION 0
+#define TDM_PATCH_VERSION 0
+#define TDM_VERSION_STRING "4.0.0"
diff --git a/Engine/Plugins/Animation/RigLogic/Source/RigLogicLib/Public/terse/types/Anchor.h b/Engine/Plugins/Animation/RigLogic/Source/RigLogicLib/Public/terse/types/Anchor.h
--- a/Engine/Plugins/Animation/RigLogic/Source/RigLogicLib/Public/terse/types/Anchor.h
+++ b/Engine/Plugins/Animation/RigLogic/Source/RigLogicLib/Public/terse/types/Anchor.h
@@ -2,7 +2,14 @@
 
 #pragma once
 
+#ifdef _MSC_VER
+    #pragma warning(push)
+    #pragma warning(disable : 4365 4987)
+#endif
 #include <utility>
+#ifdef _MSC_VER
+    #pragma warning(pop)
+#endif
 
 namespace terse {
 
diff --git a/Engine/Plugins/Animation/RigLogic/Source/RigLogicLib/Public/terse/types/ArchiveOffset.h b/Engine/Plugins/Animation/RigLogic/Source/RigLogicLib/Public/terse/types/ArchiveOffset.h
--- a/Engine/Plugins/Animation/RigLogic/Source/RigLogicLib/Public/terse/types/ArchiveOffset.h
+++ b/Engine/Plugins/Animation/RigLogic/Source/RigLogicLib/Public/terse/types/ArchiveOffset.h
@@ -32,9 +32,16 @@ struct ArchiveOffset {
         ArchiveOffset* target;
 
         explicit Proxy(ArchiveOffset& ptr) : target{std::addressof(ptr)} {
+            #if !defined(__clang__) && defined(__GNUC__) && (__GNUC__ >= 12)
+                #pragma GCC diagnostic push
+                #pragma GCC diagnostic ignored "-Wdangling-pointer"
+            #endif
             #if !defined(__clang_analyzer__)
                 target->proxy = this;
             #endif
+            #if !defined(__clang__) && defined(__GNUC__) && (__GNUC__ >= 12)
+                #pragma GCC diagnostic pop
+            #endif
         }
 
         ~Proxy() {
@@ -48,12 +55,26 @@ struct ArchiveOffset {
 
         Proxy(Proxy&& rhs) : target{nullptr} {
             std::swap(target, rhs.target);
+            #if !defined(__clang__) && defined(__GNUC__) && (__GNUC__ >= 12)
+                #pragma GCC diagnostic push
+                #pragma GCC diagnostic ignored "-Wdangling-pointer"
+            #endif
             target->proxy = this;
+            #if !defined(__clang__) && defined(__GNUC__) && (__GNUC__ >= 12)
+                #pragma GCC diagnostic pop
+            #endif
         }
 
         Proxy& operator=(Proxy&& rhs) {
             std::swap(target, rhs.target);
+            #if !defined(__clang__) && defined(__GNUC__) && (__GNUC__ >= 12)
+                #pragma GCC diagnostic push
+                #pragma GCC diagnostic ignored "-Wdangling-pointer"
+            #endif
             target->proxy = this;
+            #if !defined(__clang__) && defined(__GNUC__) && (__GNUC__ >= 12)
+                #pragma GCC diagnostic pop
+            #endif
             return *this;
         }
 
diff --git a/Engine/Plugins/Animation/RigLogic/Source/RigLogicLib/Public/terse/types/ArchiveSize.h b/Engine/Plugins/Animation/RigLogic/Source/RigLogicLib/Public/terse/types/ArchiveSize.h
--- a/Engine/Plugins/Animation/RigLogic/Source/RigLogicLib/Public/terse/types/ArchiveSize.h
+++ b/Engine/Plugins/Animation/RigLogic/Source/RigLogicLib/Public/terse/types/ArchiveSize.h
@@ -39,10 +39,17 @@ struct ArchiveSize {
         Anchor<TOffset>* base;
 
         Proxy(ArchiveSize& size, Anchor<TOffset>& offset) : target{std::addressof(size)}, base{std::addressof(offset)} {
+            #if !defined(__clang__) && defined(__GNUC__) && (__GNUC__ >= 12)
+                #pragma GCC diagnostic push
+                #pragma GCC diagnostic ignored "-Wdangling-pointer"
+            #endif
             #if !defined(__clang_analyzer__)
                 target->proxy = this;
                 base->onMove(onBaseMoved, this);
             #endif
+            #if !defined(__clang__) && defined(__GNUC__) && (__GNUC__ >= 12)
+                #pragma GCC diagnostic pop
+            #endif
         }
 
         ~Proxy() {
@@ -60,15 +67,29 @@ struct ArchiveSize {
         Proxy(Proxy&& rhs) : target{nullptr}, base{nullptr} {
             std::swap(target, rhs.target);
             std::swap(base, rhs.base);
+            #if !defined(__clang__) && defined(__GNUC__) && (__GNUC__ >= 12)
+                #pragma GCC diagnostic push
+                #pragma GCC diagnostic ignored "-Wdangling-pointer"
+            #endif
             target->proxy = this;
             base->onMove(onBaseMoved, this);
+            #if !defined(__clang__) && defined(__GNUC__) && (__GNUC__ >= 12)
+                #pragma GCC diagnostic pop
+            #endif
         }
 
         Proxy& operator=(Proxy&& rhs) {
             std::swap(target, rhs.target);
             std::swap(base, rhs.base);
+            #if !defined(__clang__) && defined(__GNUC__) && (__GNUC__ >= 12)
+                #pragma GCC diagnostic push
+                #pragma GCC diagnostic ignored "-Wdangling-pointer"
+            #endif
             target->proxy = this;
             base->onMove(onBaseMoved, this);
+            #if !defined(__clang__) && defined(__GNUC__) && (__GNUC__ >= 12)
+                #pragma GCC diagnostic pop
+            #endif
             return *this;
         }
 
diff --git a/Engine/Plugins/Animation/RigLogic/Source/RigLogicLib/Public/terse/types/DynArray.h b/Engine/Plugins/Animation/RigLogic/Source/RigLogicLib/Public/terse/types/DynArray.h
--- a/Engine/Plugins/Animation/RigLogic/Source/RigLogicLib/Public/terse/types/DynArray.h
+++ b/Engine/Plugins/Animation/RigLogic/Source/RigLogicLib/Public/terse/types/DynArray.h
@@ -13,10 +13,10 @@
 #include <functional>
 #include <iterator>
 #include <memory>
-#include <type_traits>
-#if !defined(__clang__)
-#include <span>
+#if defined(_MSC_VER) && !defined(__clang__) && (_MSC_VER >= 1900) && (__cplusplus >= 202002L)
+    #include <span>
 #endif
+#include <type_traits>
 #ifdef _MSC_VER
     #pragma warning(pop)
 #endif
@@ -73,7 +73,11 @@ class DynArray {
 
             #if defined(_MSC_VER) && !defined(__clang__)
                 if (size() != 0ul) {
-                    std::copy(start, end, std::span{data(), size()}.begin());
+                    #if (_MSC_VER >= 1900) && (__cplusplus >= 202002L)
+                        std::copy(start, end, std::span{data(), size()}.begin());
+                    #else
+                        std::copy(start, end, stdext::checked_array_iterator<value_type*>(data(), size()));
+                    #endif
                 }
             #else
                 std::copy(start, end, data());
@@ -210,7 +214,11 @@ class DynArray {
             resize_uninitialized(static_cast<std::size_t>(std::distance(start, end)));
             #if defined(_MSC_VER) && !defined(__clang__)
                 if (size() != 0ul) {
-                    std::copy(start, end, std::span{data(), size()}.begin());
+                    #if (_MSC_VER >= 1900) && (__cplusplus >= 202002L)
+                        std::copy(start, end, std::span{data(), size()}.begin());
+                    #else
+                        std::copy(start, end, stdext::checked_array_iterator<value_type*>(data(), size()));
+                    #endif
                 }
             #else
                 std::copy(start, end, data());
diff --git a/Engine/Plugins/Animation/RigLogic/Source/RigLogicLib/Public/terse/version/Version.h b/Engine/Plugins/Animation/RigLogic/Source/RigLogicLib/Public/terse/version/Version.h
--- a/Engine/Plugins/Animation/RigLogic/Source/RigLogicLib/Public/terse/version/Version.h
+++ b/Engine/Plugins/Animation/RigLogic/Source/RigLogicLib/Public/terse/version/Version.h
@@ -4,5 +4,5 @@
 
 #define TERSE_MAJOR_VERSION 3
 #define TERSE_MINOR_VERSION 1
-#define TERSE_PATCH_VERSION 7
-#define TERSE_VERSION_STRING "3.1.7"
+#define TERSE_PATCH_VERSION 9
+#define TERSE_VERSION_STRING "3.1.9"
diff --git a/Engine/Plugins/Animation/RigLogic/Source/RigLogicLib/Public/trimd/AVX.h b/Engine/Plugins/Animation/RigLogic/Source/RigLogicLib/Public/trimd/AVX.h
--- a/Engine/Plugins/Animation/RigLogic/Source/RigLogicLib/Public/trimd/AVX.h
+++ b/Engine/Plugins/Animation/RigLogic/Source/RigLogicLib/Public/trimd/AVX.h
@@ -309,6 +309,31 @@ inline F256 andnot(const F256& lhs, const F256& rhs) {
     return F256{_mm256_andnot_ps(lhs.data, rhs.data)};
 }
 
+inline F256 rsqrt(const F256& rhs) {
+    #ifndef TRIMD_ENABLE_FAST_INVERSE_SQRT
+    return F256{_mm256_rsqrt_ps(rhs.data)};
+    #else
+    // Pure AVX2-only implementation:
+    //
+    // const __m256i shifted = _mm256_srli_epi32(_mm256_castps_si256(rhs.data), 1);
+    // const __m256i subtracted = _mm256_sub_epi32(_mm256_set1_epi32(0x5f1ffff9), shifted);
+    // F256 result{_mm256_castsi256_ps(subtracted)};
+    // result *= F256{0.703952253f} * (F256{2.38924456f} - rhs * result * result);
+    // return result;
+    //
+    // Combination of SSE and AVX:
+    const __m128 lower = _mm256_castps256_ps128(rhs.data);
+    const __m128 upper = _mm256_extractf128_ps(rhs.data, 1);
+    const __m128i shiftedLower = _mm_srli_epi32(_mm_castps_si128(lower), 1);
+    const __m128i shiftedUpper = _mm_srli_epi32(_mm_castps_si128(upper), 1);
+    const __m128i subtractedLower = _mm_sub_epi32(_mm_set1_epi32(0x5f1ffff9), shiftedLower);
+    const __m128i subtractedUpper = _mm_sub_epi32(_mm_set1_epi32(0x5f1ffff9), shiftedUpper);
+    F256 result{_mm256_insertf128_ps(_mm256_castps128_ps256(_mm_castsi128_ps(subtractedLower)), _mm_castsi128_ps(subtractedUpper), 1)};
+    result *= F256{0.703952253f} * (F256{2.38924456f} - rhs * result * result);
+    return result;
+    #endif  // TRIMD_ENABLE_FAST_INVERSE_SQRT
+}
+
 } // namespace avx
 
 } // namespace trimd
diff --git a/Engine/Plugins/Animation/RigLogic/Source/RigLogicLib/Public/trimd/Fallback.h b/Engine/Plugins/Animation/RigLogic/Source/RigLogicLib/Public/trimd/Fallback.h
--- a/Engine/Plugins/Animation/RigLogic/Source/RigLogicLib/Public/trimd/Fallback.h
+++ b/Engine/Plugins/Animation/RigLogic/Source/RigLogicLib/Public/trimd/Fallback.h
@@ -2,7 +2,14 @@
 
 #pragma once
 
-#include <string_view> // std::swap
+#ifdef _MSC_VER
+    #pragma warning(push)
+    #pragma warning(disable : 4365 4987)
+#endif
+#include <utility>
+#ifdef _MSC_VER
+    #pragma warning(pop)
+#endif
 
 namespace trimd {
 
@@ -240,6 +247,11 @@ inline T256<T128> andnot(const T256<T128>& lhs, const T256<T128>& rhs) {
     return T256<T128>{andnot(lhs.data1, rhs.data1), andnot(lhs.data2, rhs.data2)};
 }
 
+template<typename T128>
+inline T256<T128> rsqrt(const T256<T128>& rhs) {
+    return T256<T128>{rsqrt(rhs.data1), rsqrt(rhs.data2)};
+}
+
 }  // namespace fallback
 
 }  // namespace trimd
diff --git a/Engine/Plugins/Animation/RigLogic/Source/RigLogicLib/Public/trimd/NEON.h b/Engine/Plugins/Animation/RigLogic/Source/RigLogicLib/Public/trimd/NEON.h
--- a/Engine/Plugins/Animation/RigLogic/Source/RigLogicLib/Public/trimd/NEON.h
+++ b/Engine/Plugins/Animation/RigLogic/Source/RigLogicLib/Public/trimd/NEON.h
@@ -248,10 +248,25 @@ inline F128 andnot(const F128& lhs, const F128& rhs) {
     return F128{vreinterpretq_f32_u32(vbicq_u32(vreinterpretq_u32_f32(rhs.data), vreinterpretq_u32_f32(lhs.data)))};
 }
 
+inline F128 rsqrt(const F128& rhs) {
+    #ifndef TRIMD_ENABLE_FAST_INVERSE_SQRT
+    const float32x4_t reciprocal0 = vrsqrteq_f32(rhs.data);
+    const float32x4_t reciprocal1 = vmulq_f32(vrsqrtsq_f32(vmulq_f32(reciprocal0, reciprocal0), rhs.data), reciprocal0);
+    return F128{reciprocal1};
+    #else
+    const uint32x4_t shifted = vshrq_n_u32(vreinterpretq_u32_f32(rhs.data), 1);
+    const uint32x4_t subtracted = vsubq_u32(vdupq_n_u32(0x5f1ffff9), shifted);
+    F128 result{vreinterpretq_f32_u32(subtracted)};
+    result *= F128{0.703952253f} * (F128{2.38924456f} - rhs * result * result);
+    return result;
+    #endif  // TRIMD_ENABLE_FAST_INVERSE_SQRT
+}
+
 using F256 = fallback::T256<F128>;
 using fallback::transpose;
 using fallback::abs;
 using fallback::andnot;
+using fallback::rsqrt;
 
 } // namespace neon
 
diff --git a/Engine/Plugins/Animation/RigLogic/Source/RigLogicLib/Public/trimd/SSE.h b/Engine/Plugins/Animation/RigLogic/Source/RigLogicLib/Public/trimd/SSE.h
--- a/Engine/Plugins/Animation/RigLogic/Source/RigLogicLib/Public/trimd/SSE.h
+++ b/Engine/Plugins/Animation/RigLogic/Source/RigLogicLib/Public/trimd/SSE.h
@@ -283,10 +283,23 @@ inline F128 andnot(const F128& lhs, const F128& rhs) {
     return F128{_mm_andnot_ps(lhs.data, rhs.data)};
 }
 
+inline F128 rsqrt(const F128& rhs) {
+    #ifndef TRIMD_ENABLE_FAST_INVERSE_SQRT
+    return F128{_mm_rsqrt_ps(rhs.data)};
+    #else
+    const __m128i shifted = _mm_srli_epi32(_mm_castps_si128(rhs.data), 1);
+    const __m128i subtracted = _mm_sub_epi32(_mm_set1_epi32(0x5f1ffff9), shifted);
+    F128 result{_mm_castsi128_ps(subtracted)};
+    result *= F128{0.703952253f} * (F128{2.38924456f} - rhs * result * result);
+    return result;
+    #endif  // TRIMD_ENABLE_FAST_INVERSE_SQRT
+}
+
 using F256 = fallback::T256<F128>;
 using fallback::transpose;
 using fallback::abs;
 using fallback::andnot;
+using fallback::rsqrt;
 
 } // namespace sse
 
diff --git a/Engine/Plugins/Animation/RigLogic/Source/RigLogicLib/Public/trimd/Scalar.h b/Engine/Plugins/Animation/RigLogic/Source/RigLogicLib/Public/trimd/Scalar.h
--- a/Engine/Plugins/Animation/RigLogic/Source/RigLogicLib/Public/trimd/Scalar.h
+++ b/Engine/Plugins/Animation/RigLogic/Source/RigLogicLib/Public/trimd/Scalar.h
@@ -295,11 +295,36 @@ inline T128<T> andnot(const T128<T>& lhs, const T128<T>& rhs) {
     return ~lhs & rhs;
 }
 
+template<typename T>
+inline T128<T> rsqrt(const T128<T>& rhs) {
+    #ifndef TRIMD_ENABLE_FAST_INVERSE_SQRT
+    return T128<T>{1.0f / std::sqrt(rhs.data[0]),
+                   1.0f / std::sqrt(rhs.data[1]),
+                   1.0f / std::sqrt(rhs.data[2]),
+                   1.0f / std::sqrt(rhs.data[3])};
+    #else
+    std::uint32_t asInts[4];
+    std::memcpy(asInts, rhs.data.data(), sizeof(asInts));
+    asInts[0] = 0x5f1ffff9 - (asInts[0] >> 1);
+    asInts[1] = 0x5f1ffff9 - (asInts[1] >> 1);
+    asInts[2] = 0x5f1ffff9 - (asInts[2] >> 1);
+    asInts[3] = 0x5f1ffff9 - (asInts[3] >> 1);
+    T128<T> result;
+    std::memcpy(result.data.data(), asInts, sizeof(asInts));
+    result.data[0] *= 0.703952253f * (2.38924456f - rhs.data[0] * result.data[0] * result.data[0]);
+    result.data[1] *= 0.703952253f * (2.38924456f - rhs.data[1] * result.data[1] * result.data[1]);
+    result.data[2] *= 0.703952253f * (2.38924456f - rhs.data[2] * result.data[2] * result.data[2]);
+    result.data[3] *= 0.703952253f * (2.38924456f - rhs.data[3] * result.data[3] * result.data[3]);
+    return result;
+    #endif  // TRIMD_ENABLE_FAST_INVERSE_SQRT
+}
+
 using F128 = T128<float>;
 using F256 = fallback::T256<F128>;
 using fallback::transpose;
 using fallback::abs;
 using fallback::andnot;
+using fallback::rsqrt;
 
 }  // namespace scalar
 
diff --git a/Engine/Plugins/Animation/RigLogic/Source/RigLogicLib/Public/trimd/TRiMD.h b/Engine/Plugins/Animation/RigLogic/Source/RigLogicLib/Public/trimd/TRiMD.h
--- a/Engine/Plugins/Animation/RigLogic/Source/RigLogicLib/Public/trimd/TRiMD.h
+++ b/Engine/Plugins/Animation/RigLogic/Source/RigLogicLib/Public/trimd/TRiMD.h
@@ -14,6 +14,7 @@ namespace trimd {
     using avx::abs;
     using avx::transpose;
     using avx::andnot;
+    using avx::rsqrt;
 #elif defined(TRIMD_ENABLE_SSE)
     using F256 = sse::F256;
 #elif defined(TRIMD_ENABLE_NEON)
@@ -27,11 +28,13 @@ namespace trimd {
     using sse::abs;
     using sse::transpose;
     using sse::andnot;
+    using sse::rsqrt;
 #elif defined(TRIMD_ENABLE_NEON)
     using F128 = neon::F128;
     using neon::abs;
     using neon::transpose;
     using neon::andnot;
+    using neon::rsqrt;
 #else
     using F128 = scalar::F128;
 #endif  // TRIMD_ENABLE_SSE
@@ -39,5 +42,6 @@ namespace trimd {
 using scalar::abs;
 using scalar::transpose;
 using scalar::andnot;
+using scalar::rsqrt;
 
 }  // namespace trimd
diff --git a/Engine/Plugins/Animation/RigLogic/Source/RigLogicLib/Public/trimd/version/Version.h b/Engine/Plugins/Animation/RigLogic/Source/RigLogicLib/Public/trimd/version/Version.h
--- a/Engine/Plugins/Animation/RigLogic/Source/RigLogicLib/Public/trimd/version/Version.h
+++ b/Engine/Plugins/Animation/RigLogic/Source/RigLogicLib/Public/trimd/version/Version.h
@@ -3,6 +3,6 @@
 #pragma once
 
 #define TRIMD_MAJOR_VERSION 2
-#define TRIMD_MINOR_VERSION 2
-#define TRIMD_PATCH_VERSION 0
-#define TRIMD_VERSION_STRING "2.2.0"
+#define TRIMD_MINOR_VERSION 3
+#define TRIMD_PATCH_VERSION 1
+#define TRIMD_VERSION_STRING "2.3.1"
diff --git a/Engine/Plugins/Animation/RigLogic/Source/RigLogicLib/Public/trio/version/Version.h b/Engine/Plugins/Animation/RigLogic/Source/RigLogicLib/Public/trio/version/Version.h
--- a/Engine/Plugins/Animation/RigLogic/Source/RigLogicLib/Public/trio/version/Version.h
+++ b/Engine/Plugins/Animation/RigLogic/Source/RigLogicLib/Public/trio/version/Version.h
@@ -4,5 +4,5 @@
 
 #define TRIO_MAJOR_VERSION 4
 #define TRIO_MINOR_VERSION 0
-#define TRIO_PATCH_VERSION 11
-#define TRIO_VERSION_STRING "4.0.11"
+#define TRIO_PATCH_VERSION 14
+#define TRIO_VERSION_STRING "4.0.14"
diff --git a/Engine/Plugins/Animation/RigLogic/Source/RigLogicLib/RigLogicLib.Build.cs b/Engine/Plugins/Animation/RigLogic/Source/RigLogicLib/RigLogicLib.Build.cs
--- a/Engine/Plugins/Animation/RigLogic/Source/RigLogicLib/RigLogicLib.Build.cs
+++ b/Engine/Plugins/Animation/RigLogic/Source/RigLogicLib/RigLogicLib.Build.cs
@@ -26,6 +26,8 @@ namespace UnrealBuildTool.Rules
 
 			Type = ModuleType.CPlusPlus;
 
+			PrivateDefinitions.Add("RL_BUILD_WITH_ZYX_ROTATION_ORDER=1");
+
 			if (Target.LinkType != TargetLinkType.Monolithic)
 			{
 				PrivateDefinitions.Add("RL_BUILD_SHARED=1");
@@ -51,7 +53,6 @@ namespace UnrealBuildTool.Rules
 			}
 
 			PrivateDefinitions.Add("RL_AUTODETECT_SSE=1");
-			PublicDefinitions.Add("RL_BUILD_WITH_ML_EVALUATOR=1");
 
 			bDisableAutoRTFMInstrumentation = true;
 		}
diff --git a/Engine/Plugins/Animation/RigLogic/Source/RigLogicLibTest/Private/arrayviewtests/TestArrayView.cpp b/Engine/Plugins/Animation/RigLogic/Source/RigLogicLibTest/Private/arrayviewtests/TestArrayView.cpp
--- a/Engine/Plugins/Animation/RigLogic/Source/RigLogicLibTest/Private/arrayviewtests/TestArrayView.cpp
+++ b/Engine/Plugins/Animation/RigLogic/Source/RigLogicLibTest/Private/arrayviewtests/TestArrayView.cpp
@@ -14,6 +14,12 @@ TYPED_TEST(TestArrayView, DefaultCtor) {
     EXPECT_EQ(arr.size(), 0ul);
 }
 
+TYPED_TEST(TestArrayView, NullCtor) {
+    ArrayView<TypeParam> arr{nullptr, 0ul};
+    EXPECT_EQ(arr.data(), nullptr);
+    EXPECT_EQ(arr.size(), 0ul);
+}
+
 TYPED_TEST(TestArrayView, CopyCtor) {
     auto src1 = this->template initArray<length>();
     auto src2 = this->template initArray<length>();
diff --git a/Engine/Plugins/Animation/RigLogic/Source/RigLogicLibTest/Private/dnatests/FakeDNAReader.h b/Engine/Plugins/Animation/RigLogic/Source/RigLogicLibTest/Private/dnatests/FakeDNAReader.h
--- a/Engine/Plugins/Animation/RigLogic/Source/RigLogicLibTest/Private/dnatests/FakeDNAReader.h
+++ b/Engine/Plugins/Animation/RigLogic/Source/RigLogicLibTest/Private/dnatests/FakeDNAReader.h
@@ -541,6 +541,146 @@ class FakeDNAReader : public Reader {
             return {};
         }
 
+        std::uint16_t getRBFPoseCount() const override {
+            return {};
+        }
+
+        StringView getRBFPoseName(std::uint16_t poseIndex) const override {
+            return {};
+        }
+
+        ConstArrayView<std::uint16_t> getRBFPoseJointOutputIndices(std::uint16_t poseIndex) const override {
+            return {};
+        }
+
+        ConstArrayView<std::uint16_t> getRBFPoseBlendShapeChannelOutputIndices(std::uint16_t poseIndex) const override {
+            return {};
+        }
+
+        ConstArrayView<std::uint16_t> getRBFPoseAnimatedMapOutputIndices(std::uint16_t poseIndex) const override {
+            return {};
+        }
+
+        ConstArrayView<float> getRBFPoseJointOutputValues(std::uint16_t poseIndex) const override {
+            return {};
+        }
+
+        float getRBFPoseScale(std::uint16_t poseIndex) const override {
+            return {};
+        }
+
+        std::uint16_t getRBFSolverCount() const override {
+            return {};
+        }
+
+        std::uint16_t getRBFSolverIndexListCount() const override {
+            return {};
+        }
+
+        ConstArrayView<std::uint16_t> getRBFSolverIndicesForLOD(std::uint16_t lod) const override {
+            return {};
+        }
+
+        StringView getRBFSolverName(std::uint16_t solverIndex) const override {
+            return {};
+        }
+
+        ConstArrayView<std::uint16_t> getRBFSolverRawControlIndices(std::uint16_t solverIndex) const override {
+            return {};
+        }
+
+        ConstArrayView<std::uint16_t> getRBFSolverPoseIndices(std::uint16_t solverIndex) const override {
+            return {};
+        }
+
+        ConstArrayView<float> getRBFSolverRawControlValues(std::uint16_t solverIndex) const override {
+            return {};
+        }
+
+        RBFSolverType getRBFSolverType(std::uint16_t solverIndex) const override {
+            return {};
+        }
+
+        float getRBFSolverRadius(std::uint16_t solverIndex) const override {
+            return {};
+        }
+
+        AutomaticRadius getRBFSolverAutomaticRadius(std::uint16_t solverIndex) const override {
+            return {};
+        }
+
+        float getRBFSolverWeightThreshold(std::uint16_t solverIndex) const override {
+            return {};
+        }
+
+        RBFDistanceMethod getRBFSolverDistanceMethod(std::uint16_t solverIndex) const override {
+            return {};
+        }
+
+        RBFNormalizeMethod getRBFSolverNormalizeMethod(std::uint16_t solverIndex) const override {
+            return {};
+        }
+
+        RBFFunctionType getRBFSolverFunctionType(std::uint16_t solverIndex) const override {
+            return {};
+        }
+
+        TwistAxis getRBFSolverTwistAxis(std::uint16_t solverIndex) const override {
+            return {};
+        }
+
+        TranslationRepresentation getJointTranslationRepresentation(std::uint16_t jointIndex) const override {
+            return {};
+        }
+
+        RotationRepresentation getJointRotationRepresentation(std::uint16_t jointIndex) const override {
+            return {};
+        }
+
+        ScaleRepresentation getJointScaleRepresentation(std::uint16_t jointIndex) const override {
+            return {};
+        }
+
+        std::uint16_t getTwistCount() const override {
+            return {};
+        }
+
+        TwistAxis getTwistSetupTwistAxis(std::uint16_t twistIndex) const override {
+            return {};
+        }
+
+        ConstArrayView<std::uint16_t> getTwistInputControlIndices(std::uint16_t twistIndex) const override {
+            return {};
+        }
+
+        ConstArrayView<std::uint16_t> getTwistOutputJointIndices(std::uint16_t twistIndex) const override {
+            return {};
+        }
+
+        ConstArrayView<float> getTwistBlendWeights(std::uint16_t twistIndex) const override {
+            return {};
+        }
+
+        std::uint16_t getSwingCount() const override {
+            return {};
+        }
+
+        TwistAxis getSwingSetupTwistAxis(std::uint16_t swingIndex) const override {
+            return {};
+        }
+
+        ConstArrayView<std::uint16_t> getSwingInputControlIndices(std::uint16_t swingIndex) const override {
+            return {};
+        }
+
+        ConstArrayView<std::uint16_t> getSwingOutputJointIndices(std::uint16_t swingIndex) const override {
+            return {};
+        }
+
+        ConstArrayView<float> getSwingBlendWeights(std::uint16_t swingIndex) const override {
+            return {};
+        }
+
         void unload(DataLayer layer) override {
         }
 
diff --git a/Engine/Plugins/Animation/RigLogic/Source/RigLogicLibTest/Private/dnatests/FixturesJSON.cpp b/Engine/Plugins/Animation/RigLogic/Source/RigLogicLibTest/Private/dnatests/FixturesJSON.cpp
--- a/Engine/Plugins/Animation/RigLogic/Source/RigLogicLibTest/Private/dnatests/FixturesJSON.cpp
+++ b/Engine/Plugins/Animation/RigLogic/Source/RigLogicLibTest/Private/dnatests/FixturesJSON.cpp
@@ -15,7 +15,7 @@ const char* jsonDNA =
     },\n\
     \"version\": {\n\
         \"generation\": 2,\n\
-        \"version\": 3\n\
+        \"version\": 4\n\
     },\n\
     \"index\": {\n\
         \"entries\": [{\n\
@@ -43,6 +43,21 @@ const char* jsonDNA =
             \"version\": 65536,\n\
             \"offset\": 0,\n\
             \"size\": 0\n\
+        }, {\n\
+            \"id\": 1919051362,\n\
+            \"version\": 65536,\n\
+            \"offset\": 0,\n\
+            \"size\": 0\n\
+        }, {\n\
+            \"id\": 1784835428,\n\
+            \"version\": 65536,\n\
+            \"offset\": 0,\n\
+            \"size\": 0\n\
+        }, {\n\
+            \"id\": 1953985399,\n\
+            \"version\": 65536,\n\
+            \"offset\": 0,\n\
+            \"size\": 0\n\
         }]\n\
     },\n\
     \"desc1.1\": {\n\
@@ -180,6 +195,21 @@ const char* jsonDNA =
             \"indices\": []\n\
         },\n\
         \"neuralNetworks\": []\n\
+    },\n\
+    \"rbfb1.0\": {\n\
+        \"lodSolverMapping\": {\n\
+            \"lods\": [],\n\
+            \"indices\": []\n\
+        },\n\
+        \"solvers\": [],\n\
+        \"poses\": []\n\
+    },\n\
+    \"jbmd1.0\": {\n\
+        \"jointRepresentations\": []\n\
+    },\n\
+    \"twsw1.0\": {\n\
+        \"twists\": [],\n\
+        \"swings\": []\n\
     }\n\
 }";
 
diff --git a/Engine/Plugins/Animation/RigLogic/Source/RigLogicLibTest/Private/dnatests/Fixturesv21.cpp b/Engine/Plugins/Animation/RigLogic/Source/RigLogicLibTest/Private/dnatests/Fixturesv21.cpp
--- a/Engine/Plugins/Animation/RigLogic/Source/RigLogicLibTest/Private/dnatests/Fixturesv21.cpp
+++ b/Engine/Plugins/Animation/RigLogic/Source/RigLogicLibTest/Private/dnatests/Fixturesv21.cpp
@@ -990,6 +990,10 @@ const unsigned char RawV21::footer[] = {
 };
 
 std::vector<char> RawV21::getBytes() {
+    #if !defined(__clang__) && defined(__GNUC__) && (__GNUC__ >= 12)
+        #pragma GCC diagnostic push
+        #pragma GCC diagnostic ignored "-Wstringop-overflow"
+    #endif
     std::vector<char> bytes;
     // Header
     bytes.insert(bytes.end(), header, header + sizeof(header));
@@ -1014,6 +1018,9 @@ std::vector<char> RawV21::getBytes() {
     // Footer
     bytes.insert(bytes.end(), footer, footer + sizeof(footer));
     return bytes;
+    #if !defined(__clang__) && defined(__GNUC__) && (__GNUC__ >= 12)
+        #pragma GCC diagnostic pop
+    #endif
 }
 
 #ifdef __clang__
@@ -1317,14 +1324,14 @@ const pma::Vector<std::uint16_t> DecodedV21::jointRowCount = {
 const std::uint16_t DecodedV21::jointColumnCount = 10u;
 const pma::Vector<pma::Matrix<std::uint16_t> > DecodedV21::jointVariableIndices = {
     {  // MaxLOD-0 - MinLOD-1
-        {2, 3, 5, 18, 20, 36, 38, 55, 56, 63, 45, 46, 71},
+        {2, 3, 5, 18, 20, 36, 38, 45, 46, 55, 56, 63, 71},
         {2, 3, 5, 18, 20, 55, 56}
     },
     {  // MaxLOD-1 - MinLOD-1
         {2, 3, 5, 18, 20, 37, 38}
     },
     {  // MaxLOD-0 - MinLOD-0
-        {2, 3, 5, 18, 20, 36, 38, 55, 56, 63, 45, 46, 71}
+        {2, 3, 5, 18, 20, 36, 38, 45, 46, 55, 56, 63, 71}
     }
 };
 const pma::Vector<pma::Matrix<std::uint16_t> > DecodedV21::jointGroupLODs = {
diff --git a/Engine/Plugins/Animation/RigLogic/Source/RigLogicLibTest/Private/dnatests/Fixturesv22.cpp b/Engine/Plugins/Animation/RigLogic/Source/RigLogicLibTest/Private/dnatests/Fixturesv22.cpp
--- a/Engine/Plugins/Animation/RigLogic/Source/RigLogicLibTest/Private/dnatests/Fixturesv22.cpp
+++ b/Engine/Plugins/Animation/RigLogic/Source/RigLogicLibTest/Private/dnatests/Fixturesv22.cpp
@@ -1181,51 +1181,63 @@ const unsigned char RawV2xNewerWithUnknownDataPreservedAndDNARewritten::header[]
     0x00, 0x02,  // Generation
     0xff, 0xff,  // Version
     // Index Table
-    0x00, 0x00, 0x00, 0x0b,  // Index table entry count
+    0x00, 0x00, 0x00, 0x0e,  // Index table entry count
     0x64, 0x65, 0x73, 0x63,  // Descriptor id
     0x00, 0x01, 0x00, 0x02,  // Descriptor version
-    0x00, 0x00, 0x00, 0xbb,  // Descriptor offset
+    0x00, 0x00, 0x00, 0xeb,  // Descriptor offset
     0x00, 0x00, 0x00, 0x10,  // Descriptor size
     0x64, 0x65, 0x66, 0x6e,  // Definition id
     0x00, 0x01, 0x00, 0x02,  // Definition version
-    0x00, 0x00, 0x00, 0xcb,  // Definition offset
+    0x00, 0x00, 0x00, 0xfb,  // Definition offset
     0x00, 0x00, 0x00, 0x10,  // Definition size
     0x62, 0x68, 0x76, 0x72,  // Behavior id
     0x00, 0x01, 0x00, 0x02,  // Behavior version
-    0x00, 0x00, 0x00, 0xdb,  // Behavior offset
+    0x00, 0x00, 0x01, 0x0b,  // Behavior offset
     0x00, 0x00, 0x00, 0x20,  // Behavior size
     0x67, 0x65, 0x6f, 0x6d,  // Geometry id
     0x00, 0x01, 0x00, 0x02,  // Geometry version
-    0x00, 0x00, 0x00, 0xfb,  // Geometry offset
+    0x00, 0x00, 0x01, 0x2b,  // Geometry offset
     0x00, 0x00, 0x00, 0x20,  // Geometry size
     0x75, 0x6e, 0x6b, 0x31,  // Unknown 1 id
     0x00, 0x01, 0x00, 0x03,  // Unknown 1 version
-    0x00, 0x00, 0x01, 0x1b,  // Unknown 1 offset
+    0x00, 0x00, 0x01, 0x4b,  // Unknown 1 offset
     0x00, 0x00, 0x00, 0x10,  // Unknown 1 size
     0x75, 0x6e, 0x6b, 0x32,  // Unknown 2 id
     0x00, 0x01, 0x00, 0x02,  // Unknown 2 version
-    0x00, 0x00, 0x01, 0x2b,  // Unknown 2 offset
+    0x00, 0x00, 0x01, 0x5b,  // Unknown 2 offset
     0x00, 0x00, 0x00, 0x20,  // Unknown 2 size
     0x64, 0x65, 0x73, 0x63,  // Descriptor id
     0x00, 0x01, 0x00, 0x01,  // Descriptor version
-    0x00, 0x00, 0x01, 0x4b,  // Descriptor offset
+    0x00, 0x00, 0x01, 0x7b,  // Descriptor offset
     0x00, 0x00, 0x00, 0x24,  // Descriptor size
     0x64, 0x65, 0x66, 0x6e,  // Definition id
     0x00, 0x01, 0x00, 0x01,  // Definition version
-    0x00, 0x00, 0x01, 0x6f,  // Definition offset
+    0x00, 0x00, 0x01, 0x9f,  // Definition offset
     0x00, 0x00, 0x00, 0x5c,  // Definition size
     0x62, 0x68, 0x76, 0x72,  // Behavior id
     0x00, 0x01, 0x00, 0x01,  // Behavior version
-    0x00, 0x00, 0x01, 0xcb,  // Behavior offset
+    0x00, 0x00, 0x01, 0xfb,  // Behavior offset
     0x00, 0x00, 0x00, 0x56,  // Behavior size
     0x67, 0x65, 0x6f, 0x6d,  // Geometry id
     0x00, 0x01, 0x00, 0x01,  // Geometry version
-    0x00, 0x00, 0x02, 0x21,  // Geometry offset
+    0x00, 0x00, 0x02, 0x51,  // Geometry offset
     0x00, 0x00, 0x00, 0x04,  // Geometry size
     0x6d, 0x6c, 0x62, 0x68,  // Machine learned behavior id
     0x00, 0x01, 0x00, 0x00,  // Machine learned behavior version
-    0x00, 0x00, 0x02, 0x25,  // Machine learned behavior offset
-    0x00, 0x00, 0x00, 0x18  // Machine learned behavior size
+    0x00, 0x00, 0x02, 0x55,  // Machine learned behavior offset
+    0x00, 0x00, 0x00, 0x18,  // Machine learned behavior size
+    0x72, 0x62, 0x66, 0x62,  // RBF behavior id
+    0x00, 0x01, 0x00, 0x00,  // RBF behavior version
+    0x00, 0x00, 0x02, 0x6d,  // RBF behavior offset
+    0x00, 0x00, 0x00, 0x10,  // RBF behavior size
+    0x6a, 0x62, 0x6d, 0x64,  // Joint behavior metadata id
+    0x00, 0x01, 0x00, 0x00,  // Joint behavior metadata version
+    0x00, 0x00, 0x02, 0x7d,  // Joint behavior metadata offset
+    0x00, 0x00, 0x00, 0x04,  // Joint behavior metadata size
+    0x74, 0x77, 0x73, 0x77,  // Twist swing setups id
+    0x00, 0x01, 0x00, 0x00,  // Twist swing setups version
+    0x00, 0x00, 0x02, 0x81,  // Twist swing setups offset
+    0x00, 0x00, 0x00, 0x08  // Twist swing setups size
 };
 
 const unsigned char RawV2xNewerWithUnknownDataPreservedAndDNARewritten::descriptor[] = {
@@ -1274,6 +1286,22 @@ const unsigned char RawV2xNewerWithUnknownDataPreservedAndDNARewritten::machineL
     0x00, 0x00, 0x00, 0x00
 };
 
+const unsigned char RawV2xNewerWithUnknownDataPreservedAndDNARewritten::rbfBehavior[] = {
+    0x00, 0x00, 0x00, 0x00,
+    0x00, 0x00, 0x00, 0x00,
+    0x00, 0x00, 0x00, 0x00,
+    0x00, 0x00, 0x00, 0x00
+};
+
+const unsigned char RawV2xNewerWithUnknownDataPreservedAndDNARewritten::jointBehaviorMetaData[] = {
+    0x00, 0x00, 0x00, 0x00
+};
+
+const unsigned char RawV2xNewerWithUnknownDataPreservedAndDNARewritten::twistSwingBehavior[] {
+    0x00, 0x00, 0x00, 0x00,
+    0x00, 0x00, 0x00, 0x00
+};
+
 std::vector<char> RawV2xNewerWithUnknownDataPreservedAndDNARewritten::getBytes() {
     std::vector<char> bytes;
     // Header
@@ -1302,6 +1330,12 @@ std::vector<char> RawV2xNewerWithUnknownDataPreservedAndDNARewritten::getBytes()
     bytes.insert(bytes.end(), geometry, geometry + sizeof(geometry));
     // Machine learned behavior
     bytes.insert(bytes.end(), machineLearnedBehavior, machineLearnedBehavior + sizeof(machineLearnedBehavior));
+    // RBF behavior
+    bytes.insert(bytes.end(), rbfBehavior, rbfBehavior + sizeof(rbfBehavior));
+    // Joint behavior metadata
+    bytes.insert(bytes.end(), jointBehaviorMetaData, jointBehaviorMetaData + sizeof(jointBehaviorMetaData));
+    // Twist swing behavior
+    bytes.insert(bytes.end(), twistSwingBehavior, twistSwingBehavior + sizeof(twistSwingBehavior));
 
     return bytes;
 }
@@ -1311,28 +1345,39 @@ const unsigned char RawV2xNewerWithUnknownDataIgnoredAndDNARewritten::header[] =
     0x00, 0x02,  // Generation
     0xff, 0xff,  // Version
     // Index Table
-    0x00, 0x00, 0x00, 0x05,  // Index table entry count
+    0x00, 0x00, 0x00, 0x08,  // Index table entry count
     0x64, 0x65, 0x73, 0x63,  // Descriptor id
     0x00, 0x01, 0x00, 0x01,  // Descriptor version
-    0x00, 0x00, 0x00, 0x5b,  // Descriptor offset
+    0x00, 0x00, 0x00, 0x8b,  // Descriptor offset
     0x00, 0x00, 0x00, 0x24,  // Descriptor size
     0x64, 0x65, 0x66, 0x6e,  // Definition id
     0x00, 0x01, 0x00, 0x01,  // Definition version
-    0x00, 0x00, 0x00, 0x7f,  // Definition offset
+    0x00, 0x00, 0x00, 0xaf,  // Definition offset
     0x00, 0x00, 0x00, 0x5c,  // Definition size
     0x62, 0x68, 0x76, 0x72,  // Behavior id
     0x00, 0x01, 0x00, 0x01,  // Behavior version
-    0x00, 0x00, 0x00, 0xdb,  // Behavior offset
+    0x00, 0x00, 0x01, 0x0b,  // Behavior offset
     0x00, 0x00, 0x00, 0x56,  // Behavior size
     0x67, 0x65, 0x6f, 0x6d,  // Geometry id
     0x00, 0x01, 0x00, 0x01,  // Geometry version
-    0x00, 0x00, 0x01, 0x31,  // Geometry offset
+    0x00, 0x00, 0x01, 0x61,  // Geometry offset
     0x00, 0x00, 0x00, 0x04,  // Geometry size
     0x6d, 0x6c, 0x62, 0x68,  // Machine learned behavior id
     0x00, 0x01, 0x00, 0x00,  // Machine learned behavior version
-    0x00, 0x00, 0x01, 0x35,  // Machine learned behavior offset
-    0x00, 0x00, 0x00, 0x18  // Machine learned behavior size
-
+    0x00, 0x00, 0x01, 0x65,  // Machine learned behavior offset
+    0x00, 0x00, 0x00, 0x18,  // Machine learned behavior size
+    0x72, 0x62, 0x66, 0x62,  // RBF behavior id
+    0x00, 0x01, 0x00, 0x00,  // RBF behavior version
+    0x00, 0x00, 0x01, 0x7d,  // RBF behavior offset
+    0x00, 0x00, 0x00, 0x10,  // RBF behavior size
+    0x6a, 0x62, 0x6d, 0x64,  // Joint behavior metadata id
+    0x00, 0x01, 0x00, 0x00,  // Joint behavior metadata version
+    0x00, 0x00, 0x01, 0x8d,  // Joint behavior metadata offset
+    0x00, 0x00, 0x00, 0x04,  // Joint behavior metadata size
+    0x74, 0x77, 0x73, 0x77,  // Twist swing setups id
+    0x00, 0x01, 0x00, 0x00,  // Twist swing setups version
+    0x00, 0x00, 0x01, 0x91,  // Twist swing setups offset
+    0x00, 0x00, 0x00, 0x08  // Twist swing setups size
 };
 
 std::vector<char> RawV2xNewerWithUnknownDataIgnoredAndDNARewritten::getBytes() {
@@ -1363,7 +1408,21 @@ std::vector<char> RawV2xNewerWithUnknownDataIgnoredAndDNARewritten::getBytes() {
                  RawV2xNewerWithUnknownDataPreservedAndDNARewritten::machineLearnedBehavior,
                  RawV2xNewerWithUnknownDataPreservedAndDNARewritten::machineLearnedBehavior +
                  sizeof(RawV2xNewerWithUnknownDataPreservedAndDNARewritten::machineLearnedBehavior));
-
+    // RBF behavior
+    bytes.insert(bytes.end(),
+                 RawV2xNewerWithUnknownDataPreservedAndDNARewritten::rbfBehavior,
+                 RawV2xNewerWithUnknownDataPreservedAndDNARewritten::rbfBehavior +
+                 sizeof(RawV2xNewerWithUnknownDataPreservedAndDNARewritten::rbfBehavior));
+    // Joint behavior metadata
+    bytes.insert(bytes.end(),
+                 RawV2xNewerWithUnknownDataPreservedAndDNARewritten::jointBehaviorMetaData,
+                 RawV2xNewerWithUnknownDataPreservedAndDNARewritten::jointBehaviorMetaData +
+                 sizeof(RawV2xNewerWithUnknownDataPreservedAndDNARewritten::jointBehaviorMetaData));
+    // Twist swing behavior
+    bytes.insert(bytes.end(),
+                 RawV2xNewerWithUnknownDataPreservedAndDNARewritten::twistSwingBehavior,
+                 RawV2xNewerWithUnknownDataPreservedAndDNARewritten::twistSwingBehavior +
+                 sizeof(RawV2xNewerWithUnknownDataPreservedAndDNARewritten::twistSwingBehavior));
     return bytes;
 }
 
@@ -1393,6 +1452,10 @@ const unsigned char RawV22Empty::header[] = {
 
 std::vector<char> RawV22Empty::getBytes() {
     std::vector<char> bytes;
+    #if !defined(__clang__) && defined(__GNUC__) && (__GNUC__ >= 12)
+        #pragma GCC diagnostic push
+        #pragma GCC diagnostic ignored "-Wstringop-overflow"
+    #endif
     // Header
     bytes.insert(bytes.end(), header, header + sizeof(header));
     // Descriptor
@@ -1416,6 +1479,9 @@ std::vector<char> RawV22Empty::getBytes() {
                  sizeof(RawV2xNewerWithUnknownDataPreservedAndDNARewritten::geometry));
 
     return bytes;
+    #if !defined(__clang__) && defined(__GNUC__) && (__GNUC__ >= 12)
+        #pragma GCC diagnostic pop
+    #endif
 }
 
 const unsigned char RawV22WithUnknownDataFromNewer2x::header[] = {
@@ -1808,14 +1874,14 @@ const pma::Vector<std::uint16_t> DecodedV22::jointRowCount = {
 const std::uint16_t DecodedV22::jointColumnCount = 10u;
 const pma::Vector<pma::Matrix<std::uint16_t> > DecodedV22::jointVariableIndices = {
     {  // MaxLOD-0 - MinLOD-1
-        {2, 3, 5, 18, 20, 36, 38, 55, 56, 63, 45, 46, 71},
+        {2, 3, 5, 18, 20, 36, 38, 45, 46, 55, 56, 63, 71},
         {2, 3, 5, 18, 20, 55, 56}
     },
     {  // MaxLOD-1 - MinLOD-1
         {2, 3, 5, 18, 20, 37, 38}
     },
     {  // MaxLOD-0 - MinLOD-0
-        {2, 3, 5, 18, 20, 36, 38, 55, 56, 63, 45, 46, 71}
+        {2, 3, 5, 18, 20, 36, 38, 45, 46, 55, 56, 63, 71}
     }
 };
 const pma::Vector<pma::Matrix<std::uint16_t> > DecodedV22::jointGroupLODs = {
diff --git a/Engine/Plugins/Animation/RigLogic/Source/RigLogicLibTest/Private/dnatests/Fixturesv22.h b/Engine/Plugins/Animation/RigLogic/Source/RigLogicLibTest/Private/dnatests/Fixturesv22.h
--- a/Engine/Plugins/Animation/RigLogic/Source/RigLogicLibTest/Private/dnatests/Fixturesv22.h
+++ b/Engine/Plugins/Animation/RigLogic/Source/RigLogicLibTest/Private/dnatests/Fixturesv22.h
@@ -47,18 +47,21 @@ struct RawV2xNewer {
 };
 
 struct RawV2xNewerWithUnknownDataPreservedAndDNARewritten {
-    static const unsigned char header[187];
+    static const unsigned char header[235];
     static const unsigned char descriptor[36];
     static const unsigned char definition[92];
     static const unsigned char behavior[86];
     static const unsigned char geometry[4];
     static const unsigned char machineLearnedBehavior[24];
+    static const unsigned char rbfBehavior[16];
+    static const unsigned char jointBehaviorMetaData[4];
+    static const unsigned char twistSwingBehavior[8];
 
     static std::vector<char> getBytes();
 };
 
 struct RawV2xNewerWithUnknownDataIgnoredAndDNARewritten {
-    static const unsigned char header[91];
+    static const unsigned char header[139];
 
     static std::vector<char> getBytes();
 };
diff --git a/Engine/Plugins/Animation/RigLogic/Source/RigLogicLibTest/Private/dnatests/Fixturesv23.cpp b/Engine/Plugins/Animation/RigLogic/Source/RigLogicLibTest/Private/dnatests/Fixturesv23.cpp
--- a/Engine/Plugins/Animation/RigLogic/Source/RigLogicLibTest/Private/dnatests/Fixturesv23.cpp
+++ b/Engine/Plugins/Animation/RigLogic/Source/RigLogicLibTest/Private/dnatests/Fixturesv23.cpp
@@ -1230,6 +1230,10 @@ const unsigned char RawV23::machineLearnedBehavior[] = {
 };
 
 std::vector<char> RawV23::getBytes() {
+    #if !defined(__clang__) && defined(__GNUC__) && (__GNUC__ >= 12)
+        #pragma GCC diagnostic push
+        #pragma GCC diagnostic ignored "-Wstringop-overflow"
+    #endif
     std::vector<char> bytes;
     // Header
     bytes.insert(bytes.end(), header, header + sizeof(header));
@@ -1254,6 +1258,9 @@ std::vector<char> RawV23::getBytes() {
     // Machine learned behavior
     bytes.insert(bytes.end(), machineLearnedBehavior, machineLearnedBehavior + sizeof(machineLearnedBehavior));
     return bytes;
+    #if !defined(__clang__) && defined(__GNUC__) && (__GNUC__ >= 12)
+        #pragma GCC diagnostic pop
+    #endif
 }
 
 std::vector<char> RawV22DowngradedFromV23::getBytes() {
@@ -1564,14 +1571,14 @@ const pma::Vector<std::uint16_t> DecodedV23::jointRowCount = {
 const std::uint16_t DecodedV23::jointColumnCount = 10u;
 const pma::Vector<pma::Matrix<std::uint16_t> > DecodedV23::jointVariableIndices = {
     {  // MaxLOD-0 - MinLOD-1
-        {2, 3, 5, 18, 20, 36, 38, 55, 56, 63, 45, 46, 71},
+        {2, 3, 5, 18, 20, 36, 38, 45, 46, 55, 56, 63, 71},
         {2, 3, 5, 18, 20, 55, 56}
     },
     {  // MaxLOD-1 - MinLOD-1
         {2, 3, 5, 18, 20, 37, 38}
     },
     {  // MaxLOD-0 - MinLOD-0
-        {2, 3, 5, 18, 20, 36, 38, 55, 56, 63, 45, 46, 71}
+        {2, 3, 5, 18, 20, 36, 38, 45, 46, 55, 56, 63, 71}
     }
 };
 const pma::Vector<pma::Matrix<std::uint16_t> > DecodedV23::jointGroupLODs = {
