Update to new RigLogic implementing body rig support

diff --git a/Engine/Plugins/Animation/RigLogic/Source/RigLogicLibTest/Private/dnatests/Fixturesv24.h b/Engine/Plugins/Animation/RigLogic/Source/RigLogicLibTest/Private/dnatests/Fixturesv24.h
new file mode 100644
--- /dev/null
+++ b/Engine/Plugins/Animation/RigLogic/Source/RigLogicLibTest/Private/dnatests/Fixturesv24.h
@@ -0,0 +1,168 @@
+// Copyright Epic Games, Inc. All Rights Reserved.
+
+#pragma once
+
+#include "dnatests/Defs.h"
+
+#include "dna/DNA.h"
+#include "dna/Reader.h"
+
+#include <pma/TypeDefs.h>
+
+namespace dna {
+
+struct RawV24 {
+    static const unsigned char header[139];
+    static const unsigned char descriptor[87];
+    static const unsigned char definition[794];
+    static const unsigned char controls[2];
+    static const unsigned char conditionals[324];
+    static const unsigned char psds[204];
+    static const unsigned char joints[444];
+    static const unsigned char blendshapes[44];
+    static const unsigned char animatedmaps[8];
+    static const unsigned char geometry[1080];
+    static const unsigned char machineLearnedBehavior[762];
+    static const unsigned char rbfBehavior[327];
+    static const unsigned char jointBehaviorMetadata[58];
+    static const unsigned char twistSwingBehavior[200];
+
+    static std::vector<char> getBytes();
+};
+
+struct RawV23DowngradedFromV24 {
+    static std::vector<char> getBytes();
+};
+
+struct DecodedV24 {
+    using VectorOfCharStringMatrix = pma::Vector<pma::Matrix<pma::String<char> > >;
+    using StringPair = std::pair<pma::String<char>, pma::String<char> >;
+
+    // Descriptor
+    static const pma::String<char> name;
+    static const Archetype archetype;
+    static const Gender gender;
+    static const std::uint16_t age;
+
+    static const pma::Vector<StringPair> metadata;
+    static const TranslationUnit translationUnit;
+    static const RotationUnit rotationUnit;
+    static const CoordinateSystem coordinateSystem;
+    static const std::uint16_t lodCount[3ul];
+    static const std::uint16_t maxLODs[3ul];
+    static const pma::String<char> complexity;
+    static const pma::String<char> dbName;
+
+    // Definition
+    static const pma::Vector<pma::String<char> > guiControlNames;
+    static const pma::Vector<pma::String<char> > rawControlNames;
+    static const VectorOfCharStringMatrix jointNames;
+    static const VectorOfCharStringMatrix blendShapeNames;
+    static const VectorOfCharStringMatrix animatedMapNames;
+    static const VectorOfCharStringMatrix meshNames;
+    static const pma::Vector<pma::Matrix<std::uint16_t> > meshBlendShapeIndices;
+    static const pma::Matrix<std::uint16_t> jointHierarchy;
+    static const pma::Vector<pma::Matrix<Vector3> > neutralJointTranslations;
+    static const pma::Vector<pma::Matrix<Vector3> > neutralJointRotations;
+
+    // Behavior
+    static const std::uint16_t guiControlCount;
+    static const std::uint16_t rawControlCount;
+    static const std::uint16_t psdCount;
+    static const pma::Vector<pma::Matrix<std::uint16_t> > conditionalInputIndices;
+    static const pma::Vector<pma::Matrix<std::uint16_t> > conditionalOutputIndices;
+    static const pma::Vector<pma::Matrix<float> > conditionalFromValues;
+    static const pma::Vector<pma::Matrix<float> > conditionalToValues;
+    static const pma::Vector<pma::Matrix<float> > conditionalSlopeValues;
+    static const pma::Vector<pma::Matrix<float> > conditionalCutValues;
+    static const pma::Vector<std::uint16_t> psdRowIndices;
+    static const pma::Vector<std::uint16_t> psdColumnIndices;
+    static const pma::Vector<float> psdValues;
+    static const pma::Vector<std::uint16_t> jointRowCount;
+    static const std::uint16_t jointColumnCount;
+    static const pma::Vector<pma::Matrix<std::uint16_t> > jointVariableIndices;
+    static const pma::Vector<pma::Matrix<std::uint16_t> > jointGroupLODs;
+    static const pma::Matrix<pma::Matrix<std::uint16_t> > jointGroupInputIndices;
+    static const pma::Matrix<pma::Matrix<std::uint16_t> > jointGroupOutputIndices;
+    static const pma::Matrix<pma::Matrix<float> > jointGroupValues;
+    static const pma::Matrix<pma::Matrix<std::uint16_t> > jointGroupJointIndices;
+    static const pma::Matrix<std::uint16_t> blendShapeLODs;
+    static const pma::Vector<pma::Matrix<std::uint16_t> > blendShapeInputIndices;
+    static const pma::Vector<pma::Matrix<std::uint16_t> > blendShapeOutputIndices;
+    static const pma::Vector<std::uint16_t> animatedMapCount;
+    static const pma::Matrix<std::uint16_t> animatedMapLODs;
+
+    // Geometry
+    static const pma::Vector<std::uint32_t> meshCount;
+    static const pma::Vector<pma::Matrix<Vector3> > vertexPositions;
+    static const pma::Vector<pma::Matrix<TextureCoordinate> > vertexTextureCoordinates;
+    static const pma::Vector<pma::Matrix<Vector3> > vertexNormals;
+    static const pma::Vector<pma::Matrix<VertexLayout> > vertexLayouts;
+    static const pma::Matrix<pma::Matrix<std::uint32_t> > faces;
+    static const pma::Matrix<std::uint16_t> maxInfluencePerVertex;
+    static const pma::Matrix<pma::Matrix<float> > skinWeightsValues;
+    static const pma::Matrix<pma::Matrix<std::uint16_t> > skinWeightsJointIndices;
+    static const pma::Vector<pma::Matrix<std::uint16_t> > correctiveBlendShapeIndices;
+    static const pma::Matrix<pma::Matrix<Vector3> > correctiveBlendShapeDeltas;
+    static const pma::Matrix<pma::Matrix<std::uint32_t> > correctiveBlendShapeVertexIndices;
+
+    // Machine learned behavior
+    static const pma::Vector<pma::String<char> > mlControlNames;
+    static const VectorOfCharStringMatrix regionNames;
+    static const pma::Matrix<std::uint16_t> neuralNetworkIndicesPerLOD;
+    static const pma::Matrix<pma::Matrix<std::uint16_t> > neuralNetworkIndicesPerMeshRegion;
+    static const pma::Vector<pma::Matrix<std::uint16_t> > neuralNetworkInputIndices;
+    static const pma::Vector<pma::Matrix<std::uint16_t> > neuralNetworkOutputIndices;
+    static const pma::Matrix<std::uint16_t> neuralNetworkLayerCount;
+    static const pma::Vector<pma::Matrix<std::uint16_t> > neuralNetworkActivationFunction;
+    static const pma::Matrix<pma::Matrix<float> > neuralNetworkActivationFunctionParameters;
+    static const pma::Matrix<pma::Matrix<float> > neuralNetworkBiases;
+    static const pma::Matrix<pma::Matrix<float> > neuralNetworkWeights;
+
+    // RBF behavior
+    static const pma::Vector<pma::String<char> > poseNames;
+    static const pma::Vector<std::uint16_t> poseDistanceMethod;
+    static const pma::Vector<std::uint16_t> poseFunctionType;
+    static const pma::Vector<float> poseScale;
+    static const pma::Matrix<std::uint16_t> solverIndicesPerLOD;
+    static const pma::Vector<pma::String<char> > solverNames;
+    static const pma::Vector<float> solverRadius;
+    static const pma::Vector<float> solverWeightThreshold;
+    static const pma::Vector<std::uint16_t> solverType;
+    static const pma::Vector<std::uint16_t> solverAutomaticRadius;
+    static const pma::Vector<std::uint16_t> solverDistanceMethod;
+    static const pma::Vector<std::uint16_t> solverNormalizeMethod;
+    static const pma::Vector<std::uint16_t> solverFunctionType;
+    static const pma::Vector<std::uint16_t> solverTwistAxis;
+    static const pma::Matrix<std::uint16_t> solverRawControlIndices;
+    static const pma::Matrix<std::uint16_t> solverPoseIndices;
+    static const pma::Matrix<float> solverRawControlValues;
+
+    // Joint behavior metadata
+    static const pma::Matrix<TranslationRepresentation> jointTranslationRepresentation;
+    static const pma::Matrix<RotationRepresentation> jointRotationRepresentation;
+    static const pma::Matrix<ScaleRepresentation> jointScaleRepresentation;
+
+    // Twist swing setups
+    static const pma::Vector<pma::Matrix<std::uint16_t> > swingInputControlIndices;
+    static const pma::Vector<pma::Matrix<std::uint16_t> > swingOutputJointIndices;
+    static const pma::Vector<pma::Matrix<float> > swingBlendWeights;
+    static const pma::Matrix<TwistAxis> swingTwistAxes;
+    static const pma::Vector<pma::Matrix<std::uint16_t> > twistInputControlIndices;
+    static const pma::Vector<pma::Matrix<std::uint16_t> > twistOutputJointIndices;
+    static const pma::Vector<pma::Matrix<float> > twistBlendWeights;
+    static const pma::Matrix<TwistAxis> twistTwistAxes;
+
+    static std::size_t lodConstraintToIndex(std::uint16_t maxLOD, std::uint16_t minLOD);
+    static RawJoints getJoints(std::uint16_t currentMaxLOD, std::uint16_t currentMinLOD, pma::MemoryResource* memRes);
+    static RawBlendShapeChannels getBlendShapes(std::uint16_t currentMaxLOD,
+                                                std::uint16_t currentMinLOD,
+                                                pma::MemoryResource* memRes);
+    static RawConditionalTable getConditionals(std::uint16_t currentMaxLOD,
+                                               std::uint16_t currentMinLOD,
+                                               pma::MemoryResource* memRes);
+    static RawAnimatedMaps getAnimatedMaps(std::uint16_t currentMaxLOD, std::uint16_t currentMinLOD, pma::MemoryResource* memRes);
+
+};
+
+}  // namespace dna
diff --git a/Engine/Plugins/Animation/RigLogic/Source/RigLogicLibTest/Private/dnatests/TestFilteredBinaryInputArchive.cpp b/Engine/Plugins/Animation/RigLogic/Source/RigLogicLibTest/Private/dnatests/TestFilteredBinaryInputArchive.cpp
--- a/Engine/Plugins/Animation/RigLogic/Source/RigLogicLibTest/Private/dnatests/TestFilteredBinaryInputArchive.cpp
+++ b/Engine/Plugins/Animation/RigLogic/Source/RigLogicLibTest/Private/dnatests/TestFilteredBinaryInputArchive.cpp
@@ -62,7 +62,7 @@ static void verifyDefinition(const dna::DNA& result, const DNAParameters& dnaPar
                        TDecodedDNA::jointHierarchy[index],
                        TDecodedDNA::jointHierarchy[index].size());
 
-    for (std::uint16_t lod = 0u; lod < result.descriptor.lodCount; ++lod) {
+    for (std::uint16_t lod = {}; lod < result.descriptor.lodCount; ++lod) {
         auto jointIndices = result.definition.lodJointMapping.getIndices(lod);
         ASSERT_EQ(jointIndices.size(), TDecodedDNA::jointNames[index][lod].size());
         for (std::size_t i = 0ul; i < jointIndices.size(); ++i) {
diff --git a/Engine/Plugins/Animation/RigLogic/Source/RigLogicLibTest/Private/dnatests/TestLODMapping.cpp b/Engine/Plugins/Animation/RigLogic/Source/RigLogicLibTest/Private/dnatests/TestLODMapping.cpp
--- a/Engine/Plugins/Animation/RigLogic/Source/RigLogicLibTest/Private/dnatests/TestLODMapping.cpp
+++ b/Engine/Plugins/Animation/RigLogic/Source/RigLogicLibTest/Private/dnatests/TestLODMapping.cpp
@@ -36,7 +36,7 @@ class LODMappingTest : public ::testing::Test {
 }  // namespace
 
 #define ASSERT_LODMAPPING_EQ(result, expected)   \
-    for (std::uint16_t lod = 0u; lod < result.getLODCount(); ++lod) {               \
+    for (std::uint16_t lod = {}; lod < result.getLODCount(); ++lod) {               \
         dna::ConstArrayView<std::uint16_t> resultIndices = result.getIndices(lod);  \
         dna::ConstArrayView<std::uint16_t> expectedIndices{expected[lod]};          \
         ASSERT_EQ(resultIndices, expectedIndices);                                  \
diff --git a/Engine/Plugins/Animation/RigLogic/Source/RigLogicLibTest/Private/dnatests/TestStreamReadWriteIntegration.cpp b/Engine/Plugins/Animation/RigLogic/Source/RigLogicLibTest/Private/dnatests/TestStreamReadWriteIntegration.cpp
--- a/Engine/Plugins/Animation/RigLogic/Source/RigLogicLibTest/Private/dnatests/TestStreamReadWriteIntegration.cpp
+++ b/Engine/Plugins/Animation/RigLogic/Source/RigLogicLibTest/Private/dnatests/TestStreamReadWriteIntegration.cpp
@@ -9,9 +9,9 @@
 #include "dnatests/Fixturesv21.h"
 #include "dnatests/Fixturesv22.h"
 #include "dnatests/Fixturesv23.h"
+#include "dnatests/Fixturesv24.h"
 
 #include "dna/DataLayer.h"
-#include "dna/DNA.h"
 #include "dna/BinaryStreamReader.h"
 #include "dna/BinaryStreamWriter.h"
 #ifdef DNA_BUILD_WITH_JSON_SUPPORT
@@ -19,6 +19,10 @@
     #include "dna/JSONStreamWriter.h"
 #endif  // DNA_BUILD_WITH_JSON_SUPPORT
 
+#ifdef _MSC_VER
+    #pragma warning(disable : 4503)
+#endif
+
 namespace dna {
 
 template<class TAPICopyParameters>
@@ -33,7 +37,7 @@ static void verifyDescriptor(DescriptorReader* reader) {
 
     const auto metaDataCount = reader->getMetaDataCount();
     ASSERT_EQ(metaDataCount, 2u);
-    for (std::uint32_t i = 0u; i < metaDataCount; ++i) {
+    for (std::uint32_t i = {}; i < metaDataCount; ++i) {
         const auto key = reader->getMetaDataKey(i);
         const auto value = reader->getMetaDataValue(key);
         ASSERT_EQ(key, StringView{DecodedDNA::metadata[i].first});
@@ -61,13 +65,13 @@ static void verifyDefinition(DefinitionReader* reader) {
 
     const auto guiControlCount = reader->getGUIControlCount();
     ASSERT_EQ(guiControlCount, DecodedDNA::guiControlNames.size());
-    for (std::uint16_t i = 0u; i < guiControlCount; ++i) {
+    for (std::uint16_t i = {}; i < guiControlCount; ++i) {
         ASSERT_EQ(reader->getGUIControlName(i), StringView{DecodedDNA::guiControlNames[i]});
     }
 
     const auto rawControlCount = reader->getRawControlCount();
     ASSERT_EQ(rawControlCount, DecodedDNA::rawControlNames.size());
-    for (std::uint16_t i = 0u; i < rawControlCount; ++i) {
+    for (std::uint16_t i = {}; i < rawControlCount; ++i) {
         ASSERT_EQ(reader->getRawControlName(i), StringView{DecodedDNA::rawControlNames[i]});
     }
 
@@ -79,7 +83,7 @@ static void verifyDefinition(DefinitionReader* reader) {
         ASSERT_EQ(reader->getJointName(jointIndices[i]), StringView{expectedJointNames[i]});
     }
 
-    for (std::uint16_t i = 0u; i < reader->getJointCount(); ++i) {
+    for (std::uint16_t i = {}; i < reader->getJointCount(); ++i) {
         ASSERT_EQ(reader->getJointParentIndex(i), DecodedDNA::jointHierarchy[index][i]);
     }
 
@@ -183,7 +187,7 @@ static void verifyBehavior(BehaviorReader* reader) {
     const auto jointGroupCount = reader->getJointGroupCount();
     ASSERT_EQ(jointGroupCount, DecodedDNA::jointGroupLODs.size());
 
-    for (std::uint16_t i = 0u; i < jointGroupCount; ++i) {
+    for (std::uint16_t i = {}; i < jointGroupCount; ++i) {
         const auto& expectedLODs = DecodedDNA::jointGroupLODs[i][index];
         ASSERT_EQ(reader->getJointGroupLODs(i), ConstArrayView<std::uint16_t>{expectedLODs});
 
@@ -281,17 +285,17 @@ static void verifyGeometry(GeometryReader* reader) {
 
     const auto meshCount = reader->getMeshCount();
     ASSERT_EQ(meshCount, DecodedDNA::meshCount[index]);
-    for (std::uint16_t meshIndex = 0u; meshIndex < meshCount; ++meshIndex) {
+    for (std::uint16_t meshIndex = {}; meshIndex < meshCount; ++meshIndex) {
         const auto vertexPositionCount = reader->getVertexPositionCount(meshIndex);
         ASSERT_EQ(vertexPositionCount, DecodedDNA::vertexPositions[index][meshIndex].size());
-        for (std::uint32_t vertexIndex = 0u; vertexIndex < vertexPositionCount; ++vertexIndex) {
+        for (std::uint32_t vertexIndex = {}; vertexIndex < vertexPositionCount; ++vertexIndex) {
             ASSERT_EQ(reader->getVertexPosition(meshIndex, vertexIndex),
                       DecodedDNA::vertexPositions[index][meshIndex][vertexIndex]);
         }
 
         const auto vertexTextureCoordinateCount = reader->getVertexTextureCoordinateCount(meshIndex);
         ASSERT_EQ(vertexTextureCoordinateCount, DecodedDNA::vertexTextureCoordinates[index][meshIndex].size());
-        for (std::uint32_t texCoordIndex = 0u; texCoordIndex < vertexTextureCoordinateCount; ++texCoordIndex) {
+        for (std::uint32_t texCoordIndex = {}; texCoordIndex < vertexTextureCoordinateCount; ++texCoordIndex) {
             const auto& textureCoordinate = reader->getVertexTextureCoordinate(meshIndex, texCoordIndex);
             const auto& expectedTextureCoordinate =
                 DecodedDNA::vertexTextureCoordinates[index][meshIndex][texCoordIndex];
@@ -301,14 +305,14 @@ static void verifyGeometry(GeometryReader* reader) {
 
         const auto vertexNormalCount = reader->getVertexNormalCount(meshIndex);
         ASSERT_EQ(vertexNormalCount, DecodedDNA::vertexNormals[index][meshIndex].size());
-        for (std::uint32_t normalIndex = 0u; normalIndex < vertexNormalCount; ++normalIndex) {
+        for (std::uint32_t normalIndex = {}; normalIndex < vertexNormalCount; ++normalIndex) {
             ASSERT_EQ(reader->getVertexNormal(meshIndex, normalIndex),
                       DecodedDNA::vertexNormals[index][meshIndex][normalIndex]);
         }
 
         const auto vertexLayoutCount = reader->getVertexLayoutCount(meshIndex);
         ASSERT_EQ(vertexLayoutCount, DecodedDNA::vertexLayouts[index][meshIndex].size());
-        for (std::uint32_t layoutIndex = 0u; layoutIndex < vertexLayoutCount; ++layoutIndex) {
+        for (std::uint32_t layoutIndex = {}; layoutIndex < vertexLayoutCount; ++layoutIndex) {
             const auto& layout = reader->getVertexLayout(meshIndex, layoutIndex);
             const auto& expectedLayout = DecodedDNA::vertexLayouts[index][meshIndex][layoutIndex];
             ASSERT_EQ(layout.position, expectedLayout.position);
@@ -318,7 +322,7 @@ static void verifyGeometry(GeometryReader* reader) {
 
         const auto faceCount = reader->getFaceCount(meshIndex);
         ASSERT_EQ(faceCount, DecodedDNA::faces[index][meshIndex].size());
-        for (std::uint32_t faceIndex = 0u; faceIndex < faceCount; ++faceIndex) {
+        for (std::uint32_t faceIndex = {}; faceIndex < faceCount; ++faceIndex) {
             ASSERT_EQ(reader->getFaceVertexLayoutIndices(meshIndex, faceIndex),
                       ConstArrayView<std::uint32_t>{DecodedDNA::faces[index][meshIndex][faceIndex]});
         }
@@ -327,7 +331,7 @@ static void verifyGeometry(GeometryReader* reader) {
                   DecodedDNA::maxInfluencePerVertex[index][meshIndex]);
 
         ASSERT_EQ(reader->getSkinWeightsCount(meshIndex), DecodedDNA::skinWeightsValues[index][meshIndex].size());
-        for (std::uint32_t vertexIndex = 0u; vertexIndex < vertexPositionCount; ++vertexIndex) {
+        for (std::uint32_t vertexIndex = {}; vertexIndex < vertexPositionCount; ++vertexIndex) {
             const auto skinWeights = reader->getSkinWeightsValues(meshIndex, vertexIndex);
             const auto& expectedSkinWeights = DecodedDNA::skinWeightsValues[index][meshIndex][vertexIndex];
             ASSERT_EQ(skinWeights, ConstArrayView<float>{expectedSkinWeights});
@@ -339,7 +343,7 @@ static void verifyGeometry(GeometryReader* reader) {
 
         const auto blendShapeCount = reader->getBlendShapeTargetCount(meshIndex);
         ASSERT_EQ(blendShapeCount, DecodedDNA::correctiveBlendShapeDeltas[index][meshIndex].size());
-        for (std::uint16_t blendShapeTargetIndex = 0u; blendShapeTargetIndex < blendShapeCount; ++blendShapeTargetIndex) {
+        for (std::uint16_t blendShapeTargetIndex = {}; blendShapeTargetIndex < blendShapeCount; ++blendShapeTargetIndex) {
             const auto channelIndex = reader->getBlendShapeChannelIndex(meshIndex, blendShapeTargetIndex);
             ASSERT_EQ(channelIndex, DecodedDNA::correctiveBlendShapeIndices[index][meshIndex][blendShapeTargetIndex]);
 
@@ -347,7 +351,7 @@ static void verifyGeometry(GeometryReader* reader) {
             ASSERT_EQ(deltaCount,
                       DecodedDNA::correctiveBlendShapeDeltas[index][meshIndex][blendShapeTargetIndex].size());
 
-            for (std::uint32_t deltaIndex = 0u; deltaIndex < deltaCount; ++deltaIndex) {
+            for (std::uint32_t deltaIndex = {}; deltaIndex < deltaCount; ++deltaIndex) {
                 const auto& delta = reader->getBlendShapeTargetDelta(meshIndex, blendShapeTargetIndex, deltaIndex);
                 const auto& expectedDelta =
                     DecodedDNA::correctiveBlendShapeDeltas[index][meshIndex][blendShapeTargetIndex][deltaIndex];
@@ -369,7 +373,7 @@ static void verifyMachineLearnedBehavior(MachineLearnedBehaviorReader* reader) {
 
     const auto mlControlCount = reader->getMLControlCount();
     ASSERT_EQ(mlControlCount, DecodedDNA::mlControlNames.size());
-    for (std::uint16_t i = 0u; i < mlControlCount; ++i) {
+    for (std::uint16_t i = {}; i < mlControlCount; ++i) {
         ASSERT_EQ(reader->getMLControlName(i), StringView{DecodedDNA::mlControlNames[i]});
     }
 
@@ -414,6 +418,122 @@ static void verifyMachineLearnedBehavior(MachineLearnedBehaviorReader* reader) {
     }
 }
 
+template<class TAPICopyParameters>
+static void verifyRBFBehavior(RBFBehaviorReader* reader) {
+    using DecodedDNA = typename TAPICopyParameters::DecodedData;
+    const auto index = DecodedDNA::lodConstraintToIndex(TAPICopyParameters::maxLOD(), TAPICopyParameters::minLOD());
+
+    const std::uint16_t solverCount = reader->getRBFSolverCount();
+    ASSERT_EQ(solverCount, DecodedDNA::solverIndicesPerLOD[index].size());
+
+    const auto poseCount = reader->getRBFPoseCount();
+    for (std::uint16_t pi = {}; pi < poseCount; ++pi) {
+        ASSERT_EQ(reader->getRBFPoseName(pi), StringView{DecodedDNA::poseNames[pi]});
+        ASSERT_EQ(reader->getRBFPoseScale(pi), DecodedDNA::poseScale[pi]);
+    }
+    for (std::uint16_t si = {}; si < solverCount; ++si) {
+        std::uint16_t esi = DecodedDNA::solverIndicesPerLOD[index][si];
+        ASSERT_EQ(reader->getRBFSolverName(si), StringView{DecodedDNA::solverNames[esi]});
+        ASSERT_EQ(reader->getRBFSolverRawControlIndices(si),
+                  ConstArrayView<std::uint16_t>{DecodedDNA::solverRawControlIndices[esi]});
+        ASSERT_EQ(reader->getRBFSolverType(si),
+                  static_cast<RBFSolverType>(DecodedDNA::solverType[esi]));
+        ASSERT_EQ(reader->getRBFSolverAutomaticRadius(si),
+                  static_cast<AutomaticRadius>(DecodedDNA::solverAutomaticRadius[esi]));
+        ASSERT_EQ(reader->getRBFSolverDistanceMethod(si),
+                  static_cast<RBFDistanceMethod>(DecodedDNA::solverDistanceMethod[esi]));
+        ASSERT_EQ(reader->getRBFSolverNormalizeMethod(si),
+                  static_cast<RBFNormalizeMethod>(DecodedDNA::solverNormalizeMethod[esi]));
+        ASSERT_EQ(reader->getRBFSolverFunctionType(si),
+                  static_cast<RBFFunctionType>(DecodedDNA::solverFunctionType[esi]));
+        ASSERT_EQ(reader->getRBFSolverTwistAxis(si),
+                  static_cast<TwistAxis>(DecodedDNA::solverTwistAxis[esi]));
+        ASSERT_EQ(reader->getRBFSolverRadius(si), DecodedDNA::solverRadius[esi]);
+        ASSERT_EQ(reader->getRBFSolverWeightThreshold(si), DecodedDNA::solverWeightThreshold[esi]);
+        auto rawControlIndices = reader->getRBFSolverRawControlIndices(si);
+        const auto& expectedRawControlIndices = DecodedDNA::solverRawControlIndices[esi];
+        ASSERT_EQ(rawControlIndices.size(), expectedRawControlIndices.size());
+        ASSERT_ELEMENTS_EQ(rawControlIndices, expectedRawControlIndices, rawControlIndices.size());
+
+        auto solverPoseIndices = reader->getRBFSolverPoseIndices(si);
+        const auto& expectedSolverPoseIndices = DecodedDNA::solverPoseIndices[esi];
+        ASSERT_EQ(solverPoseIndices.size(), expectedSolverPoseIndices.size());
+        ASSERT_ELEMENTS_EQ(solverPoseIndices, expectedSolverPoseIndices, solverPoseIndices.size());
+
+        auto solverRawControlValues = reader->getRBFSolverRawControlValues(si);
+        const auto& expectedSolverRawControlValues = DecodedDNA::solverRawControlValues[esi];
+        ASSERT_EQ(solverRawControlValues.size(), expectedSolverRawControlValues.size());
+        ASSERT_ELEMENTS_EQ(solverRawControlValues, expectedSolverRawControlValues, solverRawControlValues.size());
+    }
+
+}
+
+template<class TAPICopyParameters>
+static void verifyJointBehaviorMetadata(JointBehaviorMetadataReader* reader) {
+    using DecodedDNA = typename TAPICopyParameters::DecodedData;
+    const auto index = DecodedDNA::lodConstraintToIndex(TAPICopyParameters::maxLOD(), TAPICopyParameters::minLOD());
+
+    for (const auto ji :  reader->getJointIndicesForLOD(TAPICopyParameters::currentLOD())) {
+        ASSERT_EQ(reader->getJointTranslationRepresentation(ji), DecodedDNA::jointTranslationRepresentation[index][ji]);
+        ASSERT_EQ(reader->getJointRotationRepresentation(ji), DecodedDNA::jointRotationRepresentation[index][ji]);
+        ASSERT_EQ(reader->getJointScaleRepresentation(ji), DecodedDNA::jointScaleRepresentation[index][ji]);
+    }
+}
+
+template<class TAPICopyParameters>
+static void verifyTwistSwingBehavior(TwistSwingBehaviorReader* reader) {
+    using DecodedDNA = typename TAPICopyParameters::DecodedData;
+    const auto index = DecodedDNA::lodConstraintToIndex(TAPICopyParameters::maxLOD(), TAPICopyParameters::minLOD());
+
+    const auto expectedTwistCount = static_cast<std::uint16_t>(DecodedDNA::twistBlendWeights[index].size());
+    const auto twistCount = reader->getTwistCount();
+    ASSERT_EQ(twistCount, expectedTwistCount);
+    for (std::uint16_t ti = {}; ti < twistCount; ++ti) {
+        const auto twistInputIndices = reader->getTwistInputControlIndices(ti);
+        const auto expectedTwistInputIndices = DecodedDNA::twistInputControlIndices[index][ti];
+        ASSERT_EQ(twistInputIndices.size(), expectedTwistInputIndices.size());
+        ASSERT_ELEMENTS_EQ(twistInputIndices, expectedTwistInputIndices, expectedTwistInputIndices.size());
+
+        const auto twistOutputIndices = reader->getTwistOutputJointIndices(ti);
+        const auto expectedTwistOutputIndices = DecodedDNA::twistOutputJointIndices[index][ti];
+        ASSERT_EQ(twistOutputIndices.size(), expectedTwistOutputIndices.size());
+        ASSERT_ELEMENTS_EQ(twistOutputIndices, expectedTwistOutputIndices, expectedTwistOutputIndices.size());
+
+        const auto twistBlendWeights = reader->getTwistBlendWeights(ti);
+        const auto expectedTwistBlendWeights = DecodedDNA::twistBlendWeights[index][ti];
+        ASSERT_EQ(twistBlendWeights.size(), expectedTwistBlendWeights.size());
+        ASSERT_ELEMENTS_EQ(twistBlendWeights, expectedTwistBlendWeights, twistBlendWeights.size());
+
+        const auto twistAxis = reader->getTwistSetupTwistAxis(ti);
+        const auto expectedTwistAxis = DecodedDNA::twistTwistAxes[index][ti];
+        ASSERT_EQ(twistAxis, expectedTwistAxis);
+    }
+
+    const auto expectedSwingCount = static_cast<std::uint16_t>(DecodedDNA::swingBlendWeights[index].size());
+    const auto swingCount = reader->getSwingCount();
+    ASSERT_EQ(swingCount, expectedSwingCount);
+    for (std::uint16_t si = {}; si < swingCount; ++si) {
+        const auto swingInputIndices = reader->getSwingInputControlIndices(si);
+        const auto expectedSwingInputIndices = DecodedDNA::swingInputControlIndices[index][si];
+        ASSERT_EQ(swingInputIndices.size(), expectedSwingInputIndices.size());
+        ASSERT_ELEMENTS_EQ(swingInputIndices, expectedSwingInputIndices, expectedSwingInputIndices.size());
+
+        const auto swingOutputIndices = reader->getSwingOutputJointIndices(si);
+        const auto expectedSwingOutputIndices = DecodedDNA::swingOutputJointIndices[index][si];
+        ASSERT_EQ(swingOutputIndices.size(), expectedSwingOutputIndices.size());
+        ASSERT_ELEMENTS_EQ(swingOutputIndices, expectedSwingOutputIndices, expectedSwingOutputIndices.size());
+
+        const auto swingBlendWeights = reader->getSwingBlendWeights(si);
+        const auto expectedSwingBlendWeights = DecodedDNA::swingBlendWeights[index][si];
+        ASSERT_EQ(swingBlendWeights.size(), expectedSwingBlendWeights.size());
+        ASSERT_ELEMENTS_EQ(swingBlendWeights, expectedSwingBlendWeights, expectedSwingBlendWeights.size());
+
+        const auto twistAxis = reader->getSwingSetupTwistAxis(si);
+        const auto expectedTwistAxis = DecodedDNA::swingTwistAxes[index][si];
+        ASSERT_EQ(twistAxis, expectedTwistAxis);
+    }
+}
+
 template<class TAPICopyParameters>
 struct ReaderDataVerifier {
 
@@ -441,6 +561,23 @@ struct ReaderDataVerifier<APICopyParameters<Reader, Writer, RawV23, DecodedV23,
 
 };
 
+template<class Reader, class Writer, std::uint16_t MaxLOD, std::uint16_t MinLOD, std::uint16_t CurrentLOD>
+struct ReaderDataVerifier<APICopyParameters<Reader, Writer, RawV24, DecodedV24, MaxLOD, MinLOD, CurrentLOD> > {
+
+    static void assertHasAllData(Reader* reader) {
+        using TAPICopyParameters = APICopyParameters<Reader, Writer, RawV24, DecodedV24, MaxLOD, MinLOD, CurrentLOD>;
+        verifyDescriptor<TAPICopyParameters>(reader);
+        verifyDefinition<TAPICopyParameters>(reader);
+        verifyBehavior<TAPICopyParameters>(reader);
+        verifyGeometry<TAPICopyParameters>(reader);
+        verifyMachineLearnedBehavior<TAPICopyParameters>(reader);
+        verifyRBFBehavior<TAPICopyParameters>(reader);
+        verifyJointBehaviorMetadata<TAPICopyParameters>(reader);
+        verifyTwistSwingBehavior<TAPICopyParameters>(reader);
+    }
+
+};
+
 using TAPICopyTestParameters = ::testing::Types<
     APICopyParameters<dna::BinaryStreamReader, dna::BinaryStreamWriter, RawV21, DecodedV21, 0u, 1u, 0u>
     , APICopyParameters<dna::BinaryStreamReader, dna::BinaryStreamWriter, RawV21, DecodedV21, 0u, 1u, 1u>
@@ -454,10 +591,15 @@ using TAPICopyTestParameters = ::testing::Types<
     , APICopyParameters<dna::BinaryStreamReader, dna::BinaryStreamWriter, RawV23, DecodedV23, 0u, 1u, 1u>
     , APICopyParameters<dna::BinaryStreamReader, dna::BinaryStreamWriter, RawV23, DecodedV23, 0u, 0u, 0u>
     , APICopyParameters<dna::BinaryStreamReader, dna::BinaryStreamWriter, RawV23, DecodedV23, 1u, 1u, 0u>
+    , APICopyParameters<dna::BinaryStreamReader, dna::BinaryStreamWriter, RawV24, DecodedV24, 0u, 1u, 0u>
+    , APICopyParameters<dna::BinaryStreamReader, dna::BinaryStreamWriter, RawV24, DecodedV24, 0u, 1u, 1u>
+    , APICopyParameters<dna::BinaryStreamReader, dna::BinaryStreamWriter, RawV24, DecodedV24, 0u, 0u, 0u>
+    , APICopyParameters<dna::BinaryStreamReader, dna::BinaryStreamWriter, RawV24, DecodedV24, 1u, 1u, 0u>
     #ifdef DNA_BUILD_WITH_JSON_SUPPORT
         , APICopyParameters<dna::JSONStreamReader, dna::JSONStreamWriter, RawV21, DecodedV21, 0u, 1u, 0u>
         , APICopyParameters<dna::JSONStreamReader, dna::JSONStreamWriter, RawV22, DecodedV22, 0u, 1u, 0u>
         , APICopyParameters<dna::JSONStreamReader, dna::JSONStreamWriter, RawV23, DecodedV23, 0u, 1u, 0u>
+        , APICopyParameters<dna::JSONStreamReader, dna::JSONStreamWriter, RawV24, DecodedV24, 0u, 1u, 0u>
     #endif  // DNA_BUILD_WITH_JSON_SUPPORT
     >;
 TYPED_TEST_SUITE(StreamReadWriteAPICopyIntegrationTest, TAPICopyTestParameters, );
@@ -535,6 +677,8 @@ using TRawCopyTestParameters = ::testing::Types<
                       static_cast<std::uint16_t>(-1)>,
     RawCopyParameters<RawV23, RawV23, UnknownLayerPolicy::Preserve, 2, 3>,
     RawCopyParameters<RawV23, RawV23, UnknownLayerPolicy::Ignore, 2, 3>,
+    RawCopyParameters<RawV24, RawV24, UnknownLayerPolicy::Preserve, 2, 4>,
+    RawCopyParameters<RawV24, RawV24, UnknownLayerPolicy::Ignore, 2, 4>,
     // File format conversion tests
     RawCopyParameters<RawV21, RawV22WithUnknownDataIgnoredAndDNARewritten, UnknownLayerPolicy::Preserve, 2, 2>,
     RawCopyParameters<RawV21, RawV22WithUnknownDataIgnoredAndDNARewritten, UnknownLayerPolicy::Ignore, 2, 2>,
@@ -545,7 +689,9 @@ using TRawCopyTestParameters = ::testing::Types<
     RawCopyParameters<RawV22Empty, RawV22Empty, UnknownLayerPolicy::Preserve, 2, 2>,
     RawCopyParameters<RawV22Empty, RawV22Empty, UnknownLayerPolicy::Ignore, 2, 2>,
     RawCopyParameters<RawV23, RawV22DowngradedFromV23, UnknownLayerPolicy::Preserve, 2, 2>,
-    RawCopyParameters<RawV23, RawV22WithUnknownDataIgnoredAndDNARewritten, UnknownLayerPolicy::Ignore, 2, 2>
+    RawCopyParameters<RawV23, RawV22WithUnknownDataIgnoredAndDNARewritten, UnknownLayerPolicy::Ignore, 2, 2>,
+    RawCopyParameters<RawV24, RawV23DowngradedFromV24, UnknownLayerPolicy::Preserve, 2, 3>,
+    RawCopyParameters<RawV24, RawV23, UnknownLayerPolicy::Ignore, 2, 3>
     >;
 TYPED_TEST_SUITE(StreamReadWriteRawCopyIntegrationTest, TRawCopyTestParameters, );
 
@@ -624,6 +770,7 @@ using TReadWriteMultipleParameters = ::testing::Types<
     ReadWriteMultipleParameters<RawV21>,
     ReadWriteMultipleParameters<RawV22>,
     ReadWriteMultipleParameters<RawV23>,
+    ReadWriteMultipleParameters<RawV24>,
     ReadWriteMultipleParameters<RawV22Empty>,
     ReadWriteMultipleParameters<RawV22WithUnknownDataIgnoredAndDNARewritten>,
     ReadWriteMultipleParameters<RawV2xNewerWithUnknownDataIgnoredAndDNARewritten>,
diff --git a/Engine/Plugins/Animation/RigLogic/Source/RigLogicLibTest/Private/pmatests/TestPolyAllocIntegration.cpp b/Engine/Plugins/Animation/RigLogic/Source/RigLogicLibTest/Private/pmatests/TestPolyAllocIntegration.cpp
--- a/Engine/Plugins/Animation/RigLogic/Source/RigLogicLibTest/Private/pmatests/TestPolyAllocIntegration.cpp
+++ b/Engine/Plugins/Animation/RigLogic/Source/RigLogicLibTest/Private/pmatests/TestPolyAllocIntegration.cpp
@@ -4,12 +4,19 @@
 
 #include "pma/TypeDefs.h"
 
+// Old clang compiler in combination with new libstdc++
+#if defined(__clang__) && (__clang_major__ < 9) && defined(_GLIBCXX_RELEASE) && (_GLIBCXX_RELEASE >= 10)
+    #define PMA_OLD_CLANG_NEW_LIBSTDCPP
+#endif
+
 TEST(PolyAllocIntegrationTest, InstantiateTypes) {
     pma::String<char> str;
     pma::Vector<int> vec;
     pma::Matrix<int> mat;
+    #ifndef PMA_OLD_CLANG_NEW_LIBSTDCPP
     pma::Set<int> set;
     pma::Map<int, int> map;
+    #endif  // PMA_OLD_CLANG_NEW_LIBSTDCPP
     pma::UnorderedSet<int> uset;
     pma::UnorderedMap<int, int> umap;
     ASSERT_TRUE(true);
diff --git a/Engine/Plugins/Animation/RigLogic/Source/RigLogicLibTest/Private/rltests/animatedmaps/TestAnimatedMaps.cpp b/Engine/Plugins/Animation/RigLogic/Source/RigLogicLibTest/Private/rltests/animatedmaps/TestAnimatedMaps.cpp
--- a/Engine/Plugins/Animation/RigLogic/Source/RigLogicLibTest/Private/rltests/animatedmaps/TestAnimatedMaps.cpp
+++ b/Engine/Plugins/Animation/RigLogic/Source/RigLogicLibTest/Private/rltests/animatedmaps/TestAnimatedMaps.cpp
@@ -31,6 +31,7 @@ TEST_P(AnimatedMapsTest, LODLimitsCondTableSize) {
     auto inputInstanceFactory = ControlsFactory::getInstanceFactory(0,
                                                                     static_cast<std::uint16_t>(conditionalTableInputs.size()),
                                                                     0,
+                                                                    0,
                                                                     0);
     auto outputInstanceFactory = AnimatedMapsFactory::getInstanceFactory(conditionals.getOutputCount());
     rl4::AnimatedMapsImpl animatedMaps{std::move(lods), std::move(conditionals), outputInstanceFactory};
diff --git a/Engine/Plugins/Animation/RigLogic/Source/RigLogicLibTest/Private/rltests/blendshapes/TestBlendShapes.cpp b/Engine/Plugins/Animation/RigLogic/Source/RigLogicLibTest/Private/rltests/blendshapes/TestBlendShapes.cpp
--- a/Engine/Plugins/Animation/RigLogic/Source/RigLogicLibTest/Private/rltests/blendshapes/TestBlendShapes.cpp
+++ b/Engine/Plugins/Animation/RigLogic/Source/RigLogicLibTest/Private/rltests/blendshapes/TestBlendShapes.cpp
@@ -32,7 +32,11 @@ TEST_P(BlendShapesTest, Calculate) {
 
     pma::AlignedMemoryResource amr;
 
-    auto inputInstanceFactory = ControlsFactory::getInstanceFactory(0, static_cast<std::uint16_t>(blendShapeInputs.size()), 0, 0);
+    auto inputInstanceFactory = ControlsFactory::getInstanceFactory(0,
+                                                                    static_cast<std::uint16_t>(blendShapeInputs.size()),
+                                                                    0,
+                                                                    0,
+                                                                    0);
     auto inputInstance = inputInstanceFactory(&amr);
     auto inputBuffer = inputInstance->getInputBuffer();
     std::copy(blendShapeInputs.begin(), blendShapeInputs.end(), inputBuffer.begin());
diff --git a/Engine/Plugins/Animation/RigLogic/Source/RigLogicLibTest/Private/rltests/controls/ControlFixtures.cpp b/Engine/Plugins/Animation/RigLogic/Source/RigLogicLibTest/Private/rltests/controls/ControlFixtures.cpp
--- a/Engine/Plugins/Animation/RigLogic/Source/RigLogicLibTest/Private/rltests/controls/ControlFixtures.cpp
+++ b/Engine/Plugins/Animation/RigLogic/Source/RigLogicLibTest/Private/rltests/controls/ControlFixtures.cpp
@@ -9,13 +9,15 @@
 rl4::ControlsInputInstance::Factory ControlsFactory::getInstanceFactory(std::uint16_t guiControlCount,
                                                                         std::uint16_t rawControlCount,
                                                                         std::uint16_t psdControlCount,
-                                                                        std::uint16_t mlControlCount) {
+                                                                        std::uint16_t mlControlCount,
+                                                                        std::uint16_t rbfControlCount) {
     return [ = ](rl4::MemoryResource* memRes) {
                return pma::UniqueInstance<rl4::StandardControlsInputInstance, rl4::ControlsInputInstance>::with(memRes).create(
                    guiControlCount,
                    rawControlCount,
                    psdControlCount,
                    mlControlCount,
+                   rbfControlCount,
                    memRes);
     };
 }
diff --git a/Engine/Plugins/Animation/RigLogic/Source/RigLogicLibTest/Private/rltests/controls/ControlFixtures.h b/Engine/Plugins/Animation/RigLogic/Source/RigLogicLibTest/Private/rltests/controls/ControlFixtures.h
--- a/Engine/Plugins/Animation/RigLogic/Source/RigLogicLibTest/Private/rltests/controls/ControlFixtures.h
+++ b/Engine/Plugins/Animation/RigLogic/Source/RigLogicLibTest/Private/rltests/controls/ControlFixtures.h
@@ -11,5 +11,6 @@ struct ControlsFactory {
     static rl4::ControlsInputInstance::Factory getInstanceFactory(std::uint16_t guiControlCount,
                                                                   std::uint16_t rawControlCount,
                                                                   std::uint16_t psdControlCount,
-                                                                  std::uint16_t mlControlCount);
+                                                                  std::uint16_t mlControlCount,
+                                                                  std::uint16_t rbfControlCount);
 };
diff --git a/Engine/Plugins/Animation/RigLogic/Source/RigLogicLibTest/Private/rltests/controls/TestControls.cpp b/Engine/Plugins/Animation/RigLogic/Source/RigLogicLibTest/Private/rltests/controls/TestControls.cpp
--- a/Engine/Plugins/Animation/RigLogic/Source/RigLogicLibTest/Private/rltests/controls/TestControls.cpp
+++ b/Engine/Plugins/Animation/RigLogic/Source/RigLogicLibTest/Private/rltests/controls/TestControls.cpp
@@ -24,7 +24,7 @@ TEST(ControlsTest, GUIToRawMapping) {
     const std::uint16_t guiControlCount = conditionals.getInputCount();
     const std::uint16_t rawControlCount = conditionals.getOutputCount();
     const std::uint16_t psdControlCount = psds.getDistinctPSDCount();
-    auto instanceFactory = ControlsFactory::getInstanceFactory(guiControlCount, rawControlCount, psdControlCount, 0u);
+    auto instanceFactory = ControlsFactory::getInstanceFactory(guiControlCount, rawControlCount, psdControlCount, 0u, 0u);
     rl4::Controls controls{std::move(conditionals), std::move(psds), instanceFactory};
 
     const rl4::Vector<float> guiControls{0.1f, 0.2f};
@@ -55,7 +55,7 @@ TEST(ControlsTest, PSDsAppendToOutput) {
     const std::uint16_t guiControlCount = conditionals.getInputCount();
     const std::uint16_t rawControlCount = conditionals.getOutputCount();
     const std::uint16_t psdControlCount = psds.getDistinctPSDCount();
-    auto instanceFactory = ControlsFactory::getInstanceFactory(guiControlCount, rawControlCount, psdControlCount, 0u);
+    auto instanceFactory = ControlsFactory::getInstanceFactory(guiControlCount, rawControlCount, psdControlCount, 0u, 0u);
     rl4::Controls controls{std::move(conditionals), std::move(psds), instanceFactory};
     auto instance = controls.createInstance(&amr);
     auto buffer = instance->getInputBuffer();
diff --git a/Engine/Plugins/Animation/RigLogic/Source/RigLogicLibTest/Private/rltests/dna/DNAFixtures.cpp b/Engine/Plugins/Animation/RigLogic/Source/RigLogicLibTest/Private/rltests/dna/DNAFixtures.cpp
--- a/Engine/Plugins/Animation/RigLogic/Source/RigLogicLibTest/Private/rltests/dna/DNAFixtures.cpp
+++ b/Engine/Plugins/Animation/RigLogic/Source/RigLogicLibTest/Private/rltests/dna/DNAFixtures.cpp
@@ -1240,6 +1240,10 @@ const unsigned char machineLearnedBehavior[] = {
 };
 
 std::vector<char> getBytes() {
+    #if !defined(__clang__) && defined(__GNUC__) && (__GNUC__ >= 12)
+        #pragma GCC diagnostic push
+        #pragma GCC diagnostic ignored "-Wstringop-overflow"
+    #endif
     std::vector<char> bytes;
     // Header
     bytes.insert(bytes.end(), header, header + sizeof(header));
@@ -1264,6 +1268,9 @@ std::vector<char> getBytes() {
     // Machine learned behavior
     bytes.insert(bytes.end(), machineLearnedBehavior, machineLearnedBehavior + sizeof(machineLearnedBehavior));
     return bytes;
+    #if !defined(__clang__) && defined(__GNUC__) && (__GNUC__ >= 12)
+        #pragma GCC diagnostic pop
+    #endif
 }
 
 }  // namespace raw
diff --git a/Engine/Plugins/Animation/RigLogic/Source/RigLogicLibTest/Private/rltests/dna/FakeReader.h b/Engine/Plugins/Animation/RigLogic/Source/RigLogicLibTest/Private/rltests/dna/FakeReader.h
--- a/Engine/Plugins/Animation/RigLogic/Source/RigLogicLibTest/Private/rltests/dna/FakeReader.h
+++ b/Engine/Plugins/Animation/RigLogic/Source/RigLogicLibTest/Private/rltests/dna/FakeReader.h
@@ -541,6 +541,149 @@ class FakeReader : public Reader {
             return {};
         }
 
+        // RBFBehaviorReader
+        std::uint16_t getRBFPoseCount() const override {
+            return {};
+        }
+
+        StringView getRBFPoseName(std::uint16_t poseIndex) const override {
+            return {};
+        }
+
+        ConstArrayView<std::uint16_t> getRBFPoseJointOutputIndices(std::uint16_t poseIndex) const override {
+            return {};
+        }
+
+        ConstArrayView<std::uint16_t> getRBFPoseBlendShapeChannelOutputIndices(std::uint16_t poseIndex) const override {
+            return {};
+        }
+
+        ConstArrayView<std::uint16_t> getRBFPoseAnimatedMapOutputIndices(std::uint16_t poseIndex) const override {
+            return {};
+        }
+
+        ConstArrayView<float> getRBFPoseJointOutputValues(std::uint16_t poseIndex) const override {
+            return {};
+        }
+
+        float getRBFPoseScale(std::uint16_t poseIndex) const override {
+            return {};
+        }
+
+        std::uint16_t getRBFSolverCount() const override {
+            return {};
+        }
+
+        std::uint16_t getRBFSolverIndexListCount() const override {
+            return {};
+        }
+
+        ConstArrayView<std::uint16_t> getRBFSolverIndicesForLOD(std::uint16_t lod) const override {
+            return {};
+        }
+
+        StringView getRBFSolverName(std::uint16_t index) const override {
+            return {};
+        }
+
+        ConstArrayView<std::uint16_t> getRBFSolverRawControlIndices(std::uint16_t solverIndex) const override {
+            return {};
+        }
+
+        ConstArrayView<std::uint16_t> getRBFSolverPoseIndices(std::uint16_t solverIndex) const override {
+            return {};
+        }
+
+        ConstArrayView<float> getRBFSolverRawControlValues(std::uint16_t solverIndex) const override {
+            return {};
+        }
+
+        RBFSolverType getRBFSolverType(std::uint16_t solverIndex) const override {
+            return {};
+        }
+
+        float getRBFSolverRadius(std::uint16_t solverIndex) const override {
+            return {};
+        }
+
+        AutomaticRadius getRBFSolverAutomaticRadius(std::uint16_t solverIndex) const override {
+            return {};
+        }
+
+        float getRBFSolverWeightThreshold(std::uint16_t solverIndex) const override {
+            return {};
+        }
+
+        RBFDistanceMethod getRBFSolverDistanceMethod(std::uint16_t solverIndex) const override {
+            return {};
+        }
+
+        RBFNormalizeMethod getRBFSolverNormalizeMethod(std::uint16_t solverIndex) const override {
+            return {};
+        }
+
+        RBFFunctionType getRBFSolverFunctionType(std::uint16_t solverIndex) const override {
+            return {};
+        }
+
+        TwistAxis getRBFSolverTwistAxis(std::uint16_t solverIndex) const override {
+            return {};
+        }
+
+        // JointBehaviorMetadataReader
+        TranslationRepresentation getJointTranslationRepresentation(std::uint16_t jointIndex) const override {
+            return {};
+        }
+
+        RotationRepresentation getJointRotationRepresentation(std::uint16_t jointIndex) const override {
+            return {};
+        }
+
+        ScaleRepresentation getJointScaleRepresentation(std::uint16_t jointIndex) const override {
+            return {};
+        }
+
+        // TwistSwingBehaviorReader
+        std::uint16_t getTwistCount() const override {
+            return {};
+        }
+
+        TwistAxis getTwistSetupTwistAxis(std::uint16_t twistIndex) const override {
+            return {};
+        }
+
+        ConstArrayView<std::uint16_t> getTwistInputControlIndices(std::uint16_t twistIndex) const override {
+            return {};
+        }
+
+        ConstArrayView<std::uint16_t> getTwistOutputJointIndices(std::uint16_t twistIndex) const override {
+            return {};
+        }
+
+        ConstArrayView<float> getTwistBlendWeights(std::uint16_t twistIndex) const override {
+            return {};
+        }
+
+        std::uint16_t getSwingCount() const override {
+            return {};
+        }
+
+        TwistAxis getSwingSetupTwistAxis(std::uint16_t swingIndex) const override {
+            return {};
+        }
+
+        ConstArrayView<std::uint16_t> getSwingInputControlIndices(std::uint16_t swingIndex) const override {
+            return {};
+        }
+
+        ConstArrayView<std::uint16_t> getSwingOutputJointIndices(std::uint16_t swingIndex) const override {
+            return {};
+        }
+
+        ConstArrayView<float> getSwingBlendWeights(std::uint16_t swingIndex) const override {
+            return {};
+        }
+
         void unload(DataLayer  /*unused*/) override {
         }
 
diff --git a/Engine/Plugins/Animation/RigLogic/Source/RigLogicLibTest/Private/rltests/joints/TestJointsFactory.cpp b/Engine/Plugins/Animation/RigLogic/Source/RigLogicLibTest/Private/rltests/joints/TestJointsFactory.cpp
--- a/Engine/Plugins/Animation/RigLogic/Source/RigLogicLibTest/Private/rltests/joints/TestJointsFactory.cpp
+++ b/Engine/Plugins/Animation/RigLogic/Source/RigLogicLibTest/Private/rltests/joints/TestJointsFactory.cpp
@@ -20,5 +20,5 @@ TEST(ScalarJointsFactoryTest, NeutralJointsAreCopied) {
         6.0f, 7.0f, 8.0f, 9.0f, 10.0f, 11.0f, 1.0f, 1.0f, 1.0f,
         12.0f, 13.0f, 14.0f, 15.0f, 16.0f, 17.0f, 1.0f, 1.0f, 1.0f
     };
-    ASSERT_ELEMENTS_EQ(joints->getRawNeutralValues(), expected, 27ul);
+    ASSERT_ELEMENTS_EQ(joints->getNeutralValues(), expected, 27ul);
 }
diff --git a/Engine/Plugins/Animation/RigLogic/Source/RigLogicLibTest/Private/rltests/joints/bpcm/Assertions.h b/Engine/Plugins/Animation/RigLogic/Source/RigLogicLibTest/Private/rltests/joints/bpcm/Assertions.h
--- a/Engine/Plugins/Animation/RigLogic/Source/RigLogicLibTest/Private/rltests/joints/bpcm/Assertions.h
+++ b/Engine/Plugins/Animation/RigLogic/Source/RigLogicLibTest/Private/rltests/joints/bpcm/Assertions.h
@@ -4,10 +4,23 @@
 
 #include "rltests/Defs.h"
 
-#include "riglogic/joints/bpcm/Evaluator.h"
+#include "riglogic/joints/cpu/CPUJointsEvaluator.h"
+#include "riglogic/joints/cpu/bpcm/BPCMJointsEvaluator.h"
 
 namespace rl4 {
 
+struct CPUJointsEvaluator::Accessor {
+
+    static JointsEvaluator* getBPCMEvaluator(CPUJointsEvaluator* parent) {
+        return parent->bpcmEvaluator.get();
+    }
+
+    static JointsEvaluator* getQuaternionEvaluator(CPUJointsEvaluator* parent) {
+        return parent->quaternionEvaluator.get();
+    }
+
+};
+
 namespace bpcm {
 
 template<typename TValue>
@@ -17,6 +30,8 @@ struct Evaluator<TValue>::Accessor {
         ASSERT_EQ(result.storage.inputIndices, expected.storage.inputIndices);
         ASSERT_EQ(result.storage.outputIndices, expected.storage.outputIndices);
         ASSERT_EQ(result.storage.values, expected.storage.values);
+        ASSERT_EQ(result.storage.outputRotationIndices, expected.storage.outputRotationIndices);
+        ASSERT_EQ(result.storage.outputRotationLODs, expected.storage.outputRotationLODs);
     }
 
     static void assertJointGroupsEqual(const Evaluator<TValue>& result, const Evaluator<TValue>& expected) {
@@ -30,6 +45,8 @@ struct Evaluator<TValue>::Accessor {
             ASSERT_EQ(jointGroup.inputIndicesSizeAlignedTo8, expectedJointGroup.inputIndicesSizeAlignedTo8);
             ASSERT_EQ(jointGroup.lodsOffset, expectedJointGroup.lodsOffset);
             ASSERT_EQ(jointGroup.outputIndicesOffset, expectedJointGroup.outputIndicesOffset);
+            ASSERT_EQ(jointGroup.outputRotationIndicesOffset, expectedJointGroup.outputRotationIndicesOffset);
+            ASSERT_EQ(jointGroup.outputRotationLODsOffset, expectedJointGroup.outputRotationLODsOffset);
             ASSERT_EQ(jointGroup.valuesOffset, expectedJointGroup.valuesOffset);
             ASSERT_EQ(jointGroup.valuesSize, expectedJointGroup.valuesSize);
         }
@@ -41,9 +58,8 @@ struct Evaluator<TValue>::Accessor {
             const auto& lodRegion = result.storage.lodRegions[lod];
             const auto& expectedLodRegion = expected.storage.lodRegions[lod];
             ASSERT_EQ(lodRegion.size, expectedLodRegion.size);
-            ASSERT_EQ(lodRegion.sizeAlignedToLastFullBlock, expectedLodRegion.sizeAlignedToLastFullBlock);
-            ASSERT_EQ(lodRegion.sizeAlignedToSecondLastFullBlock,
-                      expectedLodRegion.sizeAlignedToSecondLastFullBlock);
+            ASSERT_EQ(lodRegion.sizePaddedToLastFullBlock, expectedLodRegion.sizePaddedToLastFullBlock);
+            ASSERT_EQ(lodRegion.sizePaddedToSecondLastFullBlock, expectedLodRegion.sizePaddedToSecondLastFullBlock);
         }
     }
 
diff --git a/Engine/Plugins/Animation/RigLogic/Source/RigLogicLibTest/Private/rltests/joints/bpcm/BPCMFixturesBlock4.cpp b/Engine/Plugins/Animation/RigLogic/Source/RigLogicLibTest/Private/rltests/joints/bpcm/BPCMFixturesBlock4.cpp
--- a/Engine/Plugins/Animation/RigLogic/Source/RigLogicLibTest/Private/rltests/joints/bpcm/BPCMFixturesBlock4.cpp
+++ b/Engine/Plugins/Animation/RigLogic/Source/RigLogicLibTest/Private/rltests/joints/bpcm/BPCMFixturesBlock4.cpp
@@ -3,8 +3,8 @@
 #include "rltests/joints/bpcm/BPCMFixturesBlock4.h"
 
 #include "riglogic/TypeDefs.h"
-#include "riglogic/joints/bpcm/BPCMOutputInstance.h"
-#include "riglogic/joints/bpcm/Evaluator.h"
+#include "riglogic/joints/cpu/CPUJointsOutputInstance.h"
+#include "riglogic/joints/cpu/bpcm/BPCMJointsEvaluator.h"
 #include "riglogic/types/Extent.h"
 
 namespace block4 {
@@ -18,7 +18,7 @@ namespace block4 {
 namespace unoptimized {
 
 const std::uint16_t lodCount = 4u;
-const Extent dimensions{75ul, 13ul};
+const Extent dimensions = {75ul, 13ul};
 
 const Matrix<float> values = {
     {  // Joint group 0 - Pad to block-4 => block-4
@@ -108,15 +108,34 @@ const Matrix<float> values = {
         70.0f, 70.0f, 70.0f, 70.0f, 70.0f, 70.0f, 70.0f, 70.0f, 70.0f, 70.0f, 70.0f, 70.0f, 70.0f,
         71.0f, 71.0f, 71.0f, 71.0f, 71.0f, 71.0f, 71.0f, 71.0f, 71.0f, 71.0f, 71.0f, 71.0f, 71.0f,
     },
-    {  // Joint group 8 - Pad to block-4 => block-4
+    {  // Joint group 8 - No padding => block-8 + block-8 => Mixed order of output indices
+        65.0f, 65.0f, 65.0f, 65.0f, 65.0f, 65.0f, 65.0f, 65.0f, 65.0f, 65.0f, 65.0f, 65.0f, 65.0f,
+        66.0f, 66.0f, 66.0f, 66.0f, 66.0f, 66.0f, 66.0f, 66.0f, 66.0f, 66.0f, 66.0f, 66.0f, 66.0f,
+        67.0f, 67.0f, 67.0f, 67.0f, 67.0f, 67.0f, 67.0f, 67.0f, 67.0f, 67.0f, 67.0f, 67.0f, 67.0f,
+        68.0f, 68.0f, 68.0f, 68.0f, 68.0f, 68.0f, 68.0f, 68.0f, 68.0f, 68.0f, 68.0f, 68.0f, 68.0f,
+        69.0f, 69.0f, 69.0f, 69.0f, 69.0f, 69.0f, 69.0f, 69.0f, 69.0f, 69.0f, 69.0f, 69.0f, 69.0f,
+        70.0f, 70.0f, 70.0f, 70.0f, 70.0f, 70.0f, 70.0f, 70.0f, 70.0f, 70.0f, 70.0f, 70.0f, 70.0f,
+        71.0f, 71.0f, 71.0f, 71.0f, 71.0f, 71.0f, 71.0f, 71.0f, 71.0f, 71.0f, 71.0f, 71.0f, 71.0f,
+        56.0f, 56.0f, 56.0f, 56.0f, 56.0f, 56.0f, 56.0f, 56.0f, 56.0f, 56.0f, 56.0f, 56.0f, 56.0f,
+        57.0f, 57.0f, 57.0f, 57.0f, 57.0f, 57.0f, 57.0f, 57.0f, 57.0f, 57.0f, 57.0f, 57.0f, 57.0f,
+        58.0f, 58.0f, 58.0f, 58.0f, 58.0f, 58.0f, 58.0f, 58.0f, 58.0f, 58.0f, 58.0f, 58.0f, 58.0f,
+        59.0f, 59.0f, 59.0f, 59.0f, 59.0f, 59.0f, 59.0f, 59.0f, 59.0f, 59.0f, 59.0f, 59.0f, 59.0f,
+        60.0f, 60.0f, 60.0f, 60.0f, 60.0f, 60.0f, 60.0f, 60.0f, 60.0f, 60.0f, 60.0f, 60.0f, 60.0f,
+        61.0f, 61.0f, 61.0f, 61.0f, 61.0f, 61.0f, 61.0f, 61.0f, 61.0f, 61.0f, 61.0f, 61.0f, 61.0f,
+        62.0f, 62.0f, 62.0f, 62.0f, 62.0f, 62.0f, 62.0f, 62.0f, 62.0f, 62.0f, 62.0f, 62.0f, 62.0f,
+        63.0f, 63.0f, 63.0f, 63.0f, 63.0f, 63.0f, 63.0f, 63.0f, 63.0f, 63.0f, 63.0f, 63.0f, 63.0f,
+        64.0f, 64.0f, 64.0f, 64.0f, 64.0f, 64.0f, 64.0f, 64.0f, 64.0f, 64.0f, 64.0f, 64.0f, 64.0f,
+    },
+    {  // Joint group 9 - Pad to block-4 => block-4
         1.0f, 1.0f, 1.0f, 1.0f, 1.0f,
         2.0f, 2.0f, 2.0f, 2.0f, 2.0f,
     },
-    {  // Joint group 9 - Pad to block-4 => block-4
+    {  // Joint group 10 - Pad to block-4 => block-4
         1.0f, 1.0f, 1.0f, 1.0f, 1.0f,
         2.0f, 2.0f, 2.0f, 2.0f, 2.0f
     }
 };
+
 const Matrix<std::uint16_t> inputIndices = {
     {  // Joint group 0
         12, 11, 10, 9, 8, 7, 6, 5, 4, 3, 2, 1, 0
@@ -143,9 +162,12 @@ const Matrix<std::uint16_t> inputIndices = {
         12, 11, 10, 9, 8, 7, 6, 5, 4, 3, 2, 1, 0
     },
     {  // Joint group 8
-        1, 2, 3, 4, 5
+        12, 11, 10, 9, 8, 7, 6, 5, 4, 3, 2, 1, 0
     },
     {  // Joint group 9
+        1, 2, 3, 4, 5
+    },
+    {  // Joint group 10
         7, 8, 9, 10, 11
     }
 };
@@ -154,31 +176,34 @@ const Matrix<std::uint16_t> outputIndices = {
         0
     },
     {  // Joint group 1
-        1, 2, 3, 4
+        9, 10, 11, 12
     },
     {  // Joint group 2
-        5, 6, 7, 8, 9, 10, 11
+        27, 28, 29, 30, 31, 32, 33
     },
     {  // Joint group 3
-        12, 13, 14, 15, 16, 17, 18, 19
+        36, 37, 38, 39, 40, 41, 42, 43
     },
     {  // Joint group 4
-        20, 21, 22, 23, 24, 25, 26, 27, 28
+        45, 46, 47, 48, 49, 50, 51, 52, 53
     },
     {  // Joint group 5
-        29, 30, 31, 32, 33, 34, 35, 36, 37, 38, 39, 40
+        54, 55, 56, 57, 58, 59, 60, 61, 62, 63, 64, 65
     },
     {  // Joint group 6
-        41, 42, 43, 44, 45, 46, 47, 48, 49, 50, 51, 52, 53, 54
+        72, 73, 74, 75, 76, 77, 78, 79, 80, 81, 82, 83, 84, 85
     },
     {  // Joint group 7
-        55, 56, 57, 58, 59, 60, 61, 62, 63, 64, 65, 66, 67, 68, 69, 70
+        90, 91, 92, 93, 94, 95, 96, 97, 98, 99, 100, 101, 102, 103, 104, 105
     },
     {  // Joint group 8
-        71, 72
+        117, 118, 119, 120, 121, 122, 123, 108, 109, 110, 111, 112, 113, 114, 115, 116
     },
     {  // Joint group 9
-        73, 74
+        126, 127
+    },
+    {  // Joint group 10
+        135, 136
     }
 };
 const Matrix<std::uint16_t> lods = {
@@ -207,15 +232,15 @@ const Matrix<std::uint16_t> lods = {
         16, 12, 8, 4
     },
     {  // Joint group 8
-        2, 2, 2, 2
+        16, 12, 8, 4
     },
     {  // Joint group 9
         2, 2, 2, 2
+    },
+    {  // Joint group 10
+        2, 2, 2, 2
     }
 };
-const Vector<Extent> subMatrices = {
-    {1, 13}, {4, 13}, {7, 13}, {8, 13}, {9, 13}, {12, 13}, {14, 13}, {16, 13}, {2, 5}, {2, 5}
-};
 
 const Vector<float> neutralJointTranslationXs = {0.0f, 6.0f, 12.0f};
 const Vector<float> neutralJointTranslationYs = {1.0f, 7.0f, 13.0f};
@@ -404,13 +429,41 @@ const AlignedMatrix<float> floatValues = {
         64.0f, 65.0f, 66.0f, 67.0f, 68.0f, 69.0f, 70.0f, 71.0f,
     },
     {  // Joint group 8
+        65.0f, 66.0f, 67.0f, 68.0f, 69.0f, 70.0f, 71.0f, 56.0f,
+        65.0f, 66.0f, 67.0f, 68.0f, 69.0f, 70.0f, 71.0f, 56.0f,
+        65.0f, 66.0f, 67.0f, 68.0f, 69.0f, 70.0f, 71.0f, 56.0f,
+        65.0f, 66.0f, 67.0f, 68.0f, 69.0f, 70.0f, 71.0f, 56.0f,
+        65.0f, 66.0f, 67.0f, 68.0f, 69.0f, 70.0f, 71.0f, 56.0f,
+        65.0f, 66.0f, 67.0f, 68.0f, 69.0f, 70.0f, 71.0f, 56.0f,
+        65.0f, 66.0f, 67.0f, 68.0f, 69.0f, 70.0f, 71.0f, 56.0f,
+        65.0f, 66.0f, 67.0f, 68.0f, 69.0f, 70.0f, 71.0f, 56.0f,
+        65.0f, 66.0f, 67.0f, 68.0f, 69.0f, 70.0f, 71.0f, 56.0f,
+        65.0f, 66.0f, 67.0f, 68.0f, 69.0f, 70.0f, 71.0f, 56.0f,
+        65.0f, 66.0f, 67.0f, 68.0f, 69.0f, 70.0f, 71.0f, 56.0f,
+        65.0f, 66.0f, 67.0f, 68.0f, 69.0f, 70.0f, 71.0f, 56.0f,
+        65.0f, 66.0f, 67.0f, 68.0f, 69.0f, 70.0f, 71.0f, 56.0f,
+        57.0f, 58.0f, 59.0f, 60.0f, 61.0f, 62.0f, 63.0f, 64.0f,
+        57.0f, 58.0f, 59.0f, 60.0f, 61.0f, 62.0f, 63.0f, 64.0f,
+        57.0f, 58.0f, 59.0f, 60.0f, 61.0f, 62.0f, 63.0f, 64.0f,
+        57.0f, 58.0f, 59.0f, 60.0f, 61.0f, 62.0f, 63.0f, 64.0f,
+        57.0f, 58.0f, 59.0f, 60.0f, 61.0f, 62.0f, 63.0f, 64.0f,
+        57.0f, 58.0f, 59.0f, 60.0f, 61.0f, 62.0f, 63.0f, 64.0f,
+        57.0f, 58.0f, 59.0f, 60.0f, 61.0f, 62.0f, 63.0f, 64.0f,
+        57.0f, 58.0f, 59.0f, 60.0f, 61.0f, 62.0f, 63.0f, 64.0f,
+        57.0f, 58.0f, 59.0f, 60.0f, 61.0f, 62.0f, 63.0f, 64.0f,
+        57.0f, 58.0f, 59.0f, 60.0f, 61.0f, 62.0f, 63.0f, 64.0f,
+        57.0f, 58.0f, 59.0f, 60.0f, 61.0f, 62.0f, 63.0f, 64.0f,
+        57.0f, 58.0f, 59.0f, 60.0f, 61.0f, 62.0f, 63.0f, 64.0f,
+        57.0f, 58.0f, 59.0f, 60.0f, 61.0f, 62.0f, 63.0f, 64.0f,
+    },
+    {  // Joint group 9
         1.0f, 2.0f, 0.0f, 0.0f,
         1.0f, 2.0f, 0.0f, 0.0f,
         1.0f, 2.0f, 0.0f, 0.0f,
         1.0f, 2.0f, 0.0f, 0.0f,
         1.0f, 2.0f, 0.0f, 0.0f,
     },
-    {  // Joint group 9
+    {  // Joint group 10
         1.0f, 2.0f, 0.0f, 0.0f,
         1.0f, 2.0f, 0.0f, 0.0f,
         1.0f, 2.0f, 0.0f, 0.0f,
@@ -593,13 +646,41 @@ const AlignedMatrix<std::uint16_t> halfFloatValues = {
         21504, 21520, 21536, 21552, 21568, 21584, 21600, 21616,
     },
     {  // Joint group 8
+        21520, 21536, 21552, 21568, 21584, 21600, 21616, 21248,
+        21520, 21536, 21552, 21568, 21584, 21600, 21616, 21248,
+        21520, 21536, 21552, 21568, 21584, 21600, 21616, 21248,
+        21520, 21536, 21552, 21568, 21584, 21600, 21616, 21248,
+        21520, 21536, 21552, 21568, 21584, 21600, 21616, 21248,
+        21520, 21536, 21552, 21568, 21584, 21600, 21616, 21248,
+        21520, 21536, 21552, 21568, 21584, 21600, 21616, 21248,
+        21520, 21536, 21552, 21568, 21584, 21600, 21616, 21248,
+        21520, 21536, 21552, 21568, 21584, 21600, 21616, 21248,
+        21520, 21536, 21552, 21568, 21584, 21600, 21616, 21248,
+        21520, 21536, 21552, 21568, 21584, 21600, 21616, 21248,
+        21520, 21536, 21552, 21568, 21584, 21600, 21616, 21248,
+        21520, 21536, 21552, 21568, 21584, 21600, 21616, 21248,
+        21280, 21312, 21344, 21376, 21408, 21440, 21472, 21504,
+        21280, 21312, 21344, 21376, 21408, 21440, 21472, 21504,
+        21280, 21312, 21344, 21376, 21408, 21440, 21472, 21504,
+        21280, 21312, 21344, 21376, 21408, 21440, 21472, 21504,
+        21280, 21312, 21344, 21376, 21408, 21440, 21472, 21504,
+        21280, 21312, 21344, 21376, 21408, 21440, 21472, 21504,
+        21280, 21312, 21344, 21376, 21408, 21440, 21472, 21504,
+        21280, 21312, 21344, 21376, 21408, 21440, 21472, 21504,
+        21280, 21312, 21344, 21376, 21408, 21440, 21472, 21504,
+        21280, 21312, 21344, 21376, 21408, 21440, 21472, 21504,
+        21280, 21312, 21344, 21376, 21408, 21440, 21472, 21504,
+        21280, 21312, 21344, 21376, 21408, 21440, 21472, 21504,
+        21280, 21312, 21344, 21376, 21408, 21440, 21472, 21504,
+    },
+    {  // Joint group 9
         15360, 16384, 0, 0,
         15360, 16384, 0, 0,
         15360, 16384, 0, 0,
         15360, 16384, 0, 0,
         15360, 16384, 0, 0,
     },
-    {  // Joint group 9
+    {  // Joint group 10
         15360, 16384, 0, 0,
         15360, 16384, 0, 0,
         15360, 16384, 0, 0,
@@ -608,19 +689,6 @@ const AlignedMatrix<std::uint16_t> halfFloatValues = {
     }
 };
 
-const Vector<Extent> subMatrices = {
-    {4, 13},  // Joint group 0
-    {4, 13},  // Joint group 1
-    {8, 13},  // Joint group 2
-    {8, 13},  // Joint group 3
-    {12, 13},  // Joint group 4
-    {12, 13},  // Joint group 5
-    {16, 13},  // Joint group 6
-    {16, 13},  // Joint group 7
-    {4, 5},  // Joint group 8
-    {4, 5}  // Joint group 9
-};
-
 const AlignedMatrix<std::uint16_t> inputIndices = {
     {  // Joint group 0
         12, 11, 10, 9, 8, 7, 6, 5, 4, 3, 2, 1, 0
@@ -647,77 +715,282 @@ const AlignedMatrix<std::uint16_t> inputIndices = {
         12, 11, 10, 9, 8, 7, 6, 5, 4, 3, 2, 1, 0
     },
     {  // Joint group 8
-        1, 2, 3, 4, 5
+        12, 11, 10, 9, 8, 7, 6, 5, 4, 3, 2, 1, 0
     },
     {  // Joint group 9
+        1, 2, 3, 4, 5
+    },
+    {  // Joint group 10
         7, 8, 9, 10, 11
     }
 };
 
-const AlignedMatrix<std::uint16_t> outputIndices = {
-    {  // Joint group 0
-        0, 0, 0, 0
-    },
-    {  // Joint group 1
-        1, 2, 3, 4
-    },
-    {  // Joint group 2
-        5, 6, 7, 8, 9, 10, 11, 0
-    },
-    {  // Joint group 3
-        12, 13, 14, 15, 16, 17, 18, 19
-    },
-    {  // Joint group 4
-        20, 21, 22, 23, 24, 25, 26, 27, 28, 0, 0, 0
-    },
-    {  // Joint group 5
-        29, 30, 31, 32, 33, 34, 35, 36, 37, 38, 39, 40
-    },
-    {  // Joint group 6
-        41, 42, 43, 44, 45, 46, 47, 48, 49, 50, 51, 52, 53, 54, 0, 0
-    },
-    {  // Joint group 7
-        55, 56, 57, 58, 59, 60, 61, 62, 63, 64, 65, 66, 67, 68, 69, 70
-    },
-    {  // Joint group 8
-        71, 72, 0, 0
-    },
-    {  // Joint group 9
-        73, 74, 0, 0
+const Vector<Vector<AlignedVector<std::uint16_t> > > outputIndices = {
+    {  // Quaternion outputs
+        {  // Joint group 0
+            0, 0, 0, 0
+        },
+        {  // Joint group 1
+            10, 11, 12, 13
+        },
+        {  // Joint group 2
+            30, 31, 32, 33, 34, 35, 37, 0
+        },
+        {  // Joint group 3
+            40, 41, 42, 43, 44, 45, 47, 48
+        },
+        {  // Joint group 4
+            50, 51, 52, 53, 54, 55, 57, 58, 59, 0, 0, 0
+        },
+        {  // Joint group 5
+            60, 61, 62, 63, 64, 65, 67, 68, 69, 70, 71, 72
+        },
+        {  // Joint group 6
+            80, 81, 82, 83, 84, 85, 87, 88, 89, 90, 91, 92, 93, 94, 0, 0
+        },
+        {  // Joint group 7
+            100, 101, 102, 103, 104, 105, 107, 108, 109, 110, 111, 112, 113, 114, 115, 117
+        },
+        {  // Joint group 8
+            130, 131, 132, 133, 134, 135, 137, 120, 121, 122, 123, 124, 125, 127, 128, 129
+        },
+        {  // Joint group 9
+            140, 141, 0, 0
+        },
+        {  // Joint group 10
+            150, 151, 0, 0
+        }
+    },
+    {  // Euler-angle outputs
+        {  // Joint group 0
+            0, 0, 0, 0
+        },
+        {  // Joint group 1
+            9, 10, 11, 12
+        },
+        {  // Joint group 2
+            27, 28, 29, 30, 31, 32, 33, 0
+        },
+        {  // Joint group 3
+            36, 37, 38, 39, 40, 41, 42, 43
+        },
+        {  // Joint group 4
+            45, 46, 47, 48, 49, 50, 51, 52, 53, 0, 0, 0
+        },
+        {  // Joint group 5
+            54, 55, 56, 57, 58, 59, 60, 61, 62, 63, 64, 65
+        },
+        {  // Joint group 6
+            72, 73, 74, 75, 76, 77, 78, 79, 80, 81, 82, 83, 84, 85, 0, 0
+        },
+        {  // Joint group 7
+            90, 91, 92, 93, 94, 95, 96, 97, 98, 99, 100, 101, 102, 103, 104, 105
+        },
+        {  // Joint group 8
+            117, 118, 119, 120, 121, 122, 123, 108, 109, 110, 111, 112, 113, 114, 115, 116
+        },
+        {  // Joint group 9
+            126, 127, 0, 0
+        },
+        {  // Joint group 10
+            135, 136, 0, 0
+        }
     }
 };
 
-const Vector<bpcm::JointGroup> jointGroups = {
-    // {valueOffset, inputOffset, outputOffset, lodOffset, valueSize, inputSize, inputSizeAligned4, inputSizeAligned8}
-    {  // Joint group 0
-        0, 0, 0, 0, 52, 13, 12, 8
-    },
-    {  // Joint group 1
-        52, 13, 4, 4, 52, 13, 12, 8
-    },
-    {  // Joint group 2
-        104, 26, 8, 8, 104, 13, 12, 8
-    },
-    {  // Joint group 3
-        208, 39, 16, 12, 104, 13, 12, 8
-    },
-    {  // Joint group 4
-        312, 52, 24, 16, 156, 13, 12, 8
-    },
-    {  // Joint group 5
-        468, 65, 36, 20, 156, 13, 12, 8
-    },
-    {  // Joint group 6
-        624, 78, 48, 24, 208, 13, 12, 8
-    },
-    {  // Joint group 7
-        832, 91, 64, 28, 208, 13, 12, 8
-    },
-    {  // Joint group 8
-        1040, 104, 80, 32, 20, 5, 4, 0
-    },
-    {  // Joint group 9
-        1060, 109, 84, 36, 20, 5, 4, 0
+const Vector<Vector<AlignedVector<std::uint16_t> > > outputRotationIndices = {
+    {  // Quaternion outputs
+        {  // Joint group 0
+        },
+        {  // Joint group 1
+            13
+        },
+        {  // Joint group 2
+            33
+        },
+        {  // Joint group 3
+            43
+        },
+        {  // Joint group 4
+            53
+        },
+        {  // Joint group 5
+            63
+        },
+        {  // Joint group 6
+            83, 93
+        },
+        {  // Joint group 7
+            103, 113
+        },
+        {  // Joint group 8
+            133, 123
+        },
+        {  // Joint group 9
+        },
+        {  // Joint group 10
+        }
+    },
+    {  // Euler-angle outputs
+        {  // Joint group 0
+        },
+        {  // Joint group 1
+        },
+        {  // Joint group 2
+        },
+        {  // Joint group 3
+        },
+        {  // Joint group 4
+        },
+        {  // Joint group 5
+        },
+        {  // Joint group 6
+        },
+        {  // Joint group 7
+        },
+        {  // Joint group 8
+        },
+        {  // Joint group 9
+        },
+        {  // Joint group 10
+        }
+    }
+};
+
+const Vector<Vector<AlignedVector<std::uint16_t> > > outputRotationLODs = {
+    {  // Quaternion outputs
+        {  // Joint group 0
+            0, 0, 0, 0
+        },
+        {  // Joint group 1
+            1, 0, 0, 0
+        },
+        {  // Joint group 2
+            1, 1, 1, 0
+        },
+        {  // Joint group 3
+            1, 1, 1, 0
+        },
+        {  // Joint group 4
+            1, 1, 1, 0
+        },
+        {  // Joint group 5
+            1, 1, 1, 0
+        },
+        {  // Joint group 6
+            2, 2, 1, 1
+        },
+        {  // Joint group 7
+            2, 1, 1, 1
+        },
+        {  // Joint group 8
+            2, 2, 1, 1
+        },
+        {  // Joint group 9
+            0, 0, 0, 0
+        },
+        {  // Joint group 10
+            0, 0, 0, 0
+        }
+    },
+    {  // Euler-angle outputs
+        {  // Joint group 0
+        },
+        {  // Joint group 1
+        },
+        {  // Joint group 2
+        },
+        {  // Joint group 3
+        },
+        {  // Joint group 4
+        },
+        {  // Joint group 5
+        },
+        {  // Joint group 6
+        },
+        {  // Joint group 7
+        },
+        {  // Joint group 8
+        },
+        {  // Joint group 9
+        },
+        {  // Joint group 10
+        }
+    }
+};
+
+const Vector<Vector<bpcm::JointGroup> > jointGroups = {
+    // {valueOffset, inputOffset, outputOffset, lodOffset, outputRotationIndicesOffset, outputRotationLODsOffset, valueSize,
+    // inputSize, inputSizeAligned4, inputSizeAligned8}
+    {  // Quaternion outputs
+        {  // Joint group 0
+            0, 0, 0, 0, 0, 0, 52, 13, 12, 8
+        },
+        {  // Joint group 1
+            52, 13, 4, 4, 0, 4, 52, 13, 12, 8
+        },
+        {  // Joint group 2
+            104, 26, 8, 8, 1, 8, 104, 13, 12, 8
+        },
+        {  // Joint group 3
+            208, 39, 16, 12, 2, 12, 104, 13, 12, 8
+        },
+        {  // Joint group 4
+            312, 52, 24, 16, 3, 16, 156, 13, 12, 8
+        },
+        {  // Joint group 5
+            468, 65, 36, 20, 4, 20, 156, 13, 12, 8
+        },
+        {  // Joint group 6
+            624, 78, 48, 24, 5, 24, 208, 13, 12, 8
+        },
+        {  // Joint group 7
+            832, 91, 64, 28, 7, 28, 208, 13, 12, 8
+        },
+        {  // Joint group 8
+            1040, 104, 80, 32, 9, 32, 208, 13, 12, 8
+        },
+        {  // Joint group 9
+            1248, 117, 96, 36, 11, 36, 20, 5, 4, 0
+        },
+        {  // Joint group 10
+            1268, 122, 100, 40, 11, 40, 20, 5, 4, 0
+        }
+    },
+    {  // Euler-angle outputs
+        {  // Joint group 0
+            0, 0, 0, 0, 0, 0, 52, 13, 12, 8
+        },
+        {  // Joint group 1
+            52, 13, 4, 4, 0, 0, 52, 13, 12, 8
+        },
+        {  // Joint group 2
+            104, 26, 8, 8, 0, 0, 104, 13, 12, 8
+        },
+        {  // Joint group 3
+            208, 39, 16, 12, 0, 0, 104, 13, 12, 8
+        },
+        {  // Joint group 4
+            312, 52, 24, 16, 0, 0, 156, 13, 12, 8
+        },
+        {  // Joint group 5
+            468, 65, 36, 20, 0, 0, 156, 13, 12, 8
+        },
+        {  // Joint group 6
+            624, 78, 48, 24, 0, 0, 208, 13, 12, 8
+        },
+        {  // Joint group 7
+            832, 91, 64, 28, 0, 0, 208, 13, 12, 8
+        },
+        {  // Joint group 8
+            1040, 104, 80, 32, 0, 0, 208, 13, 12, 8
+        },
+        {  // Joint group 9
+            1248, 117, 96, 36, 0, 0, 20, 5, 4, 0
+        },
+        {  // Joint group 10
+            1268, 122, 100, 40, 0, 0, 20, 5, 4, 0
+        }
     }
 };
 
@@ -772,12 +1045,18 @@ const Matrix<bpcm::LODRegion> lodRegions = {
         {4, 8, 0},  // LOD-3
     },
     {  // Joint group 8
+        {16, 16, 16},  // LOD-0
+        {12, 16, 8},  // LOD-1
+        {8, 8, 8},  // LOD-2
+        {4, 8, 0},  // LOD-3
+    },
+    {  // Joint group 9
         {2, 0, 0},  // LOD-0
         {2, 0, 0},  // LOD-1
         {2, 0, 0},  // LOD-2
         {2, 0, 0},  // LOD-3
     },
-    {  // Joint group 9
+    {  // Joint group 10
         {2, 0, 0},  // LOD-0
         {2, 0, 0},  // LOD-1
         {2, 0, 0},  // LOD-2
@@ -790,74 +1069,195 @@ const Matrix<bpcm::LODRegion> lodRegions = {
 namespace input {
 
 // Calculation input values
-const rl4::Vector<float> values = {1.0f, 2.0f, 3.0f, 4.0f, 0.0f, 6.0f, 7.0f, 8.0f, 9.0f, 0.0f, 11.0f, 12.0f, 13.0f};
+const Vector<float> values = {1.0f, 2.0f, 3.0f, 4.0f, 0.0f, 6.0f, 7.0f, 8.0f, 9.0f, 0.0f, 11.0f, 12.0f, 13.0f};
 
 }  // namespace input
 
 namespace output {
 
 // Expected output results for each LOD
-const rl4::Matrix<float> valuesPerLOD = {
-    {
-        // LOD-0
-        76.0f,  // Joint group 0
-        152.0f, 228.0f, 304.0f, 380.0f,  // Joint group 1
-        456.0f, 532.0f, 608.0f, 684.0f, 760.0f, 836.0f, 912.0f,  // Joint group 2
-        988.0f, 1064.0f, 1140.0f, 1216.0f, 1292.0f, 1368.0f, 1444.0f, 1520.0f,  // Joint group 3
-        1596.0f, 1672.0f, 1748.0f, 1824.0f, 1900.0f, 1976.0f, 2052.0f, 2128.0f, 2204.0f,  // Joint group 4
-        2280.0f, 2356.0f, 2432.0f, 2508.0f, 2584.0f, 2660.0f,  // Joint group 5
-        2736.0f, 2812.0f, 2888.0f, 2964.0f, 3040.0f, 3116.0f,  // Joint group 5
-        3192.0f, 3268.0f, 3344.0f, 3420.0f, 3496.0f, 3572.0f, 3648.0f,  // Joint group 6
-        3724.0f, 3800.0f, 3876.0f, 3952.0f, 4028.0f, 4104.0f, 4180.0f,  // Joint group 6
-        4256.0f, 4332.0f, 4408.0f, 4484.0f, 4560.0f, 4636.0f, 4712.0f, 4788.0f,  // Joint group 7
-        4864.0f, 4940.0f, 5016.0f, 5092.0f, 5168.0f, 5244.0f, 5320.0f, 5396.0f,  // Joint group 7
-        15.0f, 30.0f,  // Joint group 8
-        40.0f, 80.0f  // Joint group 9
-    }, {
-        // LOD-1
-        76.0f,  // Joint group 0
-        152.0f, 228.0f, 304.0f, 0.0f,  // Joint group 1
-        456.0f, 532.0f, 608.0f, 684.0f, 760.0f, 0.0f, 0.0f,  // Joint group 2
-        988.0f, 1064.0f, 1140.0f, 1216.0f, 1292.0f, 1368.0f, 0.0f, 0.0f,  // Joint group 3
-        1596.0f, 1672.0f, 1748.0f, 1824.0f, 1900.0f, 1976.0f, 2052.0f, 0.0f, 0.0f,  // Joint group 4
-        2280.0f, 2356.0f, 2432.0f, 2508.0f, 2584.0f, 2660.0f,  // Joint group 5
-        2736.0f, 2812.0f, 2888.0f, 0.0f, 0.0f, 0.0f,  // Joint group 5
-        3192.0f, 3268.0f, 3344.0f, 3420.0f, 3496.0f, 3572.0f, 3648.0f,  // Joint group 6
-        3724.0f, 3800.0f, 3876.0f, 3952.0f, 4028.0f, 4104.0f, 0.0f,  // Joint group 6
-        4256.0f, 4332.0f, 4408.0f, 4484.0f, 4560.0f, 4636.0f, 4712.0f, 4788.0f,  // Joint group 7
-        4864.0f, 4940.0f, 5016.0f, 5092.0f, 0.0f, 0.0f, 0.0f, 0.0f,  // Joint group 7
-        15.0f, 30.0f,  // Joint group 8
-        40.0f, 80.0f  // Joint group 9
-    }, {
-        // LOD-2
-        76.0f,  // Joint group 0
-        152.0f, 228.0f, 0.0f, 0.0f,  // Joint group 1
-        456.0f, 532.0f, 608.0f, 684.0f, 0.0f, 0.0f, 0.0f,  // Joint group 2
-        988.0f, 1064.0f, 1140.0f, 1216.0f, 0.0f, 0.0f, 0.0f, 0.0f,  // Joint group 3
-        1596.0f, 1672.0f, 1748.0f, 1824.0f, 0.0f, 0.0f, 0.0f, 0.0f, 0.0f,  // Joint group 4
-        2280.0f, 2356.0f, 2432.0f, 2508.0f, 2584.0f, 2660.0f,  // Joint group 5
-        2736.0f, 0.0f, 0.0f, 0.0f, 0.0f, 0.0f,  // Joint group 5
-        3192.0f, 3268.0f, 3344.0f, 3420.0f, 3496.0f, 3572.0f, 3648.0f,  // Joint group 6
-        3724.0f, 3800.0f, 3876.0f, 3952.0f, 0.0f, 0.0f, 0.0f,  // Joint group 6
-        4256.0f, 4332.0f, 4408.0f, 4484.0f, 4560.0f, 4636.0f, 4712.0f, 4788.0f,  // Joint group 7
-        0.0f, 0.0f, 0.0f, 0.0f, 0.0f, 0.0f, 0.0f, 0.0f,  // Joint group 7
-        15.0f, 30.0f,  // Joint group 8
-        40.0f, 80.0f  // Joint group 9
-    }, {
-        // LOD-3
-        76.0f,  // Joint group 0
-        152.0f, 0.0f, 0.0f, 0.0f,  // Joint group 1
-        456.0f, 0.0f, 0.0f, 0.0f, 0.0f, 0.0f, 0.0f,  // Joint group 2
-        988.0f, 1064.0f, 1140.0f, 0.0f, 0.0f, 0.0f, 0.0f, 0.0f,  // Joint group 3
-        1596.0f, 1672.0f, 0.0f, 0.0f, 0.0f, 0.0f, 0.0f, 0.0f, 0.0f,  // Joint group 4
-        2280.0f, 2356.0f, 2432.0f, 0.0f, 0.0f, 0.0f,  // Joint group 5
-        0.0f, 0.0f, 0.0f, 0.0f, 0.0f, 0.0f,  // Joint group 5
-        3192.0f, 3268.0f, 3344.0f, 3420.0f, 3496.0f, 3572.0f, 3648.0f,  // Joint group 6
-        3724.0f, 3800.0f, 0.0f, 0.0f, 0.0f, 0.0f, 0.0f,  // Joint group 6
-        4256.0f, 4332.0f, 4408.0f, 4484.0f, 0.0f, 0.0f, 0.0f, 0.0f,  // Joint group 7
-        0.0f, 0.0f, 0.0f, 0.0f, 0.0f, 0.0f, 0.0f, 0.0f,  // Joint group 7
-        15.0f, 30.0f,  // Joint group 8
-        40.0f, 80.0f  // Joint group 9
+const Vector<Matrix<float> > valuesPerLOD = {
+    {  // Quaternion outputs
+        {
+            // LOD-0
+            76.0f, 0.0f, 0.0f, 0.0f, 0.0f, 0.0f, 1.0f, 0.0f, 0.0f, 0.0f,
+            152.0f, 228.0f, 304.0f, -0.1736482f, 0.0f, 0.0f, -0.9848078f, 0.0f, 0.0f, 0.0f,
+            0.0f, 0.0f, 0.0f, 0.0f, 0.0f, 0.0f, 1.0f, 0.0f, 0.0f, 0.0f,
+            456.0f, 532.0f, 608.0f, 0.121975f, 0.4186294f, 0.7018941f, 0.5632195f, 912.0f, 0.0f, 0.0f,
+            988.0f, 1064.0f, 1140.0f, -0.4184162f, 0.1411047f, 0.7817417f, 0.4403377f, 1444.0f, 1520.0f, 0.0f,
+            1596.0f, 1672.0f, 1748.0f, -0.7535122f, 0.1074113f, -0.6339166f, 0.1372298f, 2052.0f, 2128.0f, 2204.0f,
+            2280.0f, 2356.0f, 2432.0f, -0.4649149f, -0.2635494f, -0.773594f, -0.3405117f, 2736.0f, 2812.0f, 2888.0f,
+            2964.0f, 3040.0f, 3116.0f, 0.0f, 0.0f, 0.0f, 1.0f, 0.0f, 0.0f, 0.0f,
+            3192.0f, 3268.0f, 3344.0f, -0.5973737f, -0.148942f, 0.7646035f, 0.1906371f, 3648.0f, 3724.0f, 3800.0f,
+            3876.0f, 3952.0f, 4028.0f, -0.3252805f, 0.290381f, 0.8937008f, -0.10569f, 0.0f, 0.0f, 0.0f,
+            4256.0f, 4332.0f, 4408.0f, 0.5042307f, 0.0737294f, -0.821218f, -0.2567421f, 4712.0f, 4788.0f, 4864.0f,
+            4940.0f, 5016.0f, 5092.0f, 0.4187725f, -0.2083559f, -0.7320557f, -0.4952896f, 5396.0f, 0.0f, 0.0f,
+            4256.0f, 4332.0f, 4408.0f, 0.5042307f, 0.0737294f, -0.821218f, -0.2567421f, 4712.0f, 4788.0f, 4864.0f,
+            4940.0f, 5016.0f, 5092.0f, 0.4187725f, -0.2083559f, -0.7320557f, -0.4952896f, 5396.0f, 0.0f, 0.0f,
+            15.0f, 30.0f, 0.0f, 0.0f, 0.0f, 0.0f, 1.0f, 0.0f, 0.0f, 0.0f,
+            40.0f, 80.0f
+        }, {
+            // LOD-1
+            76.0f, 0.0f, 0.0f, 0.0f, 0.0f, 0.0f, 1.0f, 0.0f, 0.0f, 0.0f,
+            152.0f, 228.0f, 304.0f, 0.0f, 0.0f, 0.0f, 1.0f, 0.0f, 0.0f, 0.0f,
+            0.0f, 0.0f, 0.0f, 0.0f, 0.0f, 0.0f, 1.0f, 0.0f, 0.0f, 0.0f,
+            456.0f, 532.0f, 608.0f, -0.290381f, 0.3252805f, -0.10569f, 0.8937008f, 0.0f, 0.0f, 0.0f,
+            988.0f, 1064.0f, 1140.0f, -0.4184162f, 0.1411047f, 0.7817417f, 0.4403377f, 0.0f, 0.0f, 0.0f,
+            1596.0f, 1672.0f, 1748.0f, -0.7535122f, 0.1074113f, -0.6339166f, 0.1372298f, 2052.0f, 0.0f, 0.0f,
+            2280.0f, 2356.0f, 2432.0f, -0.4649149f, -0.2635494f, -0.773594f, -0.3405117f, 2736.0f, 2812.0f, 2888.0f,
+            0.0f, 0.0f, 0.0f, 0.0f, 0.0f, 0.0f, 1.0f, 0.0f, 0.0f, 0.0f,
+            3192.0f, 3268.0f, 3344.0f, -0.5973737f, -0.148942f, 0.7646035f, 0.1906371f, 3648.0f, 3724.0f, 3800.0f,
+            3876.0f, 3952.0f, 4028.0f, -0.9510565f, 0.0f, 0.0f, -0.309017f, 0.0f, 0.0f, 0.0f,
+            4256.0f, 4332.0f, 4408.0f, 0.5042307f, 0.0737294f, -0.821218f, -0.2567421f, 4712.0f, 4788.0f, 4864.0f,
+            4940.0f, 5016.0f, 5092.0f, 0.0f, 0.0f, 0.0f, 1.0f, 0.0f, 0.0f, 0.0f,
+            4256.0f, 4332.0f, 4408.0f, -0.495134f, 0.1205274f, 0.8575973f, -0.0695866f, 0.0f, 0.0f, 0.0f,
+            4940.0f, 5016.0f, 5092.0f, 0.4187725f, -0.2083559f, -0.7320557f, -0.4952896f, 5396.0f, 0.0f, 0.0f,
+            15.0f, 30.0f, 0.0f, 0.0f, 0.0f, 0.0f, 1.0f, 0.0f, 0.0f, 0.0f,
+            40.0f, 80.0f
+        }, {
+            // LOD-2
+            76.0f, 0.0f, 0.0f, 0.0f, 0.0f, 0.0f, 1.0f, 0.0f, 0.0f, 0.0f,
+            152.0f, 228.0f, 0.0f, 0.0f, 0.0f, 0.0f, 1.0f, 0.0f, 0.0f, 0.0f,
+            0.0f, 0.0f, 0.0f, 0.0f, 0.0f, 0.0f, 1.0f, 0.0f, 0.0f, 0.0f,
+            456.0f, 532.0f, 608.0f, -0.309017f, 0.0f, 0.0f, 0.9510565f, 0.0f, 0.0f, 0.0f,
+            988.0f, 1064.0f, 1140.0f, -0.9271839f, 0.0f, 0.0f, -0.3746066f, 0.0f, 0.0f, 0.0f,
+            1596.0f, 1672.0f, 1748.0f, -0.2079117f, 0.0f, 0.0f, -0.9781476f, 0.0f, 0.0f, 0.0f,
+            2280.0f, 2356.0f, 2432.0f, -0.4649149f, -0.2635494f, -0.773594f, -0.3405117f, 2736.0f, 0.0f, 0.0f,
+            0.0f, 0.0f, 0.0f, 0.0f, 0.0f, 0.0f, 1.0f, 0.0f, 0.0f, 0.0f,
+            3192.0f, 3268.0f, 3344.0f, -0.5973737f, -0.148942f, 0.7646035f, 0.1906371f, 3648.0f, 3724.0f, 3800.0f,
+            3876.0f, 3952.0f, 0.0f, 0.0f, 0.0f, 0.0f, 1.0f, 0.0f, 0.0f, 0.0f,
+            4256.0f, 4332.0f, 4408.0f, 0.5042307f, 0.0737294f, -0.821218f, -0.2567421f, 4712.0f, 4788.0f, 0.0f,
+            0.0f, 0.0f, 0.0f, 0.0f, 0.0f, 0.0f, 1.0f, 0.0f, 0.0f, 0.0f,
+            4256.0f, 0.0f, 0.0f, 0.0f, 0.0f, 0.0f, 1.0f, 0.0f, 0.0f, 0.0f,
+            4940.0f, 5016.0f, 5092.0f, 0.4187725f, -0.2083559f, -0.7320557f, -0.4952896f, 5396.0f, 0.0f, 0.0f,
+            15.0f, 30.0f, 0.0f, 0.0f, 0.0f, 0.0f, 1.0f, 0.0f, 0.0f, 0.0f,
+            40.0f, 80.0f
+        }, {
+            // LOD-3
+            76.0f, 0.0f, 0.0f, 0.0f, 0.0f, 0.0f, 1.0f, 0.0f, 0.0f, 0.0f,
+            152.0f, 0.0f, 0.0f, 0.0f, 0.0f, 0.0f, 1.0f, 0.0f, 0.0f, 0.0f,
+            0.0f, 0.0f, 0.0f, 0.0f, 0.0f, 0.0f, 1.0f, 0.0f, 0.0f, 0.0f,
+            456.0f, 0.0f, 0.0f, 0.0f, 0.0f, 0.0f, 1.0f, 0.0f, 0.0f, 0.0f,
+            988.0f, 1064.0f, 1140.0f, 0.0f, 0.0f, 0.0f, 1.0f, 0.0f, 0.0f, 0.0f,
+            1596.0f, 1672.0f, 0.0f, 0.0f, 0.0f, 0.0f, 1.0f, 0.0f, 0.0f, 0.0f,
+            2280.0f, 2356.0f, 2432.0f, 0.0f, 0.0f, 0.0f, 1.0f, 0.0f, 0.0f, 0.0f,
+            0.0f, 0.0f, 0.0f, 0.0f, 0.0f, 0.0f, 1.0f, 0.0f, 0.0f, 0.0f,
+            3192.0f, 3268.0f, 3344.0f, -0.5973737f, -0.148942f, 0.7646035f, 0.1906371f, 3648.0f, 3724.0f, 3800.0f,
+            0.0f, 0.0f, 0.0f, 0.0f, 0.0f, 0.0f, 1.0f, 0.0f, 0.0f, 0.0f,
+            4256.0f, 4332.0f, 4408.0f, 0.9902681f, 0.0f, 0.0f, 0.1391731f, 0.0f, 0.0f, 0.0f,
+            0.0f, 0.0f, 0.0f, 0.0f, 0.0f, 0.0f, 1.0f, 0.0f, 0.0f, 0.0f,
+            0.0f, 0.0f, 0.0f, 0.0f, 0.0f, 0.0f, 1.0f, 0.0f, 0.0f, 0.0f,
+            4940.0f, 5016.0f, 5092.0f, 0.898794f, 0.0f, 0.0f, 0.4383711f, 0.0f, 0.0f, 0.0f,
+            15.0f, 30.0f, 0.0f, 0.0f, 0.0f, 0.0f, 1.0f, 0.0f, 0.0f, 0.0f,
+            40.0f, 80.0f
+        }
+    },
+    {  // Euler-angle outputs
+        {
+            // LOD-0
+            76.0f,  // Joint group 0
+            0.0f, 0.0f, 0.0f, 0.0f, 0.0f, 0.0f, 0.0f, 0.0f,
+            152.0f, 228.0f, 304.0f, 380.0f,  // Joint group 1
+            0.0f, 0.0f, 0.0f, 0.0f, 0.0f, 0.0f, 0.0f, 0.0f, 0.0f, 0.0f, 0.0f, 0.0f, 0.0f, 0.0f,
+            456.0f, 532.0f, 608.0f, 684.0f, 760.0f, 836.0f, 912.0f,  // Joint group 2
+            0.0f, 0.0f,
+            988.0f, 1064.0f, 1140.0f, 1216.0f, 1292.0f, 1368.0f, 1444.0f, 1520.0f,  // Joint group 3
+            0.0f,
+            1596.0f, 1672.0f, 1748.0f, 1824.0f, 1900.0f, 1976.0f, 2052.0f, 2128.0f, 2204.0f,  // Joint group 4
+            2280.0f, 2356.0f, 2432.0f, 2508.0f, 2584.0f, 2660.0f,  // Joint group 5
+            2736.0f, 2812.0f, 2888.0f, 2964.0f, 3040.0f, 3116.0f,  // Joint group 5
+            0.0f, 0.0f, 0.0f, 0.0f, 0.0f, 0.0f,
+            3192.0f, 3268.0f, 3344.0f, 3420.0f, 3496.0f, 3572.0f, 3648.0f,  // Joint group 6
+            3724.0f, 3800.0f, 3876.0f, 3952.0f, 4028.0f, 4104.0f, 4180.0f,  // Joint group 6
+            0.0f, 0.0f, 0.0f, 0.0f,
+            4256.0f, 4332.0f, 4408.0f, 4484.0f, 4560.0f, 4636.0f, 4712.0f, 4788.0f,  // Joint group 7
+            4864.0f, 4940.0f, 5016.0f, 5092.0f, 5168.0f, 5244.0f, 5320.0f, 5396.0f,  // Joint group 7
+            0.0f, 0.0f,
+            4256.0f, 4332.0f, 4408.0f, 4484.0f, 4560.0f, 4636.0f, 4712.0f, 4788.0f,  // Joint group 8
+            4864.0f, 4940.0f, 5016.0f, 5092.0f, 5168.0f, 5244.0f, 5320.0f, 5396.0f,  // Joint group 8
+            0.0f, 0.0f,
+            15.0f, 30.0f,  // Joint group 9
+            0.0f, 0.0f, 0.0f, 0.0f, 0.0f, 0.0f, 0.0f,
+            40.0f, 80.0f  // Joint group 10
+        }, {
+            // LOD-1
+            76.0f,  // Joint group 0
+            0.0f, 0.0f, 0.0f, 0.0f, 0.0f, 0.0f, 0.0f, 0.0f,
+            152.0f, 228.0f, 304.0f, 0.0f,  // Joint group 1
+            0.0f, 0.0f, 0.0f, 0.0f, 0.0f, 0.0f, 0.0f, 0.0f, 0.0f, 0.0f, 0.0f, 0.0f, 0.0f, 0.0f,
+            456.0f, 532.0f, 608.0f, 684.0f, 760.0f, 0.0f, 0.0f,  // Joint group 2
+            0.0f, 0.0f,
+            988.0f, 1064.0f, 1140.0f, 1216.0f, 1292.0f, 1368.0f, 0.0f, 0.0f,  // Joint group 3
+            0.0f,
+            1596.0f, 1672.0f, 1748.0f, 1824.0f, 1900.0f, 1976.0f, 2052.0f, 0.0f, 0.0f,  // Joint group 4
+            2280.0f, 2356.0f, 2432.0f, 2508.0f, 2584.0f, 2660.0f,  // Joint group 5
+            2736.0f, 2812.0f, 2888.0f, 0.0f, 0.0f, 0.0f,  // Joint group 5
+            0.0f, 0.0f, 0.0f, 0.0f, 0.0f, 0.0f,
+            3192.0f, 3268.0f, 3344.0f, 3420.0f, 3496.0f, 3572.0f, 3648.0f,  // Joint group 6
+            3724.0f, 3800.0f, 3876.0f, 3952.0f, 4028.0f, 4104.0f, 0.0f,  // Joint group 6
+            0.0f, 0.0f, 0.0f, 0.0f,
+            4256.0f, 4332.0f, 4408.0f, 4484.0f, 4560.0f, 4636.0f, 4712.0f, 4788.0f,  // Joint group 7
+            4864.0f, 4940.0f, 5016.0f, 5092.0f, 0.0f, 0.0f, 0.0f, 0.0f,  // Joint group 7
+            0.0f, 0.0f,
+            4256.0f, 4332.0f, 4408.0f, 4484.0f, 4560.0f, 0.0f, 0.0f, 0.0f,  // Joint group 8
+            0.0f, 4940.0f, 5016.0f, 5092.0f, 5168.0f, 5244.0f, 5320.0f, 5396.0f,  // Joint group 8
+            0.0f, 0.0f,
+            15.0f, 30.0f,  // Joint group 9
+            0.0f, 0.0f, 0.0f, 0.0f, 0.0f, 0.0f, 0.0f,
+            40.0f, 80.0f  // Joint group 10
+        }, {
+            // LOD-2
+            76.0f,  // Joint group 0
+            0.0f, 0.0f, 0.0f, 0.0f, 0.0f, 0.0f, 0.0f, 0.0f,
+            152.0f, 228.0f, 0.0f, 0.0f,  // Joint group 1
+            0.0f, 0.0f, 0.0f, 0.0f, 0.0f, 0.0f, 0.0f, 0.0f, 0.0f, 0.0f, 0.0f, 0.0f, 0.0f, 0.0f,
+            456.0f, 532.0f, 608.0f, 684.0f, 0.0f, 0.0f, 0.0f,  // Joint group 2
+            0.0f, 0.0f,
+            988.0f, 1064.0f, 1140.0f, 1216.0f, 0.0f, 0.0f, 0.0f, 0.0f,  // Joint group 3
+            0.0f,
+            1596.0f, 1672.0f, 1748.0f, 1824.0f, 0.0f, 0.0f, 0.0f, 0.0f, 0.0f,  // Joint group 4
+            2280.0f, 2356.0f, 2432.0f, 2508.0f, 2584.0f, 2660.0f,  // Joint group 5
+            2736.0f, 0.0f, 0.0f, 0.0f, 0.0f, 0.0f,  // Joint group 5
+            0.0f, 0.0f, 0.0f, 0.0f, 0.0f, 0.0f,
+            3192.0f, 3268.0f, 3344.0f, 3420.0f, 3496.0f, 3572.0f, 3648.0f,  // Joint group 6
+            3724.0f, 3800.0f, 3876.0f, 3952.0f, 0.0f, 0.0f, 0.0f,  // Joint group 6
+            0.0f, 0.0f, 0.0f, 0.0f,
+            4256.0f, 4332.0f, 4408.0f, 4484.0f, 4560.0f, 4636.0f, 4712.0f, 4788.0f,  // Joint group 7
+            0.0f, 0.0f, 0.0f, 0.0f, 0.0f, 0.0f, 0.0f, 0.0f,  // Joint group 7
+            0.0f, 0.0f,
+            4256.0f, 0.0f, 0.0f, 0.0f, 0.0f, 0.0f, 0.0f, 0.0f,  // Joint group 8
+            0.0f, 4940.0f, 5016.0f, 5092.0f, 5168.0f, 5244.0f, 5320.0f, 5396.0f,  // Joint group 8
+            0.0f, 0.0f,
+            15.0f, 30.0f,  // Joint group 9
+            0.0f, 0.0f, 0.0f, 0.0f, 0.0f, 0.0f, 0.0f,
+            40.0f, 80.0f  // Joint group 10
+        }, {
+            // LOD-3
+            76.0f,  // Joint group 0
+            0.0f, 0.0f, 0.0f, 0.0f, 0.0f, 0.0f, 0.0f, 0.0f,
+            152.0f, 0.0f, 0.0f, 0.0f,  // Joint group 1
+            0.0f, 0.0f, 0.0f, 0.0f, 0.0f, 0.0f, 0.0f, 0.0f, 0.0f, 0.0f, 0.0f, 0.0f, 0.0f, 0.0f,
+            456.0f, 0.0f, 0.0f, 0.0f, 0.0f, 0.0f, 0.0f,  // Joint group 2
+            0.0f, 0.0f,
+            988.0f, 1064.0f, 1140.0f, 0.0f, 0.0f, 0.0f, 0.0f, 0.0f,  // Joint group 3
+            0.0f,
+            1596.0f, 1672.0f, 0.0f, 0.0f, 0.0f, 0.0f, 0.0f, 0.0f, 0.0f,  // Joint group 4
+            2280.0f, 2356.0f, 2432.0f, 0.0f, 0.0f, 0.0f,  // Joint group 5
+            0.0f, 0.0f, 0.0f, 0.0f, 0.0f, 0.0f,  // Joint group 5
+            0.0f, 0.0f, 0.0f, 0.0f, 0.0f, 0.0f,
+            3192.0f, 3268.0f, 3344.0f, 3420.0f, 3496.0f, 3572.0f, 3648.0f,  // Joint group 6
+            3724.0f, 3800.0f, 0.0f, 0.0f, 0.0f, 0.0f, 0.0f,  // Joint group 6
+            0.0f, 0.0f, 0.0f, 0.0f,
+            4256.0f, 4332.0f, 4408.0f, 4484.0f, 0.0f, 0.0f, 0.0f, 0.0f,  // Joint group 7
+            0.0f, 0.0f, 0.0f, 0.0f, 0.0f, 0.0f, 0.0f, 0.0f,  // Joint group 7
+            0.0f, 0.0f,
+            0.0f, 0.0f, 0.0f, 0.0f, 0.0f, 0.0f, 0.0f, 0.0f,  // Joint group 8
+            0.0f, 4940.0f, 5016.0f, 5092.0f, 5168.0f, 0.0f, 0.0f, 0.0f,  // Joint group 8
+            0.0f, 0.0f,
+            15.0f, 30.0f,  // Joint group 9
+            0.0f, 0.0f, 0.0f, 0.0f, 0.0f, 0.0f, 0.0f,
+            40.0f, 80.0f  // Joint group 10
+        }
     }
 };
 
@@ -874,7 +1274,7 @@ struct OptimizedValues;
 
 template<>
 struct OptimizedValues<float> {
-    static const rl4::AlignedMatrix<float>& get() {
+    static const AlignedMatrix<float>& get() {
         return optimized::floatValues;
     }
 
@@ -882,14 +1282,17 @@ struct OptimizedValues<float> {
 
 template<>
 struct OptimizedValues<std::uint16_t> {
-    static const rl4::AlignedMatrix<std::uint16_t>& get() {
+    static const AlignedMatrix<std::uint16_t>& get() {
         return optimized::halfFloatValues;
     }
 
 };
 
 template<typename TValue>
-bpcm::Evaluator<TValue> OptimizedStorage<TValue>::create(StrategyPtr strategy, rl4::MemoryResource* memRes) {
+bpcm::Evaluator<TValue> OptimizedStorage<TValue>::create(StrategyPtr strategy,
+                                                         std::size_t rotationSelectorIndex,
+                                                         rl4::RotationType rotationType,
+                                                         MemoryResource* memRes) {
     bpcm::JointStorage<TValue> storage{memRes};
     const auto& values = OptimizedValues<TValue>::get();
     for (std::uint16_t i = 0u; i < static_cast<std::uint16_t>(unoptimized::values.size()); ++i) {
@@ -898,18 +1301,32 @@ bpcm::Evaluator<TValue> OptimizedStorage<TValue>::create(StrategyPtr strategy, r
                                     optimized::inputIndices[i].begin(),
                                     optimized::inputIndices[i].end());
         storage.outputIndices.insert(storage.outputIndices.end(),
-                                     optimized::outputIndices[i].begin(),
-                                     optimized::outputIndices[i].end());
+                                     optimized::outputIndices[rotationSelectorIndex][i].begin(),
+                                     optimized::outputIndices[rotationSelectorIndex][i].end());
+        storage.outputRotationIndices.insert(storage.outputRotationIndices.end(),
+                                             optimized::outputRotationIndices[rotationSelectorIndex][i].begin(),
+                                             optimized::outputRotationIndices[rotationSelectorIndex][i].end());
+        storage.outputRotationLODs.insert(storage.outputRotationLODs.end(),
+                                          optimized::outputRotationLODs[rotationSelectorIndex][i].begin(),
+                                          optimized::outputRotationLODs[rotationSelectorIndex][i].end());
         storage.lodRegions.insert(storage.lodRegions.end(),
                                   optimized::lodRegions[i].begin(),
                                   optimized::lodRegions[i].end());
-        storage.jointGroups.push_back(optimized::jointGroups[i]);
+        storage.jointGroups.push_back(optimized::jointGroups[rotationSelectorIndex][i]);
     }
 
-    const auto outputCount = static_cast<std::uint16_t>(unoptimized::dimensions.rows);
-    auto instanceFactory = [outputCount](rl4::MemoryResource* instanceMemRes) {
-            return pma::UniqueInstance<rl4::bpcm::OutputInstance, rl4::JointsOutputInstance>::with(instanceMemRes).create(
+    const auto lastJointGroupIndex = optimized::outputIndices[rotationSelectorIndex].size() - 1ul;
+    const auto numAttrsPerJoint = static_cast<std::uint16_t>(static_cast<std::uint8_t>(rl4::TranslationType::Vector) +
+                                                             static_cast<std::uint8_t>(rotationType) +
+                                                             static_cast<std::uint8_t>(rl4::ScaleType::Vector));
+    const auto maxOutputIndex = extd::maxOf(optimized::outputIndices[rotationSelectorIndex][lastJointGroupIndex]);
+    const auto outputCount = extd::roundUp(maxOutputIndex, numAttrsPerJoint);
+    auto instanceFactory = [outputCount, rotationType](MemoryResource* instanceMemRes) {
+            return pma::UniqueInstance<CPUJointsOutputInstance, JointsOutputInstance>::with(instanceMemRes).create(
                 outputCount,
+                rl4::TranslationType::Vector,
+                rotationType,
+                rl4::ScaleType::Vector,
                 instanceMemRes);
         };
     return bpcm::Evaluator<TValue>{std::move(storage), std::move(strategy), instanceFactory, memRes};
@@ -917,15 +1334,21 @@ bpcm::Evaluator<TValue> OptimizedStorage<TValue>::create(StrategyPtr strategy, r
 
 template<typename TValue>
 bpcm::Evaluator<TValue> OptimizedStorage<TValue>::create(StrategyPtr strategy,
+                                                         std::size_t rotationSelectorIndex,
+                                                         rl4::RotationType rotationType,
                                                          std::uint16_t jointGroupIndex,
-                                                         rl4::MemoryResource* memRes) {
-    const auto& values = OptimizedValues<TValue>::get();
+                                                         MemoryResource* memRes) {
     bpcm::JointStorage<TValue> storage{memRes};
+    const auto& values = OptimizedValues<TValue>::get();
     storage.values.assign(values[jointGroupIndex].begin(), values[jointGroupIndex].end());
     storage.inputIndices.assign(optimized::inputIndices[jointGroupIndex].begin(),
                                 optimized::inputIndices[jointGroupIndex].end());
-    storage.outputIndices.assign(optimized::outputIndices[jointGroupIndex].begin(),
-                                 optimized::outputIndices[jointGroupIndex].end());
+    storage.outputIndices.assign(optimized::outputIndices[rotationSelectorIndex][jointGroupIndex].begin(),
+                                 optimized::outputIndices[rotationSelectorIndex][jointGroupIndex].end());
+    storage.outputRotationIndices.assign(optimized::outputRotationIndices[rotationSelectorIndex][jointGroupIndex].begin(),
+                                         optimized::outputRotationIndices[rotationSelectorIndex][jointGroupIndex].end());
+    storage.outputRotationLODs.assign(optimized::outputRotationLODs[rotationSelectorIndex][jointGroupIndex].begin(),
+                                      optimized::outputRotationLODs[rotationSelectorIndex][jointGroupIndex].end());
     storage.lodRegions.assign(optimized::lodRegions[jointGroupIndex].begin(),
                               optimized::lodRegions[jointGroupIndex].end());
     storage.jointGroups.push_back(bpcm::JointGroup{
@@ -933,24 +1356,31 @@ bpcm::Evaluator<TValue> OptimizedStorage<TValue>::create(StrategyPtr strategy,
             0u,  // inputIndicesOffset
             0u,  // outputIndicesOffset
             0u,  // lodsOffset
+            0u,  // outputRotationIndicesOffset
+            0u,  // outputRotationLODsOffset
             static_cast<std::uint32_t>(storage.values.size()),
             static_cast<std::uint32_t>(storage.inputIndices.size()),
             static_cast<std::uint32_t>(storage.inputIndices.size() - (storage.inputIndices.size() % 4ul)),
             static_cast<std::uint32_t>(storage.inputIndices.size() - (storage.inputIndices.size() % 8ul)),
         });
 
-    const auto outputCount = static_cast<std::uint16_t>(unoptimized::dimensions.rows);
-    auto instanceFactory = [outputCount](rl4::MemoryResource* instanceMemRes) {
-            return pma::UniqueInstance<rl4::bpcm::OutputInstance, rl4::JointsOutputInstance>::with(instanceMemRes).create(
+    const auto lastJointGroupIndex = optimized::outputIndices[rotationSelectorIndex].size() - 1ul;
+    const auto numAttrsPerJoint = static_cast<std::uint16_t>(static_cast<std::uint8_t>(rl4::TranslationType::Vector) +
+                                                             static_cast<std::uint8_t>(rotationType) +
+                                                             static_cast<std::uint8_t>(rl4::ScaleType::Vector));
+    const auto maxOutputIndex = extd::maxOf(optimized::outputIndices[rotationSelectorIndex][lastJointGroupIndex]);
+    const auto outputCount = extd::roundUp(maxOutputIndex, numAttrsPerJoint);
+    auto instanceFactory = [outputCount, rotationType](MemoryResource* instanceMemRes) {
+            return pma::UniqueInstance<CPUJointsOutputInstance, JointsOutputInstance>::with(instanceMemRes).create(
                 outputCount,
+                rl4::TranslationType::Vector,
+                rotationType,
+                rl4::ScaleType::Vector,
                 instanceMemRes);
         };
     return bpcm::Evaluator<TValue>{std::move(storage), std::move(strategy), instanceFactory, memRes};
 }
 
-template struct OptimizedStorage<float>;
-#ifdef RL_USE_HALF_FLOATS
-    template struct OptimizedStorage<std::uint16_t>;
-#endif
+template struct OptimizedStorage<StorageValueType>;
 
 }  // namespace block4
