Update to new RigLogic implementing body rig support

diff --git a/Engine/Plugins/Animation/RigLogic/Source/RigLogicLibTest/Private/rltests/joints/bpcm/BPCMFixturesBlock4.h b/Engine/Plugins/Animation/RigLogic/Source/RigLogicLibTest/Private/rltests/joints/bpcm/BPCMFixturesBlock4.h
--- a/Engine/Plugins/Animation/RigLogic/Source/RigLogicLibTest/Private/rltests/joints/bpcm/BPCMFixturesBlock4.h
+++ b/Engine/Plugins/Animation/RigLogic/Source/RigLogicLibTest/Private/rltests/joints/bpcm/BPCMFixturesBlock4.h
@@ -7,10 +7,9 @@
 #include "rltests/dna/FakeReader.h"
 
 #include "riglogic/TypeDefs.h"
-#include "riglogic/joints/bpcm/Evaluator.h"
-#include "riglogic/joints/bpcm/JointGroup.h"
-#include "riglogic/joints/bpcm/LODRegion.h"
-#include "riglogic/riglogic/RigLogic.h"
+#include "riglogic/joints/cpu/bpcm/BPCMJointsEvaluator.h"
+#include "riglogic/joints/cpu/bpcm/JointGroup.h"
+#include "riglogic/riglogic/Configuration.h"
 #include "riglogic/types/Extent.h"
 
 #ifdef _MSC_VER
@@ -36,7 +35,6 @@ extern const Matrix<float> values;
 extern const Matrix<std::uint16_t> inputIndices;
 extern const Matrix<std::uint16_t> outputIndices;
 extern const Matrix<std::uint16_t> lods;
-extern const Vector<Extent> subMatrices;
 extern const Vector<float> neutralJointTranslationXs;
 extern const Vector<float> neutralJointTranslationYs;
 extern const Vector<float> neutralJointTranslationZs;
@@ -52,8 +50,10 @@ extern const Extent dimensions;
 extern const AlignedMatrix<float> floatValues;
 extern const AlignedMatrix<std::uint16_t> halfFloatValues;
 extern const AlignedMatrix<std::uint16_t> inputIndices;
-extern const AlignedMatrix<std::uint16_t> outputIndices;
-extern const Vector<bpcm::JointGroup> jointGroups;
+extern const Vector<Vector<AlignedVector<std::uint16_t> > > outputIndices;
+extern const Vector<Vector<AlignedVector<std::uint16_t> > > outputRotationIndices;
+extern const Vector<Vector<AlignedVector<std::uint16_t> > > outputRotationLODs;
+extern const Vector<Vector<bpcm::JointGroup> > jointGroups;
 extern const Matrix<bpcm::LODRegion> lodRegions;
 
 }  // namespace optimized
@@ -68,7 +68,7 @@ extern const Vector<float> values;
 namespace output {
 
 // Calculation output values
-extern const Matrix<float> valuesPerLOD;
+extern const Vector<Matrix<float> > valuesPerLOD;
 
 }  // namespace output
 
@@ -117,7 +117,7 @@ class CanonicalReader : public dna::FakeReader {
         }
 
         std::uint16_t getJointGroupCount() const override {
-            return static_cast<std::uint16_t>(unoptimized::subMatrices.size());
+            return static_cast<std::uint16_t>(unoptimized::values.size());
         }
 
         ConstArrayView<std::uint16_t> getJointGroupLODs(std::uint16_t jointGroupIndex) const override {
@@ -140,11 +140,18 @@ class CanonicalReader : public dna::FakeReader {
 
 template<typename TValue>
 struct OptimizedStorage {
-    using StrategyPtr = std::unique_ptr<bpcm::JointCalculationStrategy<TValue>,
-                                        std::function<void (bpcm::JointCalculationStrategy<TValue>*)> >;
-
-    static bpcm::Evaluator<TValue> create(StrategyPtr strategy, rl4::MemoryResource* memRes);
-    static bpcm::Evaluator<TValue> create(StrategyPtr strategy, std::uint16_t jointGroupIndex, rl4::MemoryResource* memRes);
+    using StrategyPtr = std::unique_ptr<bpcm::JointGroupLinearCalculationStrategy<TValue>,
+                                        std::function<void (bpcm::JointGroupLinearCalculationStrategy<TValue>*)> >;
+
+    static bpcm::Evaluator<TValue> create(StrategyPtr strategy,
+                                          std::size_t rotationSelectorIndex,
+                                          rl4::RotationType rotationType,
+                                          MemoryResource* memRes);
+    static bpcm::Evaluator<TValue> create(StrategyPtr strategy,
+                                          std::size_t rotationSelectorIndex,
+                                          rl4::RotationType rotationType,
+                                          std::uint16_t jointGroupIndex,
+                                          MemoryResource* memRes);
 
 };
 
diff --git a/Engine/Plugins/Animation/RigLogic/Source/RigLogicLibTest/Private/rltests/joints/bpcm/BPCMFixturesBlock8.cpp b/Engine/Plugins/Animation/RigLogic/Source/RigLogicLibTest/Private/rltests/joints/bpcm/BPCMFixturesBlock8.cpp
--- a/Engine/Plugins/Animation/RigLogic/Source/RigLogicLibTest/Private/rltests/joints/bpcm/BPCMFixturesBlock8.cpp
+++ b/Engine/Plugins/Animation/RigLogic/Source/RigLogicLibTest/Private/rltests/joints/bpcm/BPCMFixturesBlock8.cpp
@@ -3,8 +3,8 @@
 #include "rltests/joints/bpcm/BPCMFixturesBlock8.h"
 
 #include "riglogic/TypeDefs.h"
-#include "riglogic/joints/bpcm/BPCMOutputInstance.h"
-#include "riglogic/joints/bpcm/Evaluator.h"
+#include "riglogic/joints/cpu/CPUJointsOutputInstance.h"
+#include "riglogic/joints/cpu/bpcm/BPCMJointsEvaluator.h"
 #include "riglogic/types/Extent.h"
 
 namespace block8 {
@@ -177,28 +177,28 @@ const Matrix<float> values = {
 
 const Matrix<std::uint16_t> inputIndices = {
     {  // Joint group 0
-        12, 11, 10, 9, 8, 7, 6, 5, 4, 3, 2, 1, 0,
+        12, 11, 10, 9, 8, 7, 6, 5, 4, 3, 2, 1, 0
     },
     {  // Joint group 1
-        12, 11, 10, 9, 8, 7, 6, 5, 4, 3, 2, 1, 0,
+        12, 11, 10, 9, 8, 7, 6, 5, 4, 3, 2, 1, 0
     },
     {  // Joint group 2
-        12, 11, 10, 9, 8, 7, 6, 5, 4, 3, 2, 1, 0,
+        12, 11, 10, 9, 8, 7, 6, 5, 4, 3, 2, 1, 0
     },
     {  // Joint group 3
-        12, 11, 10, 9, 8, 7, 6, 5, 4, 3, 2, 1, 0,
+        12, 11, 10, 9, 8, 7, 6, 5, 4, 3, 2, 1, 0
     },
     {  // Joint group 4
-        12, 11, 10, 9, 8, 7, 6, 5, 4, 3, 2, 1, 0,
+        12, 11, 10, 9, 8, 7, 6, 5, 4, 3, 2, 1, 0
     },
     {  // Joint group 5
-        12, 11, 10, 9, 8, 7, 6, 5, 4, 3, 2, 1, 0,
+        12, 11, 10, 9, 8, 7, 6, 5, 4, 3, 2, 1, 0
     },
     {  // Joint group 6
-        12, 11, 10, 9, 8, 7, 6, 5, 4, 3, 2, 1, 0,
+        12, 11, 10, 9, 8, 7, 6, 5, 4, 3, 2, 1, 0
     },
     {  // Joint group 7
-        12, 11, 10, 9, 8, 7, 6, 5, 4, 3, 2, 1, 0,
+        12, 11, 10, 9, 8, 7, 6, 5, 4, 3, 2, 1, 0
     },
 };
 
@@ -207,26 +207,26 @@ const Matrix<std::uint16_t> outputIndices = {
         0, 1, 2, 3, 4
     },
     {  // Joint group 1
-        5, 6, 7, 8, 9, 10, 11, 12
+        9, 10, 11, 12, 13, 14, 15, 16
     },
     {  // Joint group 2
-        13, 14, 15, 16, 17, 18, 19, 20, 21
+        27, 28, 29, 30, 31, 32, 33, 34, 35
     },
     {  // Joint group 3
-        22, 23, 24, 25, 26, 27, 28, 29, 30, 31, 32, 33, 34, 35, 36, 37
+        36, 37, 38, 39, 40, 41, 42, 43, 44, 45, 46, 47, 48, 49, 50, 51
     },
     {  // Joint group 4
-        38, 39, 40, 41, 42, 43, 44, 45, 46, 47, 48, 49, 50, 51, 52, 53, 54
+        54, 55, 56, 57, 58, 59, 60, 61, 62, 63, 64, 65, 66, 67, 68, 69, 70
     },
     {  // Joint group 5
-        55, 56, 57, 58, 59, 60, 61, 62, 63, 64, 65, 66, 67, 68, 69, 70, 71, 72, 73, 74, 75, 76, 77, 78
+        72, 73, 74, 75, 76, 77, 78, 79, 80, 81, 82, 83, 84, 85, 86, 87, 88, 89, 90, 91, 92, 93, 94, 95
     },
     {  // Joint group 6
-        79, 80, 81, 82, 83, 84, 85, 86, 87, 88, 89, 90, 91, 92, 93, 94, 95, 96, 97, 98, 99, 100, 101, 102, 103
+        99, 100, 101, 102, 103, 104, 105, 106, 107, 108, 109, 110, 111, 112, 113, 114, 115, 116, 117, 118, 119, 120, 121, 122, 123
     },
     {  // Joint group 7
-        104, 105, 106, 107, 108, 109, 110, 111, 112, 113, 114, 115, 116, 117, 118, 119,
-        120, 121, 122, 123, 124, 125, 126, 127, 128, 129, 130, 131, 132, 133, 134, 135
+        126, 127, 128, 129, 130, 131, 132, 133, 134, 135, 136, 137, 138, 139, 140, 141,
+        142, 143, 144, 145, 146, 147, 148, 149, 150, 151, 152, 153, 154, 155, 156, 157
     }
 };
 
@@ -641,85 +641,245 @@ const AlignedMatrix<std::uint16_t> halfFloatValues = {
 
 const AlignedMatrix<std::uint16_t> inputIndices = {
     {  // Joint group 0
-        12, 11, 10, 9, 8, 7, 6, 5, 4, 3, 2, 1, 0,
+        12, 11, 10, 9, 8, 7, 6, 5, 4, 3, 2, 1, 0
     },
     {  // Joint group 1
-        12, 11, 10, 9, 8, 7, 6, 5, 4, 3, 2, 1, 0,
+        12, 11, 10, 9, 8, 7, 6, 5, 4, 3, 2, 1, 0
     },
     {  // Joint group 2
-        12, 11, 10, 9, 8, 7, 6, 5, 4, 3, 2, 1, 0,
+        12, 11, 10, 9, 8, 7, 6, 5, 4, 3, 2, 1, 0
     },
     {  // Joint group 3
-        12, 11, 10, 9, 8, 7, 6, 5, 4, 3, 2, 1, 0,
+        12, 11, 10, 9, 8, 7, 6, 5, 4, 3, 2, 1, 0
     },
     {  // Joint group 4
-        12, 11, 10, 9, 8, 7, 6, 5, 4, 3, 2, 1, 0,
+        12, 11, 10, 9, 8, 7, 6, 5, 4, 3, 2, 1, 0
     },
     {  // Joint group 5
-        12, 11, 10, 9, 8, 7, 6, 5, 4, 3, 2, 1, 0,
+        12, 11, 10, 9, 8, 7, 6, 5, 4, 3, 2, 1, 0
     },
     {  // Joint group 6
-        12, 11, 10, 9, 8, 7, 6, 5, 4, 3, 2, 1, 0,
+        12, 11, 10, 9, 8, 7, 6, 5, 4, 3, 2, 1, 0
     },
     {  // Joint group 7
-        12, 11, 10, 9, 8, 7, 6, 5, 4, 3, 2, 1, 0,
+        12, 11, 10, 9, 8, 7, 6, 5, 4, 3, 2, 1, 0
     }
 };
 
-const AlignedMatrix<std::uint16_t> outputIndices = {
-    {  // Joint group 0
-        0, 1, 2, 3, 4, 0, 0, 0
-    },
-    {  // Joint group 1
-        5, 6, 7, 8, 9, 10, 11, 12
-    },
-    {  // Joint group 2
-        13, 14, 15, 16, 17, 18, 19, 20, 21, 0, 0, 0, 0, 0, 0, 0
-    },
-    {  // Joint group 3
-        22, 23, 24, 25, 26, 27, 28, 29, 30, 31, 32, 33, 34, 35, 36, 37
-    },
-    {  // Joint group 4
-        38, 39, 40, 41, 42, 43, 44, 45, 46, 47, 48, 49, 50, 51, 52, 53, 54, 0, 0, 0, 0, 0, 0, 0
-    },
-    {  // Joint group 5
-        55, 56, 57, 58, 59, 60, 61, 62, 63, 64, 65, 66, 67, 68, 69, 70, 71, 72, 73, 74, 75, 76, 77, 78
-    },
-    {  // Joint group 6
-        79, 80, 81, 82, 83, 84, 85, 86, 87, 88, 89, 90, 91, 92, 93, 94,
-        95, 96, 97, 98, 99, 100, 101, 102, 103, 0, 0, 0, 0, 0, 0, 0
-    },
-    {  // Joint group 7
-        104, 105, 106, 107, 108, 109, 110, 111, 112, 113, 114, 115, 116, 117, 118, 119,
-        120, 121, 122, 123, 124, 125, 126, 127, 128, 129, 130, 131, 132, 133, 134, 135
+const Vector<Vector<AlignedVector<std::uint16_t> > > outputIndices = {
+    {  // Quaternion outputs
+        {  // Joint group 0
+            0, 1, 2, 3, 4, 0, 0, 0
+        },
+        {  // Joint group 1
+            10, 11, 12, 13, 14, 15, 17, 18
+        },
+        {  // Joint group 2
+            30, 31, 32, 33, 34, 35, 37, 38, 39, 0, 0, 0, 0, 0, 0, 0
+        },
+        {  // Joint group 3
+            40, 41, 42, 43, 44, 45, 47, 48, 49,
+            50, 51, 52, 53, 54, 55, 57
+        },
+        {  // Joint group 4
+            60, 61, 62, 63, 64, 65, 67, 68, 69,
+            70, 71, 72, 73, 74, 75, 77, 78, 0, 0, 0, 0, 0, 0, 0
+        },
+        {  // Joint group 5
+            80, 81, 82, 83, 84, 85, 87, 88, 89,
+            90, 91, 92, 93, 94, 95, 97, 98, 99,
+            100, 101, 102, 103, 104, 105
+        },
+        {  // Joint group 6
+            110, 111, 112, 113, 114, 115, 117, 118, 119,
+            120, 121, 122, 123, 124, 125, 127, 128, 129,
+            130, 131, 132, 133, 134, 135, 137, 0, 0, 0, 0, 0, 0, 0
+        },
+        {  // Joint group 7
+            140, 141, 142, 143, 144, 145, 147, 148, 149,
+            150, 151, 152, 153, 154, 155, 157, 158, 159,
+            160, 161, 162, 163, 164, 165, 167, 168, 169,
+            170, 171, 172, 173, 174
+        }
+    },
+    {  // Euler-angle outputs
+        {  // Joint group 0
+            0, 1, 2, 3, 4, 0, 0, 0
+        },
+        {  // Joint group 1
+            9, 10, 11, 12, 13, 14, 15, 16
+        },
+        {  // Joint group 2
+            27, 28, 29, 30, 31, 32, 33, 34, 35, 0, 0, 0, 0, 0, 0, 0
+        },
+        {  // Joint group 3
+            36, 37, 38, 39, 40, 41, 42, 43, 44, 45, 46, 47, 48, 49, 50, 51
+        },
+        {  // Joint group 4
+            54, 55, 56, 57, 58, 59, 60, 61, 62, 63, 64, 65, 66, 67, 68, 69, 70, 0, 0, 0, 0, 0, 0, 0
+        },
+        {  // Joint group 5
+            72, 73, 74, 75, 76, 77, 78, 79, 80, 81, 82, 83, 84, 85, 86, 87, 88, 89, 90, 91, 92, 93, 94, 95
+        },
+        {  // Joint group 6
+            99, 100, 101, 102, 103, 104, 105, 106, 107, 108, 109, 110, 111, 112, 113, 114,
+            115, 116, 117, 118, 119, 120, 121, 122, 123, 0, 0, 0, 0, 0, 0, 0
+        },
+        {  // Joint group 7
+            126, 127, 128, 129, 130, 131, 132, 133, 134, 135, 136, 137, 138, 139, 140, 141,
+            142, 143, 144, 145, 146, 147, 148, 149, 150, 151, 152, 153, 154, 155, 156, 157
+        }
     }
 };
 
-const Vector<bpcm::JointGroup> jointGroups = {
-    // {valueOffset, inputOffset, outputOffset, lodOffset, valueSize, inputSize, inputSizeAligned4, inputSizeAligned8}
-    {  // Joint group 0
-        0, 0, 0, 0, 104, 13, 12, 8
-    },
-    {  // Joint group 1
-        104, 13, 8, 4, 104, 13, 12, 8
-    },
-    {  // Joint group 2
-        208, 26, 16, 8, 208, 13, 12, 8
-    },
-    {  // Joint group 3
-        416, 39, 32, 12, 208, 13, 12, 8
-    },
-    {  // Joint group 4
-        624, 52, 48, 16, 312, 13, 12, 8
-    },
-    {  // Joint group 5
-        936, 65, 72, 20, 312, 13, 12, 8
-    },
-    {  // Joint group 6
-        1248, 78, 96, 24, 416, 13, 12, 8
-    },
-    {  // Joint group 7
-        1664, 91, 128, 28, 416, 13, 12, 8
+const Vector<Vector<AlignedVector<std::uint16_t> > > outputRotationIndices = {
+    {  // Quaternion outputs
+        {  // Joint group 0
+            3
+        },
+        {  // Joint group 1
+            13
+        },
+        {  // Joint group 2
+            33
+        },
+        {  // Joint group 3
+            43, 53
+        },
+        {  // Joint group 4
+            63, 73
+        },
+        {  // Joint group 5
+            83, 93, 103
+        },
+        {  // Joint group 6
+            113, 123, 133
+        },
+        {  // Joint group 7
+            143, 153, 163, 173
+        }
+    },
+    {  // Euler-angle outputs
+        {  // Joint group 0
+        },
+        {  // Joint group 1
+        },
+        {  // Joint group 2
+        },
+        {  // Joint group 3
+        },
+        {  // Joint group 4
+        },
+        {  // Joint group 5
+        },
+        {  // Joint group 6
+        },
+        {  // Joint group 7
+        }
+    }
+};
+
+const Vector<Vector<AlignedVector<std::uint16_t> > > outputRotationLODs = {
+    {  // Quaternion outputs
+        {  // Joint group 0
+            1, 1, 0, 0
+        },
+        {  // Joint group 1
+            1, 1, 0, 0
+        },
+        {  // Joint group 2
+            1, 1, 1, 0
+        },
+        {  // Joint group 3
+            2, 2, 1, 1
+        },
+        {  // Joint group 4
+            2, 2, 2, 1
+        },
+        {  // Joint group 5
+            3, 2, 2, 1
+        },
+        {  // Joint group 6
+            3, 2, 1, 1
+        },
+        {  // Joint group 7
+            4, 3, 2, 1
+        }
+    },
+    {  // Euler-angle outputs
+        {  // Joint group 0
+        },
+        {  // Joint group 1
+        },
+        {  // Joint group 2
+        },
+        {  // Joint group 3
+        },
+        {  // Joint group 4
+        },
+        {  // Joint group 5
+        },
+        {  // Joint group 6
+        },
+        {  // Joint group 7
+        }
+    }
+};
+
+const Vector<Vector<bpcm::JointGroup> > jointGroups = {
+    // {valueOffset, inputOffset, outputOffset, lodOffset, outputRotationIndicesOffset, outputRotationLODsOffset, valueSize,
+    // inputSize, inputSizeAligned4, inputSizeAligned8}
+    {  // Quaternion outputs
+        {  // Joint group 0
+            0, 0, 0, 0, 0, 0, 104, 13, 12, 8
+        },
+        {  // Joint group 1
+            104, 13, 8, 4, 1, 4, 104, 13, 12, 8
+        },
+        {  // Joint group 2
+            208, 26, 16, 8, 2, 8, 208, 13, 12, 8
+        },
+        {  // Joint group 3
+            416, 39, 32, 12, 3, 12, 208, 13, 12, 8
+        },
+        {  // Joint group 4
+            624, 52, 48, 16, 5, 16, 312, 13, 12, 8
+        },
+        {  // Joint group 5
+            936, 65, 72, 20, 7, 20, 312, 13, 12, 8
+        },
+        {  // Joint group 6
+            1248, 78, 96, 24, 10, 24, 416, 13, 12, 8
+        },
+        {  // Joint group 7
+            1664, 91, 128, 28, 13, 28, 416, 13, 12, 8
+        }
+    },
+    {  // Euler-angle outputs
+        {  // Joint group 0
+            0, 0, 0, 0, 0, 0, 104, 13, 12, 8
+        },
+        {  // Joint group 1
+            104, 13, 8, 4, 0, 0, 104, 13, 12, 8
+        },
+        {  // Joint group 2
+            208, 26, 16, 8, 0, 0, 208, 13, 12, 8
+        },
+        {  // Joint group 3
+            416, 39, 32, 12, 0, 0, 208, 13, 12, 8
+        },
+        {  // Joint group 4
+            624, 52, 48, 16, 0, 0, 312, 13, 12, 8
+        },
+        {  // Joint group 5
+            936, 65, 72, 20, 0, 0, 312, 13, 12, 8
+        },
+        {  // Joint group 6
+            1248, 78, 96, 24, 0, 0, 416, 13, 12, 8
+        },
+        {  // Joint group 7
+            1664, 91, 128, 28, 0, 0, 416, 13, 12, 8
+        }
     }
 };
 
@@ -780,90 +940,199 @@ const Matrix<bpcm::LODRegion> lodRegions = {
 namespace input {
 
 // Calculation input values
-const rl4::Vector<float> values = {1.0f, 2.0f, 3.0f, 4.0f, 0.0f, 6.0f, 7.0f, 8.0f, 9.0f, 0.0f, 11.0f, 12.0f, 13.0f};
+const Vector<float> values = {1.0f, 2.0f, 3.0f, 4.0f, 0.0f, 6.0f, 7.0f, 8.0f, 9.0f, 0.0f, 11.0f, 12.0f, 13.0f};
 
 }  // namespace input
 
 namespace output {
 
 // Expected output results for each LOD
-const rl4::Matrix<float> valuesPerLOD = {
-    {
-        // LOD-0
-        76.0f, 152.0f, 228.0f, 304.0f, 380.0f,  // Joint group 0
-        456.0f, 532.0f, 608.0f, 684.0f, 760.0f, 836.0f, 912.0f, 988.0f,  // Joint group 1
-        1064.0f, 1140.0f, 1216.0f, 1292.0f, 1368.0f, 1444.0f, 1520.0f, 1596.0f, 1672.0f,  // Joint group 2
-        1748.0f, 1824.0f, 1900.0f, 1976.0f, 2052.0f, 2128.0f, 2204.0f, 2280.0f,  // Joint group 3
-        2356.0f, 2432.0f, 2508.0f, 2584.0f, 2660.0f, 2736.0f, 2812.0f, 2888.0f,  // Joint group 3
-        2964.0f, 3040.0f, 3116.0f, 3192.0f, 3268.0f, 3344.0f, 3420.0f, 3496.0f,  // Joint group 4
-        3572.0f, 3648.0f, 3724.0f, 3800.0f, 3876.0f, 3952.0f, 4028.0f, 4104.0f, 4180.0f,  // Joint group 4
-        4256.0f, 4332.0f, 4408.0f, 4484.0f, 4560.0f, 4636.0f, 4712.0f, 4788.0f,  // Joint group 5
-        4864.0f, 4940.0f, 5016.0f, 5092.0f, 5168.0f, 5244.0f, 5320.0f, 5396.0f,  // Joint group 5
-        5472.0f, 5548.0f, 5624.0f, 5700.0f, 5776.0f, 5852.0f, 5928.0f, 6004.0f,  // Joint group 5
-        6080.0f, 6156.0f, 6232.0f, 6308.0f, 6384.0f, 6460.0f, 6536.0f, 6612.0f,  // Joint group 6
-        6688.0f, 6764.0f, 6840.0f, 6916.0f, 6992.0f, 7068.0f, 7144.0f, 7220.0f,  // Joint group 6
-        7296.0f, 7372.0f, 7448.0f, 7524.0f, 7600.0f, 7676.0f, 7752.0f, 7828.0f, 7904.0f,  // Joint group 6
-        7980.0f, 8056.0f, 8132.0f, 8208.0f, 8284.0f, 8360.0f, 8436.0f, 8512.0f,  // Joint group 7
-        8588.0f, 8664.0f, 8740.0f, 8816.0f, 8892.0f, 8968.0f, 9044.0f, 9120.0f,  // Joint group 7
-        9196.0f, 9272.0f, 9348.0f, 9424.0f, 9500.0f, 9576.0f, 9652.0f, 9728.0f,  // Joint group 7
-        9804.0f, 9880.0f, 9956.0f, 10032.0f, 10108.0f, 10184.0f, 10260.0f, 10336.0f  // Joint group 7
-    }, {
-        // LOD-1
-        76.0f, 152.0f, 228.0f, 304.0f, 0.0f,  // Joint group 0
-        456.0f, 532.0f, 608.0f, 684.0f, 760.0f, 0.0f, 0.0f, 0.0f,  // Joint group 1
-        1064.0f, 1140.0f, 1216.0f, 1292.0f, 1368.0f, 1444.0f, 1520.0f, 0.0f, 0.0f,  // Joint group 2
-        1748.0f, 1824.0f, 1900.0f, 1976.0f, 2052.0f, 2128.0f, 2204.0f, 2280.0f,  // Joint group 3
-        2356.0f, 2432.0f, 2508.0f, 2584.0f, 2660.0f, 0.0f, 0.0f, 0.0f,  // Joint group 3
-        2964.0f, 3040.0f, 3116.0f, 3192.0f, 3268.0f, 3344.0f, 3420.0f, 3496.0f,  // Joint group 4
-        3572.0f, 3648.0f, 3724.0f, 3800.0f, 3876.0f, 3952.0f, 4028.0f, 4104.0f, 0.0f,  // Joint group 4
-        4256.0f, 4332.0f, 4408.0f, 4484.0f, 4560.0f, 4636.0f, 4712.0f, 4788.0f,  // Joint group 5
-        4864.0f, 4940.0f, 5016.0f, 5092.0f, 5168.0f, 5244.0f, 5320.0f, 5396.0f,  // Joint group 5
-        0.0f, 0.0f, 0.0f, 0.0f, 0.0f, 0.0f, 0.0f, 0.0f,  // Joint group 5
-        6080.0f, 6156.0f, 6232.0f, 6308.0f, 6384.0f, 6460.0f, 6536.0f, 6612.0f,  // Joint group 6
-        6688.0f, 6764.0f, 6840.0f, 6916.0f, 6992.0f, 7068.0f, 7144.0f, 7220.0f,  // Joint group 6
-        7296.0f, 0.0f, 0.0f, 0.0f, 0.0f, 0.0f, 0.0f, 0.0f, 0.0f,  // Joint group 6
-        7980.0f, 8056.0f, 8132.0f, 8208.0f, 8284.0f, 8360.0f, 8436.0f, 8512.0f,  // Joint group 7
-        8588.0f, 8664.0f, 8740.0f, 8816.0f, 8892.0f, 8968.0f, 9044.0f, 9120.0f,  // Joint group 7
-        9196.0f, 9272.0f, 9348.0f, 9424.0f, 9500.0f, 9576.0f, 9652.0f, 9728.0f,  // Joint group 7
-        9804.0f, 0.0f, 0.0f, 0.0f, 0.0f, 0.0f, 0.0f, 0.0f  // Joint group 7
-    }, {
-        // LOD-2
-        76.0f, 152.0f, 0.0f, 0.0f, 0.0f,  // Joint group 0
-        456.0f, 532.0f, 608.0f, 0.0f, 0.0f, 0.0f, 0.0f, 0.0f,  // Joint group 1
-        1064.0f, 1140.0f, 1216.0f, 1292.0f, 0.0f, 0.0f, 0.0f, 0.0f, 0.0f,  // Joint group 2
-        1748.0f, 1824.0f, 1900.0f, 1976.0f, 2052.0f, 2128.0f, 2204.0f, 2280.0f,  // Joint group 3
-        0.0f, 0.0f, 0.0f, 0.0f, 0.0f, 0.0f, 0.0f, 0.0f,  // Joint group 3
-        2964.0f, 3040.0f, 3116.0f, 3192.0f, 3268.0f, 3344.0f, 3420.0f, 3496.0f,  // Joint group 4
-        3572.0f, 3648.0f, 3724.0f, 3800.0f, 3876.0f, 3952.0f, 4028.0f, 0.0f, 0.0f,  // Joint group 4
-        4256.0f, 4332.0f, 4408.0f, 4484.0f, 4560.0f, 4636.0f, 4712.0f, 4788.0f,  // Joint group 5
-        4864.0f, 4940.0f, 5016.0f, 5092.0f, 5168.0f, 5244.0f, 5320.0f, 0.0f,  // Joint group 5
-        0.0f, 0.0f, 0.0f, 0.0f, 0.0f, 0.0f, 0.0f, 0.0f,  // Joint group 5
-        6080.0f, 6156.0f, 6232.0f, 6308.0f, 6384.0f, 6460.0f, 6536.0f, 6612.0f,  // Joint group 6
-        6688.0f, 6764.0f, 6840.0f, 6916.0f, 0.0f, 0.0f, 0.0f, 0.0f,  // Joint group 6
-        0.0f, 0.0f, 0.0f, 0.0f, 0.0f, 0.0f, 0.0f, 0.0f, 0.0f,  // Joint group 6
-        7980.0f, 8056.0f, 8132.0f, 8208.0f, 8284.0f, 8360.0f, 8436.0f, 8512.0f,  // Joint group 7
-        8588.0f, 8664.0f, 8740.0f, 8816.0f, 8892.0f, 8968.0f, 9044.0f, 9120.0f,  // Joint group 7
-        9196.0f, 0.0f, 0.0f, 0.0f, 0.0f, 0.0f, 0.0f, 0.0f,  // Joint group 7
-        0.0f, 0.0f, 0.0f, 0.0f, 0.0f, 0.0f, 0.0f, 0.0f  // Joint group 7
-    }, {
-        // LOD-3
-        76.0f, 0.0f, 0.0f, 0.0f, 0.0f,  // Joint group 0
-        456.0f, 0.0f, 0.0f, 0.0f, 0.0f, 0.0f, 0.0f, 0.0f,  // Joint group 1
-        1064.0f, 1140.0f, 0.0f, 0.0f, 0.0f, 0.0f, 0.0f, 0.0f, 0.0f,  // Joint group 2
-        1748.0f, 1824.0f, 1900.0f, 1976.0f, 0.0f, 0.0f, 0.0f, 0.0f,  // Joint group 3
-        0.0f, 0.0f, 0.0f, 0.0f, 0.0f, 0.0f, 0.0f, 0.0f,  // Joint group 3
-        2964.0f, 3040.0f, 3116.0f, 3192.0f, 3268.0f, 3344.0f, 3420.0f, 3496.0f,  // Joint group 4
-        3572.0f, 3648.0f, 3724.0f, 3800.0f, 0.0f, 0.0f, 0.0f, 0.0f, 0.0f,  // Joint group 4
-        4256.0f, 4332.0f, 4408.0f, 4484.0f, 4560.0f, 4636.0f, 4712.0f, 4788.0f,  // Joint group 5
-        0.0f, 0.0f, 0.0f, 0.0f, 0.0f, 0.0f, 0.0f, 0.0f,  // Joint group 5
-        0.0f, 0.0f, 0.0f, 0.0f, 0.0f, 0.0f, 0.0f, 0.0f,  // Joint group 5
-        6080.0f, 6156.0f, 6232.0f, 6308.0f, 6384.0f, 6460.0f, 6536.0f, 0.0f,  // Joint group 6
-        0.0f, 0.0f, 0.0f, 0.0f, 0.0f, 0.0f, 0.0f, 0.0f,  // Joint group 6
-        0.0f, 0.0f, 0.0f, 0.0f, 0.0f, 0.0f, 0.0f, 0.0f, 0.0f,  // Joint group 6
-        7980.0f, 8056.0f, 8132.0f, 8208.0f, 8284.0f, 8360.0f, 8436.0f, 8512.0f,  // Joint group 7
-        8588.0f, 0.0f, 0.0f, 0.0f, 0.0f, 0.0f, 0.0f, 0.0f,  // Joint group 7
-        0.0f, 0.0f, 0.0f, 0.0f, 0.0f, 0.0f, 0.0f, 0.0f,  // Joint group 7
-        0.0f, 0.0f, 0.0f, 0.0f, 0.0f, 0.0f, 0.0f, 0.0f  // Joint group 7
+const Vector<Matrix<float> > valuesPerLOD = {
+    {  // Quaternion outputs
+        {
+            // LOD-0
+            76.0, 152.0f, 228.0f, -0.4623392f, 0.1533222f, -0.0815229f, 0.8695336f, 0.0f, 0.0f, 0.0f,
+            456.0f, 532.0f, 608.0f, 0.121975f, 0.4186294f, 0.7018941f, 0.5632195f, 912.0f, 988.0f, 0.0f,
+            0.0f, 0.0f, 0.0f, 0.0f, 0.0f, 0.0f, 1.0f, 0.0f, 0.0f, 0.0f,
+            1064.0f, 1140.0f, 1216.0f, -0.7828576f, -0.1347763f, 0.5724537f, 0.2031407f, 1520.0f, 1596.0f, 1672.0f,
+            1748.0f, 1824.0f, 1900.0f, -0.5724537f, -0.1347763f, 0.7828576f, 0.2031407f, 2204.0f, 2280.0f, 2356.0f,
+            2432.0f, 2508.0f, 2584.0f, -0.4226313f, 0.1072908f, 0.8000127f, 0.4121301f, 2888.0f, 0.0f, 0.0f,
+            2964.0f, 3040.0f, 3116.0f, 0.0688182f, -0.4470575f, -0.6379199f, -0.6232671f, 3420.0f, 3496.0f, 3572.0f,
+            3648.0f, 3724.0f, 3800.0f, 0.5823712f, -0.330285f, -0.4532453f, -0.5884932f, 4104.0f, 4180.0f, 0.0f,
+            4256.0f, 4332.0f, 4408.0f, 0.5042307f, 0.0737294f, -0.821218f, -0.2567421f, 4712.0f, 4788.0f, 4864.0f,
+            4940.0f, 5016.0f, 5092.0f, 0.4187725f, -0.2083559f, -0.7320557f, -0.4952896f, 5396.0f, 5472.0f, 5548.0f,
+            5624.0f, 5700.0f, 5776.0f, 0.5460311f, -0.4298623f, -0.3175257f, -0.6451712f, 0.0f, 0.0f, 0.0f,
+            6080.0f, 6156.0f, 6232.0f, -0.648358f, -0.1669619f, 0.7219668f, 0.1746985f, 6536.0f, 6612.0f, 6688.0f,
+            6764.0f, 6840.0f, 6916.0f, -0.4522159f, 0.0098581f, 0.8288482f, 0.3292632f, 7220.0f, 7296.0f, 7372.0f,
+            7448.0f, 7524.0f, 7600.0f, -0.4406819f, 0.3023303f, 0.625722f, 0.5682146f, 7904.0f, 0.0f, 0.0f,
+            7980.0f, 8056.0f, 8132.0f, 0.4318467f, -0.3992196f, -0.5039951f, -0.6325512f, 8436.0f, 8512.0f, 8588.0f,
+            8664.0f, 8740.0f, 8816.0f, 0.7817417f, -0.1411047f, -0.4184162f, -0.4403377f, 9120.0f, 9196.0f, 9272.0f,
+            9348.0f, 9424.0f, 9500.0f, 0.7985949f, 0.1173789f, -0.5483992f, -0.2184645f, 9804.0f, 9880.0f, 9956.0f,
+            10032.0f, 10108.0f, 10184.0f, -0.6156615f, 0.0f, 0.7880108f, 0.0f
+        }, {
+            // LOD-1
+            76.0f, 152.0f, 228.0f, 0.4694716f, 0.0f, 0.0f, -0.8829476f, 0.0f, 0.0f, 0.0f,
+            456.0f, 532.0f, 608.0f, -0.290381f, 0.3252805f, -0.10569f, 0.8937008f, 0.0f, 0.0f, 0.0f,
+            0.0f, 0.0f, 0.0f, 0.0f, 0.0f, 0.0f, 1.0f, 0.0f, 0.0f, 0.0f,
+            1064.0f, 1140.0f, 1216.0f, -0.7828576f, -0.1347763f, 0.5724537f, 0.2031407f, 1520.0f, 0.0f, 0.0f,
+            1748.0f, 1824.0f, 1900.0f, -0.5724537f, -0.1347763f, 0.7828576f, 0.2031407f, 2204.0f, 2280.0f, 2356.0f,
+            2432.0f, 2508.0f, 2584.0f, -0.9396926f, 0.0f, 0.0f, -0.3420201f, 0.0f, 0.0f, 0.0f,
+            2964.0f, 3040.0f, 3116.0f, 0.0688182f, -0.4470575f, -0.6379199f, -0.6232671f, 3420.0f, 3496.0f, 3572.0f,
+            3648.0f, 3724.0f, 3800.0f, 0.5823712f, -0.330285f, -0.4532453f, -0.5884932f, 4104.0f, 0.0f, 0.0f,
+            4256.0f, 4332.0f, 4408.0f, 0.5042307f, 0.0737294f, -0.821218f, -0.2567421f, 4712.0f, 4788.0f, 4864.0f,
+            4940.0f, 5016.0f, 5092.0f, 0.4187725f, -0.2083559f, -0.7320557f, -0.4952896f, 5396.0f, 0.0f, 0.0f,
+            0.0f, 0.0f, 0.0f, 0.0f, 0.0f, 0.0f, 1.0f, 0.0f, 0.0f, 0.0f,
+            6080.0f, 6156.0f, 6232.0f, -0.648358f, -0.1669619f, 0.7219668f, 0.1746985f, 6536.0f, 6612.0f, 6688.0f,
+            6764.0f, 6840.0f, 6916.0f, -0.4522159f, 0.0098581f, 0.8288482f, 0.3292632f, 7220.0f, 7296.0f, 0.0f,
+            0.0f, 0.0f, 0.0f, 0.0f, 0.0f, 0.0f, 1.0f, 0.0f, 0.0f, 0.0f,
+            7980.0f, 8056.0f, 8132.0f, 0.4318467f, -0.3992196f, -0.5039951f, -0.6325512f, 8436.0f, 8512.0f, 8588.0f,
+            8664.0f, 8740.0f, 8816.0f, 0.7817417f, -0.1411047f, -0.4184162f, -0.4403377f, 9120.0f, 9196.0f, 9272.0f,
+            9348.0f, 9424.0f, 9500.0f, 0.7985949f, 0.1173789f, -0.5483992f, -0.2184645f, 9804.0f, 0.0f, 0.0f,
+            0.0f, 0.0f, 0.0f, 0.0f, 0.0f, 0.0f, 1.0f
+        }, {
+            // LOD-2
+            76.0f, 152.0f, 0.0f, 0.0f, 0.0f, 0.0f, 1.0f, 0.0f, 0.0f, 0.0,
+            456.0f, 532.0f, 608.0f, 0.0f, 0.0f, 0.0f, 1.0f, 0.0f, 0.0f, 0.0,
+            0.0f, 0.0f, 0.0f, 0.0f, 0.0f, 0.0f, 1.0f, 0.0f, 0.0f, 0.0,
+            1064.0f, 1140.0f, 1216.0f, -0.9612617f, 0.0f, 0.0f, 0.2756374f, 0.0f, 0.0f, 0.0,
+            1748.0f, 1824.0f, 1900.0f, -0.5724537f, -0.1347763f, 0.7828576f, 0.2031407f, 2204.0f, 2280.0f, 0.0,
+            0.0f, 0.0f, 0.0f, 0.0f, 0.0f, 0.0f, 1.0f, 0.0f, 0.0f, 0.0,
+            2964.0f, 3040.0f, 3116.0f, 0.0688182f, -0.4470575f, -0.6379199f, -0.6232671f, 3420.0f, 3496.0f, 3572.0,
+            3648.0f, 3724.0f, 3800.0f, 0.5823712f, -0.330285f, -0.4532453f, -0.5884932f, 0.0f, 0.0f, 0.0,
+            4256.0, 4332.0f, 4408.0f, 0.5042307f, 0.0737294f, -0.821218f, -0.2567421f, 4712.0f, 4788.0f, 4864.0,
+            4940.0f, 5016.0f, 5092.0f, 0.4187725f, -0.2083559f, -0.7320557f, -0.4952896f, 0.0f, 0.0f, 0.0,
+            0.0f, 0.0f, 0.0f, 0.0f, 0.0f, 0.0f, 1.0f, 0.0f, 0.0f, 0.0,
+            6080.0f, 6156.0f, 6232.0f, -0.648358f, -0.1669619f, 0.7219668f, 0.1746985f, 6536.0f, 6612.0f, 6688.0,
+            6764.0f, 6840.0f, 6916.0f, 0.0f, 0.0f, 0.0f, 1.0f, 0.0f, 0.0f, 0.0,
+            0.0f, 0.0f, 0.0f, 0.0f, 0.0f, 0.0f, 1.0f, 0.0f, 0.0f, 0.0,
+            7980.0f, 8056.0f, 8132.0f, 0.4318467f, -0.3992196f, -0.5039951f, -0.6325512f, 8436.0f, 8512.0f, 8588.0,
+            8664.0f, 8740.0f, 8816.0f, 0.7817417f, -0.1411047f, -0.4184162f, -0.4403377f, 9120.0f, 9196.0f, 0.0,
+            0.0f, 0.0f, 0.0f, 0.0f, 0.0f, 0.0f, 1.0f, 0.0f, 0.0f, 0.0,
+            0.0f, 0.0f, 0.0f, 0.0f, 0.0f, 0.0f, 1.0f
+        }, {
+            // LOD-3
+            76.0f, 0.0f, 0.0f, 0.0f, 0.0f, 0.0f, 1.0f, 0.0f, 0.0f, 0.0f,
+            456.0f, 0.0f, 0.0f, 0.0f, 0.0f, 0.0f, 1.0f, 0.0f, 0.0f, 0.0f,
+            0.0f, 0.0f, 0.0f, 0.0f, 0.0f, 0.0f, 1.0f, 0.0f, 0.0f, 0.0f,
+            1064.0f, 1140.0f, 0.0f, 0.0f, 0.0f, 0.0f, 1.0f, 0.0f, 0.0f, 0.0f,
+            1748.0f, 1824.0f, 1900.0f, -0.9993908f, 0.0f, 0.0f, -0.0348995f, 0.0f, 0.0f, 0.0f,
+            0.0f, 0.0f, 0.0f, 0.0f, 0.0f, 0.0f, 1.0f, 0.0f, 0.0f, 0.0f,
+            2964.0f, 3040.0f, 3116.0f, 0.0688182f, -0.4470575f, -0.6379199f, -0.6232671f, 3420.0f, 3496.0f, 3572.0f,
+            3648.0f, 3724.0f, 3800.0f, 0.0f, 0.0f, 0.0f, 1.0f, 0.0f, 0.0f, 0.0f,
+            4256.0f, 4332.0f, 4408.0f, 0.5042307f, 0.0737294f, -0.821218f, -0.2567421f, 4712.0f, 4788.0f, 0.0f,
+            0.0f, 0.0f, 0.0f, 0.0f, 0.0f, 0.0f, 1.0f, 0.0f, 0.0f, 0.0f,
+            0.0f, 0.0f, 0.0f, 0.0f, 0.0f, 0.0f, 1.0f, 0.0f, 0.0f, 0.0f,
+            6080.0f, 6156.0f, 6232.0f, -0.648358f, -0.1669619f, 0.7219668f, 0.1746985f, 6536.0f, 0.0f, 0.0f,
+            0.0f, 0.0f, 0.0f, 0.0f, 0.0f, 0.0f, 1.0f, 0.0f, 0.0f, 0.0f,
+            0.0f, 0.0f, 0.0f, 0.0f, 0.0f, 0.0f, 1.0f, 0.0f, 0.0f, 0.0f,
+            7980.0f, 8056.0f, 8132.0f, 0.4318467f, -0.3992196f, -0.5039951f, -0.6325512f, 8436.0f, 8512.0f, 8588.0f,
+            0.0f, 0.0f, 0.0f, 0.0f, 0.0f, 0.0f, 1.0f, 0.0f, 0.0f, 0.0f,
+            0.0f, 0.0f, 0.0f, 0.0f, 0.0f, 0.0f, 1.0f, 0.0f, 0.0f, 0.0f,
+            0.0f, 0.0f, 0.0f, 0.0f, 0.0f, 0.0f, 1.0f
+        }
+    },
+    {  // Euler-angle outputs
+        {
+            // LOD-0
+            76.0f, 152.0f, 228.0f, 304.0f, 380.0f,  // Joint group 0
+            0.0f, 0.0f, 0.0f, 0.0f,
+            456.0f, 532.0f, 608.0f, 684.0f, 760.0f, 836.0f, 912.0f, 988.0f,  // Joint group 1
+            0.0f, 0.0f, 0.0f, 0.0f, 0.0f, 0.0f, 0.0f, 0.0f, 0.0f, 0.0f,
+            1064.0f, 1140.0f, 1216.0f, 1292.0f, 1368.0f, 1444.0f, 1520.0f, 1596.0f, 1672.0f,  // Joint group 2
+            1748.0f, 1824.0f, 1900.0f, 1976.0f, 2052.0f, 2128.0f, 2204.0f, 2280.0f,  // Joint group 3
+            2356.0f, 2432.0f, 2508.0f, 2584.0f, 2660.0f, 2736.0f, 2812.0f, 2888.0f,  // Joint group 3
+            0.0f, 0.0f,
+            2964.0f, 3040.0f, 3116.0f, 3192.0f, 3268.0f, 3344.0f, 3420.0f, 3496.0f,  // Joint group 4
+            3572.0f, 3648.0f, 3724.0f, 3800.0f, 3876.0f, 3952.0f, 4028.0f, 4104.0f, 4180.0f,  // Joint group 4
+            0.0f,
+            4256.0f, 4332.0f, 4408.0f, 4484.0f, 4560.0f, 4636.0f, 4712.0f, 4788.0f,  // Joint group 5
+            4864.0f, 4940.0f, 5016.0f, 5092.0f, 5168.0f, 5244.0f, 5320.0f, 5396.0f,  // Joint group 5
+            5472.0f, 5548.0f, 5624.0f, 5700.0f, 5776.0f, 5852.0f, 5928.0f, 6004.0f,  // Joint group 5
+            0.0f, 0.0f, 0.0f,
+            6080.0f, 6156.0f, 6232.0f, 6308.0f, 6384.0f, 6460.0f, 6536.0f, 6612.0f,  // Joint group 6
+            6688.0f, 6764.0f, 6840.0f, 6916.0f, 6992.0f, 7068.0f, 7144.0f, 7220.0f,  // Joint group 6
+            7296.0f, 7372.0f, 7448.0f, 7524.0f, 7600.0f, 7676.0f, 7752.0f, 7828.0f, 7904.0f,  // Joint group 6
+            0.0f, 0.0f,
+            7980.0f, 8056.0f, 8132.0f, 8208.0f, 8284.0f, 8360.0f, 8436.0f, 8512.0f,  // Joint group 7
+            8588.0f, 8664.0f, 8740.0f, 8816.0f, 8892.0f, 8968.0f, 9044.0f, 9120.0f,  // Joint group 7
+            9196.0f, 9272.0f, 9348.0f, 9424.0f, 9500.0f, 9576.0f, 9652.0f, 9728.0f,  // Joint group 7
+            9804.0f, 9880.0f, 9956.0f, 10032.0f, 10108.0f, 10184.0f, 10260.0f, 10336.0f  // Joint group 7
+        }, {
+            // LOD-1
+            76.0f, 152.0f, 228.0f, 304.0f, 0.0f,  // Joint group 0
+            0.0f, 0.0f, 0.0f, 0.0f,
+            456.0f, 532.0f, 608.0f, 684.0f, 760.0f, 0.0f, 0.0f, 0.0f,  // Joint group 1
+            0.0f, 0.0f, 0.0f, 0.0f, 0.0f, 0.0f, 0.0f, 0.0f, 0.0f, 0.0f,
+            1064.0f, 1140.0f, 1216.0f, 1292.0f, 1368.0f, 1444.0f, 1520.0f, 0.0f, 0.0f,  // Joint group 2
+            1748.0f, 1824.0f, 1900.0f, 1976.0f, 2052.0f, 2128.0f, 2204.0f, 2280.0f,  // Joint group 3
+            2356.0f, 2432.0f, 2508.0f, 2584.0f, 2660.0f, 0.0f, 0.0f, 0.0f,  // Joint group 3
+            0.0f, 0.0f,
+            2964.0f, 3040.0f, 3116.0f, 3192.0f, 3268.0f, 3344.0f, 3420.0f, 3496.0f,  // Joint group 4
+            3572.0f, 3648.0f, 3724.0f, 3800.0f, 3876.0f, 3952.0f, 4028.0f, 4104.0f, 0.0f,  // Joint group 4
+            0.0f,
+            4256.0f, 4332.0f, 4408.0f, 4484.0f, 4560.0f, 4636.0f, 4712.0f, 4788.0f,  // Joint group 5
+            4864.0f, 4940.0f, 5016.0f, 5092.0f, 5168.0f, 5244.0f, 5320.0f, 5396.0f,  // Joint group 5
+            0.0f, 0.0f, 0.0f, 0.0f, 0.0f, 0.0f, 0.0f, 0.0f,  // Joint group 5
+            0.0f, 0.0f, 0.0f,
+            6080.0f, 6156.0f, 6232.0f, 6308.0f, 6384.0f, 6460.0f, 6536.0f, 6612.0f,  // Joint group 6
+            6688.0f, 6764.0f, 6840.0f, 6916.0f, 6992.0f, 7068.0f, 7144.0f, 7220.0f,  // Joint group 6
+            7296.0f, 0.0f, 0.0f, 0.0f, 0.0f, 0.0f, 0.0f, 0.0f, 0.0f,  // Joint group 6
+            0.0f, 0.0f,
+            7980.0f, 8056.0f, 8132.0f, 8208.0f, 8284.0f, 8360.0f, 8436.0f, 8512.0f,  // Joint group 7
+            8588.0f, 8664.0f, 8740.0f, 8816.0f, 8892.0f, 8968.0f, 9044.0f, 9120.0f,  // Joint group 7
+            9196.0f, 9272.0f, 9348.0f, 9424.0f, 9500.0f, 9576.0f, 9652.0f, 9728.0f,  // Joint group 7
+            9804.0f, 0.0f, 0.0f, 0.0f, 0.0f, 0.0f, 0.0f, 0.0f  // Joint group 7
+        }, {
+            // LOD-2
+            76.0f, 152.0f, 0.0f, 0.0f, 0.0f,  // Joint group 0
+            0.0f, 0.0f, 0.0f, 0.0f,
+            456.0f, 532.0f, 608.0f, 0.0f, 0.0f, 0.0f, 0.0f, 0.0f,  // Joint group 1
+            0.0f, 0.0f, 0.0f, 0.0f, 0.0f, 0.0f, 0.0f, 0.0f, 0.0f, 0.0f,
+            1064.0f, 1140.0f, 1216.0f, 1292.0f, 0.0f, 0.0f, 0.0f, 0.0f, 0.0f,  // Joint group 2
+            1748.0f, 1824.0f, 1900.0f, 1976.0f, 2052.0f, 2128.0f, 2204.0f, 2280.0f,  // Joint group 3
+            0.0f, 0.0f, 0.0f, 0.0f, 0.0f, 0.0f, 0.0f, 0.0f,  // Joint group 3
+            0.0f, 0.0f,
+            2964.0f, 3040.0f, 3116.0f, 3192.0f, 3268.0f, 3344.0f, 3420.0f, 3496.0f,  // Joint group 4
+            3572.0f, 3648.0f, 3724.0f, 3800.0f, 3876.0f, 3952.0f, 4028.0f, 0.0f, 0.0f,  // Joint group 4
+            0.0f,
+            4256.0f, 4332.0f, 4408.0f, 4484.0f, 4560.0f, 4636.0f, 4712.0f, 4788.0f,  // Joint group 5
+            4864.0f, 4940.0f, 5016.0f, 5092.0f, 5168.0f, 5244.0f, 5320.0f, 0.0f,  // Joint group 5
+            0.0f, 0.0f, 0.0f, 0.0f, 0.0f, 0.0f, 0.0f, 0.0f,  // Joint group 5
+            0.0f, 0.0f, 0.0f,
+            6080.0f, 6156.0f, 6232.0f, 6308.0f, 6384.0f, 6460.0f, 6536.0f, 6612.0f,  // Joint group 6
+            6688.0f, 6764.0f, 6840.0f, 6916.0f, 0.0f, 0.0f, 0.0f, 0.0f,  // Joint group 6
+            0.0f, 0.0f, 0.0f, 0.0f, 0.0f, 0.0f, 0.0f, 0.0f, 0.0f,  // Joint group 6
+            0.0f, 0.0f,
+            7980.0f, 8056.0f, 8132.0f, 8208.0f, 8284.0f, 8360.0f, 8436.0f, 8512.0f,  // Joint group 7
+            8588.0f, 8664.0f, 8740.0f, 8816.0f, 8892.0f, 8968.0f, 9044.0f, 9120.0f,  // Joint group 7
+            9196.0f, 0.0f, 0.0f, 0.0f, 0.0f, 0.0f, 0.0f, 0.0f,  // Joint group 7
+            0.0f, 0.0f, 0.0f, 0.0f, 0.0f, 0.0f, 0.0f, 0.0f  // Joint group 7
+        }, {
+            // LOD-3
+            76.0f, 0.0f, 0.0f, 0.0f, 0.0f,  // Joint group 0
+            0.0f, 0.0f, 0.0f, 0.0f,
+            456.0f, 0.0f, 0.0f, 0.0f, 0.0f, 0.0f, 0.0f, 0.0f,  // Joint group 1
+            0.0f, 0.0f, 0.0f, 0.0f, 0.0f, 0.0f, 0.0f, 0.0f, 0.0f, 0.0f,
+            1064.0f, 1140.0f, 0.0f, 0.0f, 0.0f, 0.0f, 0.0f, 0.0f, 0.0f,  // Joint group 2
+            1748.0f, 1824.0f, 1900.0f, 1976.0f, 0.0f, 0.0f, 0.0f, 0.0f,  // Joint group 3
+            0.0f, 0.0f, 0.0f, 0.0f, 0.0f, 0.0f, 0.0f, 0.0f,  // Joint group 3
+            0.0f, 0.0f,
+            2964.0f, 3040.0f, 3116.0f, 3192.0f, 3268.0f, 3344.0f, 3420.0f, 3496.0f,  // Joint group 4
+            3572.0f, 3648.0f, 3724.0f, 3800.0f, 0.0f, 0.0f, 0.0f, 0.0f, 0.0f,  // Joint group 4
+            0.0f,
+            4256.0f, 4332.0f, 4408.0f, 4484.0f, 4560.0f, 4636.0f, 4712.0f, 4788.0f,  // Joint group 5
+            0.0f, 0.0f, 0.0f, 0.0f, 0.0f, 0.0f, 0.0f, 0.0f,  // Joint group 5
+            0.0f, 0.0f, 0.0f, 0.0f, 0.0f, 0.0f, 0.0f, 0.0f,  // Joint group 5
+            0.0f, 0.0f, 0.0f,
+            6080.0f, 6156.0f, 6232.0f, 6308.0f, 6384.0f, 6460.0f, 6536.0f, 0.0f,  // Joint group 6
+            0.0f, 0.0f, 0.0f, 0.0f, 0.0f, 0.0f, 0.0f, 0.0f,  // Joint group 6
+            0.0f, 0.0f, 0.0f, 0.0f, 0.0f, 0.0f, 0.0f, 0.0f, 0.0f,  // Joint group 6
+            0.0f, 0.0f,
+            7980.0f, 8056.0f, 8132.0f, 8208.0f, 8284.0f, 8360.0f, 8436.0f, 8512.0f,  // Joint group 7
+            8588.0f, 0.0f, 0.0f, 0.0f, 0.0f, 0.0f, 0.0f, 0.0f,  // Joint group 7
+            0.0f, 0.0f, 0.0f, 0.0f, 0.0f, 0.0f, 0.0f, 0.0f,  // Joint group 7
+            0.0f, 0.0f, 0.0f, 0.0f, 0.0f, 0.0f, 0.0f, 0.0f  // Joint group 7
+        }
     }
 };
 
@@ -879,23 +1148,26 @@ template<typename TValue>
 struct OptimizedValues;
 
 template<>
-struct OptimizedValues<std::uint16_t> {
-    static const rl4::AlignedMatrix<std::uint16_t>& get() {
-        return optimized::halfFloatValues;
+struct OptimizedValues<float> {
+    static const AlignedMatrix<float>& get() {
+        return optimized::floatValues;
     }
 
 };
 
 template<>
-struct OptimizedValues<float> {
-    static const rl4::AlignedMatrix<float>& get() {
-        return optimized::floatValues;
+struct OptimizedValues<std::uint16_t> {
+    static const AlignedMatrix<std::uint16_t>& get() {
+        return optimized::halfFloatValues;
     }
 
 };
 
 template<typename TValue>
-bpcm::Evaluator<TValue> OptimizedStorage<TValue>::create(StrategyPtr strategy, rl4::MemoryResource* memRes) {
+bpcm::Evaluator<TValue> OptimizedStorage<TValue>::create(StrategyPtr strategy,
+                                                         std::size_t rotationSelectorIndex,
+                                                         rl4::RotationType rotationType,
+                                                         MemoryResource* memRes) {
     bpcm::JointStorage<TValue> storage{memRes};
     const auto& values = OptimizedValues<TValue>::get();
     for (std::uint16_t i = 0u; i < static_cast<std::uint16_t>(unoptimized::values.size()); ++i) {
@@ -904,18 +1176,32 @@ bpcm::Evaluator<TValue> OptimizedStorage<TValue>::create(StrategyPtr strategy, r
                                     optimized::inputIndices[i].begin(),
                                     optimized::inputIndices[i].end());
         storage.outputIndices.insert(storage.outputIndices.end(),
-                                     optimized::outputIndices[i].begin(),
-                                     optimized::outputIndices[i].end());
+                                     optimized::outputIndices[rotationSelectorIndex][i].begin(),
+                                     optimized::outputIndices[rotationSelectorIndex][i].end());
+        storage.outputRotationIndices.insert(storage.outputRotationIndices.end(),
+                                             optimized::outputRotationIndices[rotationSelectorIndex][i].begin(),
+                                             optimized::outputRotationIndices[rotationSelectorIndex][i].end());
+        storage.outputRotationLODs.insert(storage.outputRotationLODs.end(),
+                                          optimized::outputRotationLODs[rotationSelectorIndex][i].begin(),
+                                          optimized::outputRotationLODs[rotationSelectorIndex][i].end());
         storage.lodRegions.insert(storage.lodRegions.end(),
                                   optimized::lodRegions[i].begin(),
                                   optimized::lodRegions[i].end());
-        storage.jointGroups.push_back(optimized::jointGroups[i]);
+        storage.jointGroups.push_back(optimized::jointGroups[rotationSelectorIndex][i]);
     }
 
-    const auto outputCount = static_cast<std::uint16_t>(unoptimized::dimensions.rows);
-    auto instanceFactory = [outputCount](rl4::MemoryResource* instanceMemRes) {
-            return pma::UniqueInstance<rl4::bpcm::OutputInstance, rl4::JointsOutputInstance>::with(instanceMemRes).create(
+    const auto lastJointGroupIndex = optimized::outputIndices[rotationSelectorIndex].size() - 1ul;
+    const auto numAttrsPerJoint = static_cast<std::uint16_t>(static_cast<std::uint8_t>(rl4::TranslationType::Vector) +
+                                                             static_cast<std::uint8_t>(rotationType) +
+                                                             static_cast<std::uint8_t>(rl4::ScaleType::Vector));
+    const auto maxOutputIndex = extd::maxOf(optimized::outputIndices[rotationSelectorIndex][lastJointGroupIndex]);
+    const auto outputCount = extd::roundUp(maxOutputIndex, numAttrsPerJoint);
+    auto instanceFactory = [outputCount, rotationType](MemoryResource* instanceMemRes) {
+            return pma::UniqueInstance<CPUJointsOutputInstance, JointsOutputInstance>::with(instanceMemRes).create(
                 outputCount,
+                rl4::TranslationType::Vector,
+                rotationType,
+                rl4::ScaleType::Vector,
                 instanceMemRes);
         };
     return bpcm::Evaluator<TValue>{std::move(storage), std::move(strategy), instanceFactory, memRes};
@@ -923,15 +1209,21 @@ bpcm::Evaluator<TValue> OptimizedStorage<TValue>::create(StrategyPtr strategy, r
 
 template<typename TValue>
 bpcm::Evaluator<TValue> OptimizedStorage<TValue>::create(StrategyPtr strategy,
+                                                         std::size_t rotationSelectorIndex,
+                                                         rl4::RotationType rotationType,
                                                          std::uint16_t jointGroupIndex,
-                                                         rl4::MemoryResource* memRes) {
+                                                         MemoryResource* memRes) {
     bpcm::JointStorage<TValue> storage{memRes};
     const auto& values = OptimizedValues<TValue>::get();
     storage.values.assign(values[jointGroupIndex].begin(), values[jointGroupIndex].end());
     storage.inputIndices.assign(optimized::inputIndices[jointGroupIndex].begin(),
                                 optimized::inputIndices[jointGroupIndex].end());
-    storage.outputIndices.assign(optimized::outputIndices[jointGroupIndex].begin(),
-                                 optimized::outputIndices[jointGroupIndex].end());
+    storage.outputIndices.assign(optimized::outputIndices[rotationSelectorIndex][jointGroupIndex].begin(),
+                                 optimized::outputIndices[rotationSelectorIndex][jointGroupIndex].end());
+    storage.outputRotationIndices.assign(optimized::outputRotationIndices[rotationSelectorIndex][jointGroupIndex].begin(),
+                                         optimized::outputRotationIndices[rotationSelectorIndex][jointGroupIndex].end());
+    storage.outputRotationLODs.assign(optimized::outputRotationLODs[rotationSelectorIndex][jointGroupIndex].begin(),
+                                      optimized::outputRotationLODs[rotationSelectorIndex][jointGroupIndex].end());
     storage.lodRegions.assign(optimized::lodRegions[jointGroupIndex].begin(),
                               optimized::lodRegions[jointGroupIndex].end());
     storage.jointGroups.push_back(bpcm::JointGroup{
@@ -939,16 +1231,26 @@ bpcm::Evaluator<TValue> OptimizedStorage<TValue>::create(StrategyPtr strategy,
             0u,  // inputIndicesOffset
             0u,  // outputIndicesOffset
             0u,  // lodsOffset
+            0u,  // outputRotationIndicesOffset
+            0u,  // outputRotationLODsOffset
             static_cast<std::uint32_t>(storage.values.size()),
             static_cast<std::uint32_t>(storage.inputIndices.size()),
             static_cast<std::uint32_t>(storage.inputIndices.size() - (storage.inputIndices.size() % 4ul)),
             static_cast<std::uint32_t>(storage.inputIndices.size() - (storage.inputIndices.size() % 8ul)),
         });
 
-    const auto outputCount = static_cast<std::uint16_t>(unoptimized::dimensions.rows);
-    auto instanceFactory = [outputCount](rl4::MemoryResource* instanceMemRes) {
-            return pma::UniqueInstance<rl4::bpcm::OutputInstance, rl4::JointsOutputInstance>::with(instanceMemRes).create(
+    const auto lastJointGroupIndex = optimized::outputIndices[rotationSelectorIndex].size() - 1ul;
+    const auto numAttrsPerJoint = static_cast<std::uint16_t>(static_cast<std::uint8_t>(rl4::TranslationType::Vector) +
+                                                             static_cast<std::uint8_t>(rotationType) +
+                                                             static_cast<std::uint8_t>(rl4::ScaleType::Vector));
+    const auto maxOutputIndex = extd::maxOf(optimized::outputIndices[rotationSelectorIndex][lastJointGroupIndex]);
+    const auto outputCount = extd::roundUp(maxOutputIndex, numAttrsPerJoint);
+    auto instanceFactory = [outputCount, rotationType](MemoryResource* instanceMemRes) {
+            return pma::UniqueInstance<CPUJointsOutputInstance, JointsOutputInstance>::with(instanceMemRes).create(
                 outputCount,
+                rl4::TranslationType::Vector,
+                rotationType,
+                rl4::ScaleType::Vector,
                 instanceMemRes);
         };
     return bpcm::Evaluator<TValue>{std::move(storage), std::move(strategy), instanceFactory, memRes};
diff --git a/Engine/Plugins/Animation/RigLogic/Source/RigLogicLibTest/Private/rltests/joints/bpcm/BPCMFixturesBlock8.h b/Engine/Plugins/Animation/RigLogic/Source/RigLogicLibTest/Private/rltests/joints/bpcm/BPCMFixturesBlock8.h
--- a/Engine/Plugins/Animation/RigLogic/Source/RigLogicLibTest/Private/rltests/joints/bpcm/BPCMFixturesBlock8.h
+++ b/Engine/Plugins/Animation/RigLogic/Source/RigLogicLibTest/Private/rltests/joints/bpcm/BPCMFixturesBlock8.h
@@ -7,9 +7,9 @@
 #include "rltests/dna/FakeReader.h"
 
 #include "riglogic/TypeDefs.h"
-#include "riglogic/joints/bpcm/Evaluator.h"
-#include "riglogic/joints/bpcm/JointGroup.h"
-#include "riglogic/joints/bpcm/LODRegion.h"
+#include "riglogic/joints/cpu/bpcm/BPCMJointsEvaluator.h"
+#include "riglogic/joints/cpu/bpcm/JointGroup.h"
+#include "riglogic/riglogic/Configuration.h"
 #include "riglogic/types/Extent.h"
 
 #ifdef _MSC_VER
@@ -44,8 +44,10 @@ extern const Extent dimensions;
 extern const AlignedMatrix<float> floatValues;
 extern const AlignedMatrix<std::uint16_t> halfFloatValues;
 extern const AlignedMatrix<std::uint16_t> inputIndices;
-extern const AlignedMatrix<std::uint16_t> outputIndices;
-extern const Vector<bpcm::JointGroup> jointGroups;
+extern const Vector<Vector<AlignedVector<std::uint16_t> > > outputIndices;
+extern const Vector<Vector<AlignedVector<std::uint16_t> > > outputRotationIndices;
+extern const Vector<Vector<AlignedVector<std::uint16_t> > > outputRotationLODs;
+extern const Vector<Vector<bpcm::JointGroup> > jointGroups;
 extern const Matrix<bpcm::LODRegion> lodRegions;
 
 }  // namespace optimized
@@ -60,7 +62,7 @@ extern const Vector<float> values;
 namespace output {
 
 // Calculation output values
-extern const Matrix<float> valuesPerLOD;
+extern const Vector<Matrix<float> > valuesPerLOD;
 
 }  // namespace output
 
@@ -104,11 +106,18 @@ class CanonicalReader : public dna::FakeReader {
 
 template<typename TValue>
 struct OptimizedStorage {
-    using StrategyPtr = std::unique_ptr<bpcm::JointCalculationStrategy<TValue>,
-                                        std::function<void (bpcm::JointCalculationStrategy<TValue>*)> >;
-
-    static bpcm::Evaluator<TValue> create(StrategyPtr strategy, rl4::MemoryResource* memRes);
-    static bpcm::Evaluator<TValue> create(StrategyPtr strategy, std::uint16_t jointGroupIndex, rl4::MemoryResource* memRes);
+    using StrategyPtr = std::unique_ptr<bpcm::JointGroupLinearCalculationStrategy<TValue>,
+                                        std::function<void (bpcm::JointGroupLinearCalculationStrategy<TValue>*)> >;
+
+    static bpcm::Evaluator<TValue> create(StrategyPtr strategy,
+                                          std::size_t rotationSelectorIndex,
+                                          rl4::RotationType rotationType,
+                                          MemoryResource* memRes);
+    static bpcm::Evaluator<TValue> create(StrategyPtr strategy,
+                                          std::size_t rotationSelectorIndex,
+                                          rl4::RotationType rotationType,
+                                          std::uint16_t jointGroupIndex,
+                                          MemoryResource* memRes);
 
 };
 
diff --git a/Engine/Plugins/Animation/RigLogic/Source/RigLogicLibTest/Private/rltests/joints/bpcm/Helpers.h b/Engine/Plugins/Animation/RigLogic/Source/RigLogicLibTest/Private/rltests/joints/bpcm/Helpers.h
--- a/Engine/Plugins/Animation/RigLogic/Source/RigLogicLibTest/Private/rltests/joints/bpcm/Helpers.h
+++ b/Engine/Plugins/Animation/RigLogic/Source/RigLogicLibTest/Private/rltests/joints/bpcm/Helpers.h
@@ -2,6 +2,7 @@
 
 #pragma once
 
+#include "riglogic/joints/cpu/bpcm/RotationAdapters.h"
 #include "riglogic/riglogic/Configuration.h"
 
 #include <cstddef>
@@ -34,3 +35,30 @@ struct TCalculationType {
     }
 
 };
+
+template<class TRotationAdapter>
+struct BPCMRotationOutputTypeSelector;
+
+template<>
+struct BPCMRotationOutputTypeSelector<rl4::bpcm::NoopAdapter> {
+    static constexpr std::size_t value() {
+        return 1ul;
+    }
+
+    static constexpr rl4::RotationType rotation() {
+        return rl4::RotationType::EulerAngles;
+    }
+
+};
+
+template<typename T>
+struct BPCMRotationOutputTypeSelector<rl4::bpcm::EulerAnglesToQuaternions<T, tdm::rot_seq::xyz> > {
+    static constexpr std::size_t value() {
+        return 0ul;
+    }
+
+    static constexpr rl4::RotationType rotation() {
+        return rl4::RotationType::Quaternions;
+    }
+
+};
diff --git a/Engine/Plugins/Animation/RigLogic/Source/RigLogicLibTest/Private/rltests/joints/bpcm/TestBlock4CalculationStrategy.cpp b/Engine/Plugins/Animation/RigLogic/Source/RigLogicLibTest/Private/rltests/joints/bpcm/TestBlock4CalculationStrategy.cpp
--- a/Engine/Plugins/Animation/RigLogic/Source/RigLogicLibTest/Private/rltests/joints/bpcm/TestBlock4CalculationStrategy.cpp
+++ b/Engine/Plugins/Animation/RigLogic/Source/RigLogicLibTest/Private/rltests/joints/bpcm/TestBlock4CalculationStrategy.cpp
@@ -1,5 +1,9 @@
 // Copyright Epic Games, Inc. All Rights Reserved.
 
+#ifdef _MSC_VER
+    #pragma warning(disable : 4503)
+#endif
+
 #include "riglogic/system/simd/Detect.h"
 
 #include "rltests/Defs.h"
@@ -8,9 +12,11 @@
 #include "rltests/joints/bpcm/Helpers.h"
 
 #include "riglogic/TypeDefs.h"
-#include "riglogic/joints/bpcm/Evaluator.h"
-#include "riglogic/joints/bpcm/strategies/Block4.h"
+#include "riglogic/joints/cpu/bpcm/BPCMJointsEvaluator.h"
+#include "riglogic/joints/cpu/bpcm/CalculationStrategy.h"
+#include "riglogic/joints/cpu/bpcm/RotationAdapters.h"
 #include "riglogic/system/simd/SIMD.h"
+#include "riglogic/utils/Extd.h"
 
 namespace {
 
@@ -26,28 +32,30 @@ class Block4JointCalculationStrategyTest : public ::testing::TestWithParam<Strat
             using T = typename std::tuple_element<0, TestTypes>::type;
             using TFVec = typename std::tuple_element<1, TestTypes>::type;
             using TStrategyTestParams = typename std::tuple_element<2, TestTypes>::type;
+            using TRotationAdapter = typename std::tuple_element<3, TestTypes>::type;
             params.lod = TStrategyTestParams::lod();
 
-            using CalculationStrategyBase = rl4::bpcm::JointCalculationStrategy<T>;
-            using CalculationStrategy = rl4::bpcm::block4::Block4JointCalculationStrategy<T, TFVec>;
+            using CalculationStrategyBase = rl4::bpcm::JointGroupLinearCalculationStrategy<T>;
+            using CalculationStrategy = rl4::bpcm::VectorizedJointGroupLinearCalculationStrategy<T, TFVec, TRotationAdapter>;
             strategy = pma::UniqueInstance<CalculationStrategy, CalculationStrategyBase>::with(&memRes).create();
+
+            rotationSelectorIndex = BPCMRotationOutputTypeSelector<TRotationAdapter>::value();
+            rotationType = BPCMRotationOutputTypeSelector<TRotationAdapter>::rotation();
         }
 
         template<typename TestTypes = TTestTypes>
         typename std::enable_if<std::tuple_size<TestTypes>::value == 0ul, void>::type SetUpImpl() {
+            GTEST_SKIP();
         }
 
         template<typename TArray>
         void execute(const rl4::bpcm::Evaluator<StorageValueType>& joints, const TArray& expected, OutputScope scope) {
-            if (strategy == nullptr) {
-                return;
-            }
-
             auto outputInstance = joints.createInstance(&memRes);
+            outputInstance->resetOutputBuffer();
             auto outputBuffer = outputInstance->getOutputBuffer();
 
             auto inputInstanceFactory =
-                ControlsFactory::getInstanceFactory(0, static_cast<std::uint16_t>(block4::input::values.size()), 0, 0);
+                ControlsFactory::getInstanceFactory(0, static_cast<std::uint16_t>(block4::input::values.size()), 0, 0, 0);
             auto inputInstance = inputInstanceFactory(&memRes);
             auto inputBuffer = inputInstance->getInputBuffer();
             std::copy(block4::input::values.begin(), block4::input::values.end(), inputBuffer.begin());
@@ -55,13 +63,15 @@ class Block4JointCalculationStrategyTest : public ::testing::TestWithParam<Strat
             rl4::ConstArrayView<float> expectedView{expected[scope.lod].data() + scope.offset, scope.size};
             rl4::ConstArrayView<float> outputView{outputBuffer.data() + scope.offset, scope.size};
             joints.calculate(inputInstance.get(), outputInstance.get(), scope.lod);
-            ASSERT_EQ(outputView, expectedView);
+            ASSERT_ELEMENTS_NEAR(outputView, expectedView, expectedView.size(), 0.002f);
         }
 
     protected:
         pma::AlignedMemoryResource memRes;
         block4::OptimizedStorage<StorageValueType>::StrategyPtr strategy;
         StrategyTestParams params;
+        std::size_t rotationSelectorIndex;
+        rl4::RotationType rotationType;
 
 };
 
@@ -69,22 +79,46 @@ class Block4JointCalculationStrategyTest : public ::testing::TestWithParam<Strat
 
 using Block4JointCalculationTypeList = ::testing::Types<
 #if defined(RL_BUILD_WITH_AVX) || defined(RL_BUILD_WITH_SSE)
-        std::tuple<StorageValueType, trimd::sse::F128, TStrategyTestParams<0> >,
-        std::tuple<StorageValueType, trimd::sse::F128, TStrategyTestParams<1> >,
-        std::tuple<StorageValueType, trimd::sse::F128, TStrategyTestParams<2> >,
-        std::tuple<StorageValueType, trimd::sse::F128, TStrategyTestParams<3> >,
+        std::tuple<StorageValueType, trimd::sse::F128, TStrategyTestParams<0>, rl4::bpcm::NoopAdapter>,
+        std::tuple<StorageValueType, trimd::sse::F128, TStrategyTestParams<0>, rl4::bpcm::EulerAnglesToQuaternions<tdm::fdeg,
+                                                                                                                   tdm::rot_seq::xyz> >,
+        std::tuple<StorageValueType, trimd::sse::F128, TStrategyTestParams<1>, rl4::bpcm::NoopAdapter>,
+        std::tuple<StorageValueType, trimd::sse::F128, TStrategyTestParams<1>, rl4::bpcm::EulerAnglesToQuaternions<tdm::fdeg,
+                                                                                                                   tdm::rot_seq::xyz> >,
+        std::tuple<StorageValueType, trimd::sse::F128, TStrategyTestParams<2>, rl4::bpcm::NoopAdapter>,
+        std::tuple<StorageValueType, trimd::sse::F128, TStrategyTestParams<2>, rl4::bpcm::EulerAnglesToQuaternions<tdm::fdeg,
+                                                                                                                   tdm::rot_seq::xyz> >,
+        std::tuple<StorageValueType, trimd::sse::F128, TStrategyTestParams<3>, rl4::bpcm::NoopAdapter>,
+        std::tuple<StorageValueType, trimd::sse::F128, TStrategyTestParams<3>, rl4::bpcm::EulerAnglesToQuaternions<tdm::fdeg,
+                                                                                                                   tdm::rot_seq::xyz> >,
 #endif  // RL_BUILD_WITH_AVX || RL_BUILD_WITH_SSE
 #if defined(RL_BUILD_WITH_NEON)
-        std::tuple<StorageValueType, trimd::neon::F128, TStrategyTestParams<0> >,
-        std::tuple<StorageValueType, trimd::neon::F128, TStrategyTestParams<1> >,
-        std::tuple<StorageValueType, trimd::neon::F128, TStrategyTestParams<2> >,
-        std::tuple<StorageValueType, trimd::neon::F128, TStrategyTestParams<3> >,
+        std::tuple<StorageValueType, trimd::neon::F128, TStrategyTestParams<0>, rl4::bpcm::NoopAdapter>,
+        std::tuple<StorageValueType, trimd::neon::F128, TStrategyTestParams<0>, rl4::bpcm::EulerAnglesToQuaternions<tdm::fdeg,
+                                                                                                                    tdm::rot_seq::xyz> >,
+        std::tuple<StorageValueType, trimd::neon::F128, TStrategyTestParams<1>, rl4::bpcm::NoopAdapter>,
+        std::tuple<StorageValueType, trimd::neon::F128, TStrategyTestParams<1>, rl4::bpcm::EulerAnglesToQuaternions<tdm::fdeg,
+                                                                                                                    tdm::rot_seq::xyz> >,
+        std::tuple<StorageValueType, trimd::neon::F128, TStrategyTestParams<2>, rl4::bpcm::NoopAdapter>,
+        std::tuple<StorageValueType, trimd::neon::F128, TStrategyTestParams<2>, rl4::bpcm::EulerAnglesToQuaternions<tdm::fdeg,
+                                                                                                                    tdm::rot_seq::xyz> >,
+        std::tuple<StorageValueType, trimd::neon::F128, TStrategyTestParams<3>, rl4::bpcm::NoopAdapter>,
+        std::tuple<StorageValueType, trimd::neon::F128, TStrategyTestParams<3>, rl4::bpcm::EulerAnglesToQuaternions<tdm::fdeg,
+                                                                                                                    tdm::rot_seq::xyz> >,
 #endif  // RL_BUILD_WITH_NEON
 #if !defined(RL_USE_HALF_FLOATS)
-        std::tuple<StorageValueType, trimd::scalar::F128, TStrategyTestParams<0> >,
-        std::tuple<StorageValueType, trimd::scalar::F128, TStrategyTestParams<1> >,
-        std::tuple<StorageValueType, trimd::scalar::F128, TStrategyTestParams<2> >,
-        std::tuple<StorageValueType, trimd::scalar::F128, TStrategyTestParams<3> >,
+        std::tuple<StorageValueType, trimd::scalar::F128, TStrategyTestParams<0>, rl4::bpcm::NoopAdapter>,
+        std::tuple<StorageValueType, trimd::scalar::F128, TStrategyTestParams<0>,
+                   rl4::bpcm::EulerAnglesToQuaternions<tdm::fdeg, tdm::rot_seq::xyz> >,
+        std::tuple<StorageValueType, trimd::scalar::F128, TStrategyTestParams<1>, rl4::bpcm::NoopAdapter>,
+        std::tuple<StorageValueType, trimd::scalar::F128, TStrategyTestParams<1>,
+                   rl4::bpcm::EulerAnglesToQuaternions<tdm::fdeg, tdm::rot_seq::xyz> >,
+        std::tuple<StorageValueType, trimd::scalar::F128, TStrategyTestParams<2>, rl4::bpcm::NoopAdapter>,
+        std::tuple<StorageValueType, trimd::scalar::F128, TStrategyTestParams<2>,
+                   rl4::bpcm::EulerAnglesToQuaternions<tdm::fdeg, tdm::rot_seq::xyz> >,
+        std::tuple<StorageValueType, trimd::scalar::F128, TStrategyTestParams<3>, rl4::bpcm::NoopAdapter>,
+        std::tuple<StorageValueType, trimd::scalar::F128, TStrategyTestParams<3>,
+                   rl4::bpcm::EulerAnglesToQuaternions<tdm::fdeg, tdm::rot_seq::xyz> >,
 #endif  // RL_USE_HALF_FLOATS
     std::tuple<>
     >;
@@ -92,67 +126,164 @@ using Block4JointCalculationTypeList = ::testing::Types<
 TYPED_TEST_SUITE(Block4JointCalculationStrategyTest, Block4JointCalculationTypeList, );
 
 TYPED_TEST(Block4JointCalculationStrategyTest, Block4Padded) {
-    const OutputScope scope{this->params.lod, 0ul, 1ul};
-    auto joints = block4::OptimizedStorage<StorageValueType>::create(std::move(this->strategy), 0u, &this->memRes);
-    this->execute(joints, block4::output::valuesPerLOD, scope);
+    const std::uint16_t jointGroupIndex = 0u;
+    const auto& outputIndices = block4::optimized::outputIndices[this->rotationSelectorIndex][jointGroupIndex];
+    const auto outputCount = block4::unoptimized::outputIndices[jointGroupIndex].size();
+    const auto outputOffset = extd::minOf(rl4::ConstArrayView<std::uint16_t>{outputIndices.data(), outputCount});
+    const OutputScope scope{this->params.lod, outputOffset, outputCount};
+    auto joints = block4::OptimizedStorage<StorageValueType>::create(std::move(this->strategy),
+                                                                     this->rotationSelectorIndex,
+                                                                     this->rotationType,
+                                                                     jointGroupIndex,
+                                                                     &this->memRes);
+    this->execute(joints, block4::output::valuesPerLOD[this->rotationSelectorIndex], scope);
 }
 
 TYPED_TEST(Block4JointCalculationStrategyTest, Block4Exact) {
-    const OutputScope scope{this->params.lod, 1ul, 4ul};
-    auto joints = block4::OptimizedStorage<StorageValueType>::create(std::move(this->strategy), 1u, &this->memRes);
-    this->execute(joints, block4::output::valuesPerLOD, scope);
+    const std::uint16_t jointGroupIndex = 1u;
+    const auto& outputIndices = block4::optimized::outputIndices[this->rotationSelectorIndex][jointGroupIndex];
+    const auto outputCount = block4::unoptimized::outputIndices[jointGroupIndex].size();
+    const auto outputOffset = extd::minOf(rl4::ConstArrayView<std::uint16_t>{outputIndices.data(), outputCount});
+    const OutputScope scope{this->params.lod, outputOffset, outputCount};
+    auto joints = block4::OptimizedStorage<StorageValueType>::create(std::move(this->strategy),
+                                                                     this->rotationSelectorIndex,
+                                                                     this->rotationType,
+                                                                     jointGroupIndex,
+                                                                     &this->memRes);
+    this->execute(joints, block4::output::valuesPerLOD[this->rotationSelectorIndex], scope);
 }
 
 TYPED_TEST(Block4JointCalculationStrategyTest, Block8Padded) {
-    const OutputScope scope{this->params.lod, 5ul, 7ul};
-    auto joints = block4::OptimizedStorage<StorageValueType>::create(std::move(this->strategy), 2u, &this->memRes);
-    this->execute(joints, block4::output::valuesPerLOD, scope);
+    const std::uint16_t jointGroupIndex = 2u;
+    const auto& outputIndices = block4::optimized::outputIndices[this->rotationSelectorIndex][jointGroupIndex];
+    const auto outputCount = block4::unoptimized::outputIndices[jointGroupIndex].size();
+    const auto outputOffset = extd::minOf(rl4::ConstArrayView<std::uint16_t>{outputIndices.data(), outputCount});
+    const OutputScope scope{this->params.lod, outputOffset, outputCount};
+    auto joints = block4::OptimizedStorage<StorageValueType>::create(std::move(this->strategy),
+                                                                     this->rotationSelectorIndex,
+                                                                     this->rotationType,
+                                                                     jointGroupIndex,
+                                                                     &this->memRes);
+    this->execute(joints, block4::output::valuesPerLOD[this->rotationSelectorIndex], scope);
 }
 
 TYPED_TEST(Block4JointCalculationStrategyTest, Block8Exact) {
-    const OutputScope scope{this->params.lod, 12ul, 8ul};
-    auto joints = block4::OptimizedStorage<StorageValueType>::create(std::move(this->strategy), 3u, &this->memRes);
-    this->execute(joints, block4::output::valuesPerLOD, scope);
+    const std::uint16_t jointGroupIndex = 3u;
+    const auto& outputIndices = block4::optimized::outputIndices[this->rotationSelectorIndex][jointGroupIndex];
+    const auto outputCount = block4::unoptimized::outputIndices[jointGroupIndex].size();
+    const auto outputOffset = extd::minOf(rl4::ConstArrayView<std::uint16_t>{outputIndices.data(), outputCount});
+    const OutputScope scope{this->params.lod, outputOffset, outputCount};
+    auto joints = block4::OptimizedStorage<StorageValueType>::create(std::move(this->strategy),
+                                                                     this->rotationSelectorIndex,
+                                                                     this->rotationType,
+                                                                     jointGroupIndex,
+                                                                     &this->memRes);
+    this->execute(joints, block4::output::valuesPerLOD[this->rotationSelectorIndex], scope);
 }
 
 TYPED_TEST(Block4JointCalculationStrategyTest, Block12Padded) {
-    const OutputScope scope{this->params.lod, 20ul, 9ul};
-    auto joints = block4::OptimizedStorage<StorageValueType>::create(std::move(this->strategy), 4u, &this->memRes);
-    this->execute(joints, block4::output::valuesPerLOD, scope);
+    const std::uint16_t jointGroupIndex = 4u;
+    const auto& outputIndices = block4::optimized::outputIndices[this->rotationSelectorIndex][jointGroupIndex];
+    const auto outputCount = block4::unoptimized::outputIndices[jointGroupIndex].size();
+    const auto outputOffset = extd::minOf(rl4::ConstArrayView<std::uint16_t>{outputIndices.data(), outputCount});
+    const OutputScope scope{this->params.lod, outputOffset, outputCount};
+    auto joints = block4::OptimizedStorage<StorageValueType>::create(std::move(this->strategy),
+                                                                     this->rotationSelectorIndex,
+                                                                     this->rotationType,
+                                                                     jointGroupIndex,
+                                                                     &this->memRes);
+    this->execute(joints, block4::output::valuesPerLOD[this->rotationSelectorIndex], scope);
 }
 
 TYPED_TEST(Block4JointCalculationStrategyTest, Block12Exact) {
-    const OutputScope scope{this->params.lod, 29ul, 12ul};
-    auto joints = block4::OptimizedStorage<StorageValueType>::create(std::move(this->strategy), 5u, &this->memRes);
-    this->execute(joints, block4::output::valuesPerLOD, scope);
+    const std::uint16_t jointGroupIndex = 5u;
+    const auto& outputIndices = block4::optimized::outputIndices[this->rotationSelectorIndex][jointGroupIndex];
+    const auto outputCount = block4::unoptimized::outputIndices[jointGroupIndex].size();
+    const auto outputOffset = extd::minOf(rl4::ConstArrayView<std::uint16_t>{outputIndices.data(), outputCount});
+    const OutputScope scope{this->params.lod, outputOffset, outputCount};
+    auto joints = block4::OptimizedStorage<StorageValueType>::create(std::move(this->strategy),
+                                                                     this->rotationSelectorIndex,
+                                                                     this->rotationType,
+                                                                     jointGroupIndex,
+                                                                     &this->memRes);
+    this->execute(joints, block4::output::valuesPerLOD[this->rotationSelectorIndex], scope);
 }
 
 TYPED_TEST(Block4JointCalculationStrategyTest, Block16Padded) {
-    const OutputScope scope{this->params.lod, 41ul, 14ul};
-    auto joints = block4::OptimizedStorage<StorageValueType>::create(std::move(this->strategy), 6u, &this->memRes);
-    this->execute(joints, block4::output::valuesPerLOD, scope);
+    const std::uint16_t jointGroupIndex = 6u;
+    const auto& outputIndices = block4::optimized::outputIndices[this->rotationSelectorIndex][jointGroupIndex];
+    const auto outputCount = block4::unoptimized::outputIndices[jointGroupIndex].size();
+    const auto outputOffset = extd::minOf(rl4::ConstArrayView<std::uint16_t>{outputIndices.data(), outputCount});
+    const OutputScope scope{this->params.lod, outputOffset, outputCount};
+    auto joints = block4::OptimizedStorage<StorageValueType>::create(std::move(this->strategy),
+                                                                     this->rotationSelectorIndex,
+                                                                     this->rotationType,
+                                                                     jointGroupIndex,
+                                                                     &this->memRes);
+    this->execute(joints, block4::output::valuesPerLOD[this->rotationSelectorIndex], scope);
 }
 
 TYPED_TEST(Block4JointCalculationStrategyTest, Block16Exact) {
-    const OutputScope scope{this->params.lod, 55ul, 16ul};
-    auto joints = block4::OptimizedStorage<StorageValueType>::create(std::move(this->strategy), 7u, &this->memRes);
-    this->execute(joints, block4::output::valuesPerLOD, scope);
+    const std::uint16_t jointGroupIndex = 7u;
+    const auto& outputIndices = block4::optimized::outputIndices[this->rotationSelectorIndex][jointGroupIndex];
+    const auto outputCount = block4::unoptimized::outputIndices[jointGroupIndex].size();
+    const auto outputOffset = extd::minOf(rl4::ConstArrayView<std::uint16_t>{outputIndices.data(), outputCount});
+    const OutputScope scope{this->params.lod, outputOffset, outputCount};
+    auto joints = block4::OptimizedStorage<StorageValueType>::create(std::move(this->strategy),
+                                                                     this->rotationSelectorIndex,
+                                                                     this->rotationType,
+                                                                     jointGroupIndex,
+                                                                     &this->memRes);
+    this->execute(joints, block4::output::valuesPerLOD[this->rotationSelectorIndex], scope);
+}
+
+TYPED_TEST(Block4JointCalculationStrategyTest, Block16ExactMixedOutputOrder) {
+    const std::uint16_t jointGroupIndex = 8u;
+    const auto& outputIndices = block4::optimized::outputIndices[this->rotationSelectorIndex][jointGroupIndex];
+    const auto outputCount = block4::unoptimized::outputIndices[jointGroupIndex].size();
+    const auto outputOffset = extd::minOf(rl4::ConstArrayView<std::uint16_t>{outputIndices.data(), outputCount});
+    const OutputScope scope{this->params.lod, outputOffset, outputCount};
+    auto joints = block4::OptimizedStorage<StorageValueType>::create(std::move(this->strategy),
+                                                                     this->rotationSelectorIndex,
+                                                                     this->rotationType,
+                                                                     jointGroupIndex,
+                                                                     &this->memRes);
+    this->execute(joints, block4::output::valuesPerLOD[this->rotationSelectorIndex], scope);
 }
 
 TYPED_TEST(Block4JointCalculationStrategyTest, MultipleBlocks) {
-    const OutputScope scope{this->params.lod, 0ul, block4::output::valuesPerLOD.front().size()};
-    auto joints = block4::OptimizedStorage<StorageValueType>::create(std::move(this->strategy), &this->memRes);
-    this->execute(joints, block4::output::valuesPerLOD, scope);
+    const OutputScope scope{this->params.lod, 0ul, block4::output::valuesPerLOD[this->rotationSelectorIndex].front().size()};
+    auto joints = block4::OptimizedStorage<StorageValueType>::create(std::move(this->strategy),
+                                                                     this->rotationSelectorIndex,
+                                                                     this->rotationType,
+                                                                     &this->memRes);
+    this->execute(joints, block4::output::valuesPerLOD[this->rotationSelectorIndex], scope);
 }
 
 TYPED_TEST(Block4JointCalculationStrategyTest, InputRegionA) {
-    const OutputScope scope{this->params.lod, 71ul, 2ul};
-    auto joints = block4::OptimizedStorage<StorageValueType>::create(std::move(this->strategy), 8u, &this->memRes);
-    this->execute(joints, block4::output::valuesPerLOD, scope);
+    const std::uint16_t jointGroupIndex = 9u;
+    const auto& outputIndices = block4::optimized::outputIndices[this->rotationSelectorIndex][jointGroupIndex];
+    const auto outputCount = block4::unoptimized::outputIndices[jointGroupIndex].size();
+    const auto outputOffset = extd::minOf(rl4::ConstArrayView<std::uint16_t>{outputIndices.data(), outputCount});
+    const OutputScope scope{this->params.lod, outputOffset, outputCount};
+    auto joints = block4::OptimizedStorage<StorageValueType>::create(std::move(this->strategy),
+                                                                     this->rotationSelectorIndex,
+                                                                     this->rotationType,
+                                                                     jointGroupIndex,
+                                                                     &this->memRes);
+    this->execute(joints, block4::output::valuesPerLOD[this->rotationSelectorIndex], scope);
 }
 
 TYPED_TEST(Block4JointCalculationStrategyTest, InputRegionB) {
-    const OutputScope scope{this->params.lod, 73ul, 2ul};
-    auto joints = block4::OptimizedStorage<StorageValueType>::create(std::move(this->strategy), 9u, &this->memRes);
-    this->execute(joints, block4::output::valuesPerLOD, scope);
+    const std::uint16_t jointGroupIndex = 10u;
+    const auto& outputIndices = block4::optimized::outputIndices[this->rotationSelectorIndex][jointGroupIndex];
+    const auto outputCount = block4::unoptimized::outputIndices[jointGroupIndex].size();
+    const auto outputOffset = extd::minOf(rl4::ConstArrayView<std::uint16_t>{outputIndices.data(), outputCount});
+    const OutputScope scope{this->params.lod, outputOffset, outputCount};
+    auto joints = block4::OptimizedStorage<StorageValueType>::create(std::move(this->strategy),
+                                                                     this->rotationSelectorIndex,
+                                                                     this->rotationType,
+                                                                     jointGroupIndex,
+                                                                     &this->memRes);
+    this->execute(joints, block4::output::valuesPerLOD[this->rotationSelectorIndex], scope);
 }
diff --git a/Engine/Plugins/Animation/RigLogic/Source/RigLogicLibTest/Private/rltests/joints/bpcm/TestBlock4StorageBuilder.cpp b/Engine/Plugins/Animation/RigLogic/Source/RigLogicLibTest/Private/rltests/joints/bpcm/TestBlock4StorageBuilder.cpp
--- a/Engine/Plugins/Animation/RigLogic/Source/RigLogicLibTest/Private/rltests/joints/bpcm/TestBlock4StorageBuilder.cpp
+++ b/Engine/Plugins/Animation/RigLogic/Source/RigLogicLibTest/Private/rltests/joints/bpcm/TestBlock4StorageBuilder.cpp
@@ -8,9 +8,12 @@
 #include "rltests/joints/bpcm/Helpers.h"
 
 #include "riglogic/TypeDefs.h"
+#include "riglogic/joints/JointBehaviorFilter.h"
 #include "riglogic/joints/JointsBuilder.h"
-#include "riglogic/joints/bpcm/Evaluator.h"
-#include "riglogic/joints/bpcm/strategies/Block4.h"
+#include "riglogic/joints/cpu/CPUJointsEvaluator.h"
+#include "riglogic/joints/cpu/bpcm/BPCMJointsEvaluator.h"
+#include "riglogic/joints/cpu/bpcm/CalculationStrategy.h"
+#include "riglogic/joints/cpu/bpcm/RotationAdapters.h"
 #include "riglogic/riglogic/RigLogic.h"
 #include "riglogic/system/simd/SIMD.h"
 
@@ -28,23 +31,39 @@ class Block4JointStorageBuilderTest : public ::testing::Test {
             using TValue = typename std::tuple_element<0, TestTypes>::type;
             using TFVec = typename std::tuple_element<1, TestTypes>::type;
             using TCalculationType = typename std::tuple_element<2, TestTypes>::type;
+            using TRotationAdapter = typename std::tuple_element<3, TestTypes>::type;
 
             rl4::Configuration config{};
             config.calculationType = TCalculationType::get();
+            config.rotationType = BPCMRotationOutputTypeSelector<TRotationAdapter>::rotation();
             auto builder = rl4::JointsBuilder::create(config, &memRes);
-            builder->computeStorageRequirements(&reader);
-            builder->allocateStorage(&reader);
-            builder->fillStorage(&reader);
+
+            rl4::JointBehaviorFilter filter{&reader, &memRes};
+            filter.include(dna::TranslationRepresentation::Vector);
+            filter.include(dna::RotationRepresentation::EulerAngles);
+            filter.include(dna::ScaleRepresentation::Vector);
+
+            builder->computeStorageRequirements(filter);
+            builder->allocateStorage(filter);
+            builder->fillStorage(filter);
             auto joints = builder->build();
-            auto jointsImpl = static_cast<rl4::bpcm::Evaluator<TValue>*>(joints.get());
+            auto jointsImpl = static_cast<rl4::CPUJointsEvaluator*>(joints.get());
+            auto bpcmJointsImpl =
+                static_cast<rl4::bpcm::Evaluator<TValue>*>(rl4::CPUJointsEvaluator::Accessor::getBPCMEvaluator(jointsImpl));
 
-            auto strategy = pma::UniqueInstance<rl4::bpcm::block4::Block4JointCalculationStrategy<TValue, TFVec>,
-                                                rl4::bpcm::JointCalculationStrategy<TValue> >::with(&memRes).create();
-            auto expected = block4::OptimizedStorage<TValue>::create(std::move(strategy), &memRes);
+            const auto rotationSelectorIndex = BPCMRotationOutputTypeSelector<TRotationAdapter>::value();
+            const auto rotationType = BPCMRotationOutputTypeSelector<TRotationAdapter>::rotation();
+            auto strategy = pma::UniqueInstance<rl4::bpcm::VectorizedJointGroupLinearCalculationStrategy<TValue, TFVec,
+                                                                                                         TRotationAdapter>,
+                                                rl4::bpcm::JointGroupLinearCalculationStrategy<TValue> >::with(&memRes).create();
+            auto expected = block4::OptimizedStorage<TValue>::create(std::move(strategy),
+                                                                     rotationSelectorIndex,
+                                                                     rotationType,
+                                                                     &memRes);
 
-            rl4::bpcm::Evaluator<TValue>::Accessor::assertRawDataEqual(*jointsImpl, expected);
-            rl4::bpcm::Evaluator<TValue>::Accessor::assertJointGroupsEqual(*jointsImpl, expected);
-            rl4::bpcm::Evaluator<TValue>::Accessor::assertLODsEqual(*jointsImpl, expected);
+            rl4::bpcm::Evaluator<TValue>::Accessor::assertRawDataEqual(*bpcmJointsImpl, expected);
+            rl4::bpcm::Evaluator<TValue>::Accessor::assertJointGroupsEqual(*bpcmJointsImpl, expected);
+            rl4::bpcm::Evaluator<TValue>::Accessor::assertLODsEqual(*bpcmJointsImpl, expected);
         }
 
     protected:
@@ -57,13 +76,19 @@ class Block4JointStorageBuilderTest : public ::testing::Test {
 
 using Block4StorageValueTypeList = ::testing::Types<
 #if defined(RL_BUILD_WITH_SSE)
-        std::tuple<StorageValueType, trimd::sse::F128, TCalculationType<rl4::CalculationType::SSE> >,
+        std::tuple<StorageValueType, trimd::sse::F128, TCalculationType<rl4::CalculationType::SSE>, rl4::bpcm::NoopAdapter>,
+        std::tuple<StorageValueType, trimd::sse::F128, TCalculationType<rl4::CalculationType::SSE>,
+                   rl4::bpcm::EulerAnglesToQuaternions<tdm::fdeg, tdm::rot_seq::xyz> >,
 #endif  // RL_BUILD_WITH_AVX || RL_BUILD_WITH_SSE
 #if defined(RL_BUILD_WITH_NEON)
-        std::tuple<StorageValueType, trimd::neon::F128, TCalculationType<rl4::CalculationType::NEON> >,
+        std::tuple<StorageValueType, trimd::neon::F128, TCalculationType<rl4::CalculationType::NEON>, rl4::bpcm::NoopAdapter>,
+        std::tuple<StorageValueType, trimd::neon::F128, TCalculationType<rl4::CalculationType::NEON>,
+                   rl4::bpcm::EulerAnglesToQuaternions<tdm::fdeg, tdm::rot_seq::xyz> >,
 #endif  // RL_BUILD_WITH_NEON
 #if !defined(RL_USE_HALF_FLOATS)
-        std::tuple<StorageValueType, trimd::scalar::F128, TCalculationType<rl4::CalculationType::Scalar> >,
+        std::tuple<StorageValueType, trimd::scalar::F128, TCalculationType<rl4::CalculationType::Scalar>, rl4::bpcm::NoopAdapter>,
+        std::tuple<StorageValueType, trimd::scalar::F128, TCalculationType<rl4::CalculationType::Scalar>,
+                   rl4::bpcm::EulerAnglesToQuaternions<tdm::fdeg, tdm::rot_seq::xyz> >,
 #endif  // RL_USE_HALF_FLOATS
     std::tuple<>
     >;
