Update to new RigLogic implementing body rig support

diff --git a/Engine/Plugins/Animation/RigLogic/Source/RigLogicLibTest/Private/rltests/joints/bpcm/TestBlock8CalculationStrategy.cpp b/Engine/Plugins/Animation/RigLogic/Source/RigLogicLibTest/Private/rltests/joints/bpcm/TestBlock8CalculationStrategy.cpp
--- a/Engine/Plugins/Animation/RigLogic/Source/RigLogicLibTest/Private/rltests/joints/bpcm/TestBlock8CalculationStrategy.cpp
+++ b/Engine/Plugins/Animation/RigLogic/Source/RigLogicLibTest/Private/rltests/joints/bpcm/TestBlock8CalculationStrategy.cpp
@@ -1,5 +1,9 @@
 // Copyright Epic Games, Inc. All Rights Reserved.
 
+#ifdef _MSC_VER
+    #pragma warning(disable : 4503)
+#endif
+
 #include "riglogic/system/simd/Detect.h"
 
 #include "rltests/Defs.h"
@@ -8,14 +12,15 @@
 #include "rltests/joints/bpcm/Helpers.h"
 
 #include "riglogic/TypeDefs.h"
-#include "riglogic/joints/bpcm/Evaluator.h"
-#include "riglogic/joints/bpcm/strategies/Block8.h"
+#include "riglogic/joints/cpu/bpcm/BPCMJointsEvaluator.h"
+#include "riglogic/joints/cpu/bpcm/CalculationStrategy.h"
+#include "riglogic/joints/cpu/bpcm/RotationAdapters.h"
 #include "riglogic/system/simd/SIMD.h"
 
 namespace {
 
 template<typename TTestTypes>
-class Block8JointCalculationStrategyTest : public ::testing::Test {
+class Block8JointCalculationStrategyTest : public ::testing::TestWithParam<StrategyTestParams> {
     protected:
         void SetUp() override {
             Block8JointCalculationStrategyTest::SetUpImpl();
@@ -26,28 +31,30 @@ class Block8JointCalculationStrategyTest : public ::testing::Test {
             using T = typename std::tuple_element<0, TestTypes>::type;
             using TFVec = typename std::tuple_element<1, TestTypes>::type;
             using TStrategyTestParams = typename std::tuple_element<2, TestTypes>::type;
+            using TRotationAdapter = typename std::tuple_element<3, TestTypes>::type;
             params.lod = TStrategyTestParams::lod();
 
-            using CalculationStrategyBase = rl4::bpcm::JointCalculationStrategy<T>;
-            using CalculationStrategy = rl4::bpcm::block8::Block8JointCalculationStrategy<T, TFVec>;
+            using CalculationStrategyBase = rl4::bpcm::JointGroupLinearCalculationStrategy<T>;
+            using CalculationStrategy = rl4::bpcm::VectorizedJointGroupLinearCalculationStrategy<T, TFVec, TRotationAdapter>;
             strategy = pma::UniqueInstance<CalculationStrategy, CalculationStrategyBase>::with(&memRes).create();
+
+            rotationSelectorIndex = BPCMRotationOutputTypeSelector<TRotationAdapter>::value();
+            rotationType = BPCMRotationOutputTypeSelector<TRotationAdapter>::rotation();
         }
 
         template<typename TestTypes = TTestTypes>
         typename std::enable_if<std::tuple_size<TestTypes>::value == 0ul, void>::type SetUpImpl() {
+            GTEST_SKIP();
         }
 
         template<typename TArray>
         void execute(const rl4::bpcm::Evaluator<StorageValueType>& joints, const TArray& expected, OutputScope scope) {
-            if (strategy == nullptr) {
-                return;
-            }
-
             auto outputInstance = joints.createInstance(&memRes);
+            outputInstance->resetOutputBuffer();
             auto outputBuffer = outputInstance->getOutputBuffer();
 
             auto inputInstanceFactory =
-                ControlsFactory::getInstanceFactory(0, static_cast<std::uint16_t>(block8::input::values.size()), 0, 0);
+                ControlsFactory::getInstanceFactory(0, static_cast<std::uint16_t>(block8::input::values.size()), 0, 0, 0);
             auto inputInstance = inputInstanceFactory(&memRes);
             auto inputBuffer = inputInstance->getInputBuffer();
             std::copy(block8::input::values.begin(), block8::input::values.end(), inputBuffer.begin());
@@ -55,13 +62,15 @@ class Block8JointCalculationStrategyTest : public ::testing::Test {
             rl4::ConstArrayView<float> expectedView{expected[scope.lod].data() + scope.offset, scope.size};
             rl4::ConstArrayView<float> outputView{outputBuffer.data() + scope.offset, scope.size};
             joints.calculate(inputInstance.get(), outputInstance.get(), scope.lod);
-            ASSERT_EQ(outputView, expectedView);
+            ASSERT_ELEMENTS_NEAR(outputView, expectedView, expectedView.size(), 0.002f);
         }
 
     protected:
         pma::AlignedMemoryResource memRes;
         block8::OptimizedStorage<StorageValueType>::StrategyPtr strategy;
         StrategyTestParams params;
+        std::size_t rotationSelectorIndex;
+        rl4::RotationType rotationType;
 
 };
 
@@ -69,28 +78,60 @@ class Block8JointCalculationStrategyTest : public ::testing::Test {
 
 using Block8JointCalculationTypeList = ::testing::Types<
 #if defined(RL_BUILD_WITH_AVX)
-        std::tuple<StorageValueType, trimd::avx::F256, TStrategyTestParams<0> >,
-        std::tuple<StorageValueType, trimd::avx::F256, TStrategyTestParams<1> >,
-        std::tuple<StorageValueType, trimd::avx::F256, TStrategyTestParams<2> >,
-        std::tuple<StorageValueType, trimd::avx::F256, TStrategyTestParams<3> >,
+        std::tuple<StorageValueType, trimd::avx::F256, TStrategyTestParams<0>, rl4::bpcm::NoopAdapter>,
+        std::tuple<StorageValueType, trimd::avx::F256, TStrategyTestParams<0>, rl4::bpcm::EulerAnglesToQuaternions<tdm::fdeg,
+                                                                                                                   tdm::rot_seq::xyz> >,
+        std::tuple<StorageValueType, trimd::avx::F256, TStrategyTestParams<1>, rl4::bpcm::NoopAdapter>,
+        std::tuple<StorageValueType, trimd::avx::F256, TStrategyTestParams<1>, rl4::bpcm::EulerAnglesToQuaternions<tdm::fdeg,
+                                                                                                                   tdm::rot_seq::xyz> >,
+        std::tuple<StorageValueType, trimd::avx::F256, TStrategyTestParams<2>, rl4::bpcm::NoopAdapter>,
+        std::tuple<StorageValueType, trimd::avx::F256, TStrategyTestParams<2>, rl4::bpcm::EulerAnglesToQuaternions<tdm::fdeg,
+                                                                                                                   tdm::rot_seq::xyz> >,
+        std::tuple<StorageValueType, trimd::avx::F256, TStrategyTestParams<3>, rl4::bpcm::NoopAdapter>,
+        std::tuple<StorageValueType, trimd::avx::F256, TStrategyTestParams<3>, rl4::bpcm::EulerAnglesToQuaternions<tdm::fdeg,
+                                                                                                                   tdm::rot_seq::xyz> >,
 #endif  // RL_BUILD_WITH_AVX
 #if defined(RL_BUILD_WITH_AVX) || defined(RL_BUILD_WITH_SSE)
-        std::tuple<StorageValueType, trimd::sse::F256, TStrategyTestParams<0> >,
-        std::tuple<StorageValueType, trimd::sse::F256, TStrategyTestParams<1> >,
-        std::tuple<StorageValueType, trimd::sse::F256, TStrategyTestParams<2> >,
-        std::tuple<StorageValueType, trimd::sse::F256, TStrategyTestParams<3> >,
+        std::tuple<StorageValueType, trimd::sse::F256, TStrategyTestParams<0>, rl4::bpcm::NoopAdapter>,
+        std::tuple<StorageValueType, trimd::sse::F256, TStrategyTestParams<0>, rl4::bpcm::EulerAnglesToQuaternions<tdm::fdeg,
+                                                                                                                   tdm::rot_seq::xyz> >,
+        std::tuple<StorageValueType, trimd::sse::F256, TStrategyTestParams<1>, rl4::bpcm::NoopAdapter>,
+        std::tuple<StorageValueType, trimd::sse::F256, TStrategyTestParams<1>, rl4::bpcm::EulerAnglesToQuaternions<tdm::fdeg,
+                                                                                                                   tdm::rot_seq::xyz> >,
+        std::tuple<StorageValueType, trimd::sse::F256, TStrategyTestParams<2>, rl4::bpcm::NoopAdapter>,
+        std::tuple<StorageValueType, trimd::sse::F256, TStrategyTestParams<2>, rl4::bpcm::EulerAnglesToQuaternions<tdm::fdeg,
+                                                                                                                   tdm::rot_seq::xyz> >,
+        std::tuple<StorageValueType, trimd::sse::F256, TStrategyTestParams<3>, rl4::bpcm::NoopAdapter>,
+        std::tuple<StorageValueType, trimd::sse::F256, TStrategyTestParams<3>, rl4::bpcm::EulerAnglesToQuaternions<tdm::fdeg,
+                                                                                                                   tdm::rot_seq::xyz> >,
 #endif  // RL_BUILD_WITH_AVX || RL_BUILD_WITH_SSE
 #if defined(RL_BUILD_WITH_NEON)
-        std::tuple<StorageValueType, trimd::neon::F256, TStrategyTestParams<0> >,
-        std::tuple<StorageValueType, trimd::neon::F256, TStrategyTestParams<1> >,
-        std::tuple<StorageValueType, trimd::neon::F256, TStrategyTestParams<2> >,
-        std::tuple<StorageValueType, trimd::neon::F256, TStrategyTestParams<3> >,
+        std::tuple<StorageValueType, trimd::neon::F256, TStrategyTestParams<0>, rl4::bpcm::NoopAdapter>,
+        std::tuple<StorageValueType, trimd::neon::F256, TStrategyTestParams<0>, rl4::bpcm::EulerAnglesToQuaternions<tdm::fdeg,
+                                                                                                                    tdm::rot_seq::xyz> >,
+        std::tuple<StorageValueType, trimd::neon::F256, TStrategyTestParams<1>, rl4::bpcm::NoopAdapter>,
+        std::tuple<StorageValueType, trimd::neon::F256, TStrategyTestParams<1>, rl4::bpcm::EulerAnglesToQuaternions<tdm::fdeg,
+                                                                                                                    tdm::rot_seq::xyz> >,
+        std::tuple<StorageValueType, trimd::neon::F256, TStrategyTestParams<2>, rl4::bpcm::NoopAdapter>,
+        std::tuple<StorageValueType, trimd::neon::F256, TStrategyTestParams<2>, rl4::bpcm::EulerAnglesToQuaternions<tdm::fdeg,
+                                                                                                                    tdm::rot_seq::xyz> >,
+        std::tuple<StorageValueType, trimd::neon::F256, TStrategyTestParams<3>, rl4::bpcm::NoopAdapter>,
+        std::tuple<StorageValueType, trimd::neon::F256, TStrategyTestParams<3>, rl4::bpcm::EulerAnglesToQuaternions<tdm::fdeg,
+                                                                                                                    tdm::rot_seq::xyz> >,
 #endif  // RL_BUILD_WITH_NEON
 #if !defined(RL_USE_HALF_FLOATS)
-        std::tuple<StorageValueType, trimd::scalar::F256, TStrategyTestParams<0> >,
-        std::tuple<StorageValueType, trimd::scalar::F256, TStrategyTestParams<1> >,
-        std::tuple<StorageValueType, trimd::scalar::F256, TStrategyTestParams<2> >,
-        std::tuple<StorageValueType, trimd::scalar::F256, TStrategyTestParams<3> >,
+        std::tuple<StorageValueType, trimd::scalar::F256, TStrategyTestParams<0>, rl4::bpcm::NoopAdapter>,
+        std::tuple<StorageValueType, trimd::scalar::F256, TStrategyTestParams<0>,
+                   rl4::bpcm::EulerAnglesToQuaternions<tdm::fdeg, tdm::rot_seq::xyz> >,
+        std::tuple<StorageValueType, trimd::scalar::F256, TStrategyTestParams<1>, rl4::bpcm::NoopAdapter>,
+        std::tuple<StorageValueType, trimd::scalar::F256, TStrategyTestParams<1>,
+                   rl4::bpcm::EulerAnglesToQuaternions<tdm::fdeg, tdm::rot_seq::xyz> >,
+        std::tuple<StorageValueType, trimd::scalar::F256, TStrategyTestParams<2>, rl4::bpcm::NoopAdapter>,
+        std::tuple<StorageValueType, trimd::scalar::F256, TStrategyTestParams<2>,
+                   rl4::bpcm::EulerAnglesToQuaternions<tdm::fdeg, tdm::rot_seq::xyz> >,
+        std::tuple<StorageValueType, trimd::scalar::F256, TStrategyTestParams<3>, rl4::bpcm::NoopAdapter>,
+        std::tuple<StorageValueType, trimd::scalar::F256, TStrategyTestParams<3>,
+                   rl4::bpcm::EulerAnglesToQuaternions<tdm::fdeg, tdm::rot_seq::xyz> >,
 #endif  // RL_USE_HALF_FLOATS
     std::tuple<>
     >;
@@ -98,55 +139,106 @@ using Block8JointCalculationTypeList = ::testing::Types<
 TYPED_TEST_SUITE(Block8JointCalculationStrategyTest, Block8JointCalculationTypeList, );
 
 TYPED_TEST(Block8JointCalculationStrategyTest, Block8Padded) {
-    const OutputScope scope{this->params.lod, 0ul, 5ul};
-    auto joints = block8::OptimizedStorage<StorageValueType>::create(std::move(this->strategy), 0u, &this->memRes);
-    this->execute(joints, block8::output::valuesPerLOD, scope);
+    const std::uint16_t jointGroupIndex = 0u;
+    const auto& outputIndices = block8::optimized::outputIndices[this->rotationSelectorIndex][jointGroupIndex];
+    const OutputScope scope{this->params.lod, outputIndices[0], block8::unoptimized::outputIndices[jointGroupIndex].size()};
+    auto joints = block8::OptimizedStorage<StorageValueType>::create(std::move(this->strategy),
+                                                                     this->rotationSelectorIndex,
+                                                                     this->rotationType,
+                                                                     jointGroupIndex,
+                                                                     &this->memRes);
+    this->execute(joints, block8::output::valuesPerLOD[this->rotationSelectorIndex], scope);
 }
 
 TYPED_TEST(Block8JointCalculationStrategyTest, Block8Exact) {
-    const OutputScope scope{this->params.lod, 5ul, 8ul};
-    auto joints = block8::OptimizedStorage<StorageValueType>::create(std::move(this->strategy), 1u, &this->memRes);
-    this->execute(joints, block8::output::valuesPerLOD, scope);
+    const std::uint16_t jointGroupIndex = 1u;
+    const auto& outputIndices = block8::optimized::outputIndices[this->rotationSelectorIndex][jointGroupIndex];
+    const OutputScope scope{this->params.lod, outputIndices[0], block8::unoptimized::outputIndices[jointGroupIndex].size()};
+    auto joints = block8::OptimizedStorage<StorageValueType>::create(std::move(this->strategy),
+                                                                     this->rotationSelectorIndex,
+                                                                     this->rotationType,
+                                                                     jointGroupIndex,
+                                                                     &this->memRes);
+    this->execute(joints, block8::output::valuesPerLOD[this->rotationSelectorIndex], scope);
 }
 
 TYPED_TEST(Block8JointCalculationStrategyTest, Block16Padded) {
-    const OutputScope scope{this->params.lod, 13ul, 9ul};
-    auto joints = block8::OptimizedStorage<StorageValueType>::create(std::move(this->strategy), 2u, &this->memRes);
-    this->execute(joints, block8::output::valuesPerLOD, scope);
+    const std::uint16_t jointGroupIndex = 2u;
+    const auto& outputIndices = block8::optimized::outputIndices[this->rotationSelectorIndex][jointGroupIndex];
+    const OutputScope scope{this->params.lod, outputIndices[0], block8::unoptimized::outputIndices[jointGroupIndex].size()};
+    auto joints = block8::OptimizedStorage<StorageValueType>::create(std::move(this->strategy),
+                                                                     this->rotationSelectorIndex,
+                                                                     this->rotationType,
+                                                                     jointGroupIndex,
+                                                                     &this->memRes);
+    this->execute(joints, block8::output::valuesPerLOD[this->rotationSelectorIndex], scope);
 }
 
 TYPED_TEST(Block8JointCalculationStrategyTest, Block16Exact) {
-    const OutputScope scope{this->params.lod, 22ul, 16ul};
-    auto joints = block8::OptimizedStorage<StorageValueType>::create(std::move(this->strategy), 3u, &this->memRes);
-    this->execute(joints, block8::output::valuesPerLOD, scope);
+    const std::uint16_t jointGroupIndex = 3u;
+    const auto& outputIndices = block8::optimized::outputIndices[this->rotationSelectorIndex][jointGroupIndex];
+    const OutputScope scope{this->params.lod, outputIndices[0], block8::unoptimized::outputIndices[jointGroupIndex].size()};
+    auto joints = block8::OptimizedStorage<StorageValueType>::create(std::move(this->strategy),
+                                                                     this->rotationSelectorIndex,
+                                                                     this->rotationType,
+                                                                     jointGroupIndex,
+                                                                     &this->memRes);
+    this->execute(joints, block8::output::valuesPerLOD[this->rotationSelectorIndex], scope);
 }
 
 TYPED_TEST(Block8JointCalculationStrategyTest, Block24Padded) {
-    const OutputScope scope{this->params.lod, 38ul, 17ul};
-    auto joints = block8::OptimizedStorage<StorageValueType>::create(std::move(this->strategy), 4u, &this->memRes);
-    this->execute(joints, block8::output::valuesPerLOD, scope);
+    const std::uint16_t jointGroupIndex = 4u;
+    const auto& outputIndices = block8::optimized::outputIndices[this->rotationSelectorIndex][jointGroupIndex];
+    const OutputScope scope{this->params.lod, outputIndices[0], block8::unoptimized::outputIndices[jointGroupIndex].size()};
+    auto joints = block8::OptimizedStorage<StorageValueType>::create(std::move(this->strategy),
+                                                                     this->rotationSelectorIndex,
+                                                                     this->rotationType,
+                                                                     jointGroupIndex,
+                                                                     &this->memRes);
+    this->execute(joints, block8::output::valuesPerLOD[this->rotationSelectorIndex], scope);
 }
 
 TYPED_TEST(Block8JointCalculationStrategyTest, Block24Exact) {
-    const OutputScope scope{this->params.lod, 55ul, 24ul};
-    auto joints = block8::OptimizedStorage<StorageValueType>::create(std::move(this->strategy), 5u, &this->memRes);
-    this->execute(joints, block8::output::valuesPerLOD, scope);
+    const std::uint16_t jointGroupIndex = 5u;
+    const auto& outputIndices = block8::optimized::outputIndices[this->rotationSelectorIndex][jointGroupIndex];
+    const OutputScope scope{this->params.lod, outputIndices[0], block8::unoptimized::outputIndices[jointGroupIndex].size()};
+    auto joints = block8::OptimizedStorage<StorageValueType>::create(std::move(this->strategy),
+                                                                     this->rotationSelectorIndex,
+                                                                     this->rotationType,
+                                                                     jointGroupIndex,
+                                                                     &this->memRes);
+    this->execute(joints, block8::output::valuesPerLOD[this->rotationSelectorIndex], scope);
 }
 
 TYPED_TEST(Block8JointCalculationStrategyTest, Block32Padded) {
-    const OutputScope scope{this->params.lod, 79ul, 25ul};
-    auto joints = block8::OptimizedStorage<StorageValueType>::create(std::move(this->strategy), 6u, &this->memRes);
-    this->execute(joints, block8::output::valuesPerLOD, scope);
+    const std::uint16_t jointGroupIndex = 6u;
+    const auto& outputIndices = block8::optimized::outputIndices[this->rotationSelectorIndex][jointGroupIndex];
+    const OutputScope scope{this->params.lod, outputIndices[0], block8::unoptimized::outputIndices[jointGroupIndex].size()};
+    auto joints = block8::OptimizedStorage<StorageValueType>::create(std::move(this->strategy),
+                                                                     this->rotationSelectorIndex,
+                                                                     this->rotationType,
+                                                                     jointGroupIndex,
+                                                                     &this->memRes);
+    this->execute(joints, block8::output::valuesPerLOD[this->rotationSelectorIndex], scope);
 }
 
 TYPED_TEST(Block8JointCalculationStrategyTest, Block32Exact) {
-    const OutputScope scope{this->params.lod, 104ul, 32ul};
-    auto joints = block8::OptimizedStorage<StorageValueType>::create(std::move(this->strategy), 7u, &this->memRes);
-    this->execute(joints, block8::output::valuesPerLOD, scope);
+    const std::uint16_t jointGroupIndex = 7u;
+    const auto& outputIndices = block8::optimized::outputIndices[this->rotationSelectorIndex][jointGroupIndex];
+    const OutputScope scope{this->params.lod, outputIndices[0], block8::unoptimized::outputIndices[jointGroupIndex].size()};
+    auto joints = block8::OptimizedStorage<StorageValueType>::create(std::move(this->strategy),
+                                                                     this->rotationSelectorIndex,
+                                                                     this->rotationType,
+                                                                     jointGroupIndex,
+                                                                     &this->memRes);
+    this->execute(joints, block8::output::valuesPerLOD[this->rotationSelectorIndex], scope);
 }
 
 TYPED_TEST(Block8JointCalculationStrategyTest, MultipleBlocks) {
-    const OutputScope scope{this->params.lod, 0ul, block8::output::valuesPerLOD.front().size()};
-    auto joints = block8::OptimizedStorage<StorageValueType>::create(std::move(this->strategy), &this->memRes);
-    this->execute(joints, block8::output::valuesPerLOD, scope);
+    const OutputScope scope{this->params.lod, 0ul, block8::output::valuesPerLOD[this->rotationSelectorIndex].front().size()};
+    auto joints = block8::OptimizedStorage<StorageValueType>::create(std::move(this->strategy),
+                                                                     this->rotationSelectorIndex,
+                                                                     this->rotationType,
+                                                                     &this->memRes);
+    this->execute(joints, block8::output::valuesPerLOD[this->rotationSelectorIndex], scope);
 }
diff --git a/Engine/Plugins/Animation/RigLogic/Source/RigLogicLibTest/Private/rltests/joints/bpcm/TestBlock8StorageBuilder.cpp b/Engine/Plugins/Animation/RigLogic/Source/RigLogicLibTest/Private/rltests/joints/bpcm/TestBlock8StorageBuilder.cpp
--- a/Engine/Plugins/Animation/RigLogic/Source/RigLogicLibTest/Private/rltests/joints/bpcm/TestBlock8StorageBuilder.cpp
+++ b/Engine/Plugins/Animation/RigLogic/Source/RigLogicLibTest/Private/rltests/joints/bpcm/TestBlock8StorageBuilder.cpp
@@ -8,9 +8,12 @@
 #include "rltests/joints/bpcm/Helpers.h"
 
 #include "riglogic/TypeDefs.h"
+#include "riglogic/joints/JointBehaviorFilter.h"
 #include "riglogic/joints/JointsBuilder.h"
-#include "riglogic/joints/bpcm/Evaluator.h"
-#include "riglogic/joints/bpcm/strategies/Block8.h"
+#include "riglogic/joints/cpu/CPUJointsEvaluator.h"
+#include "riglogic/joints/cpu/bpcm/BPCMJointsEvaluator.h"
+#include "riglogic/joints/cpu/bpcm/CalculationStrategy.h"
+#include "riglogic/joints/cpu/bpcm/RotationAdapters.h"
 #include "riglogic/riglogic/RigLogic.h"
 #include "riglogic/system/simd/SIMD.h"
 
@@ -28,23 +31,39 @@ class Block8JointStorageBuilderTest : public ::testing::Test {
             using TValue = typename std::tuple_element<0, TestTypes>::type;
             using TFVec = typename std::tuple_element<1, TestTypes>::type;
             using TCalculationType = typename std::tuple_element<2, TestTypes>::type;
+            using TRotationAdapter = typename std::tuple_element<3, TestTypes>::type;
 
             rl4::Configuration config{};
             config.calculationType = TCalculationType::get();
+            config.rotationType = BPCMRotationOutputTypeSelector<TRotationAdapter>::rotation();
             auto builder = rl4::JointsBuilder::create(config, &memRes);
-            builder->computeStorageRequirements(&reader);
-            builder->allocateStorage(&reader);
-            builder->fillStorage(&reader);
+
+            rl4::JointBehaviorFilter filter{&reader, &memRes};
+            filter.include(dna::TranslationRepresentation::Vector);
+            filter.include(dna::RotationRepresentation::EulerAngles);
+            filter.include(dna::ScaleRepresentation::Vector);
+
+            builder->computeStorageRequirements(filter);
+            builder->allocateStorage(filter);
+            builder->fillStorage(filter);
             auto joints = builder->build();
-            auto jointsImpl = static_cast<rl4::bpcm::Evaluator<TValue>*>(joints.get());
+            auto jointsImpl = static_cast<rl4::CPUJointsEvaluator*>(joints.get());
+            auto bpcmJointsImpl =
+                static_cast<rl4::bpcm::Evaluator<TValue>*>(rl4::CPUJointsEvaluator::Accessor::getBPCMEvaluator(jointsImpl));
 
-            auto strategy = pma::UniqueInstance<rl4::bpcm::block8::Block8JointCalculationStrategy<TValue, TFVec>,
-                                                rl4::bpcm::JointCalculationStrategy<TValue> >::with(&memRes).create();
-            auto expected = block8::OptimizedStorage<TValue>::create(std::move(strategy), &memRes);
+            const auto rotationSelectorIndex = BPCMRotationOutputTypeSelector<TRotationAdapter>::value();
+            const auto rotationType = BPCMRotationOutputTypeSelector<TRotationAdapter>::rotation();
+            auto strategy = pma::UniqueInstance<rl4::bpcm::VectorizedJointGroupLinearCalculationStrategy<TValue, TFVec,
+                                                                                                         TRotationAdapter>,
+                                                rl4::bpcm::JointGroupLinearCalculationStrategy<TValue> >::with(&memRes).create();
+            auto expected = block8::OptimizedStorage<TValue>::create(std::move(strategy),
+                                                                     rotationSelectorIndex,
+                                                                     rotationType,
+                                                                     &memRes);
 
-            rl4::bpcm::Evaluator<TValue>::Accessor::assertRawDataEqual(*jointsImpl, expected);
-            rl4::bpcm::Evaluator<TValue>::Accessor::assertJointGroupsEqual(*jointsImpl, expected);
-            rl4::bpcm::Evaluator<TValue>::Accessor::assertLODsEqual(*jointsImpl, expected);
+            rl4::bpcm::Evaluator<TValue>::Accessor::assertRawDataEqual(*bpcmJointsImpl, expected);
+            rl4::bpcm::Evaluator<TValue>::Accessor::assertJointGroupsEqual(*bpcmJointsImpl, expected);
+            rl4::bpcm::Evaluator<TValue>::Accessor::assertLODsEqual(*bpcmJointsImpl, expected);
         }
 
     protected:
@@ -59,7 +78,9 @@ class Block8JointStorageBuilderTest : public ::testing::Test {
 // type. In all other cases- Block-4 storage optimizer will run.
 using Block8StorageValueTypeList = ::testing::Types<
 #if defined(RL_BUILD_WITH_AVX)
-        std::tuple<StorageValueType, trimd::avx::F256, TCalculationType<rl4::CalculationType::AVX> >,
+        std::tuple<StorageValueType, trimd::avx::F256, TCalculationType<rl4::CalculationType::AVX>, rl4::bpcm::NoopAdapter>,
+        std::tuple<StorageValueType, trimd::avx::F256, TCalculationType<rl4::CalculationType::AVX>,
+                   rl4::bpcm::EulerAnglesToQuaternions<tdm::fdeg, tdm::rot_seq::xyz> >,
 #endif  // RL_BUILD_WITH_AVX
     std::tuple<>
     >;
diff --git a/Engine/Plugins/Animation/RigLogic/Source/RigLogicLibTest/Private/rltests/joints/bpcm/TestLODRegion.cpp b/Engine/Plugins/Animation/RigLogic/Source/RigLogicLibTest/Private/rltests/joints/bpcm/TestLODRegion.cpp
--- a/Engine/Plugins/Animation/RigLogic/Source/RigLogicLibTest/Private/rltests/joints/bpcm/TestLODRegion.cpp
+++ b/Engine/Plugins/Animation/RigLogic/Source/RigLogicLibTest/Private/rltests/joints/bpcm/TestLODRegion.cpp
@@ -2,7 +2,7 @@
 
 #include "rltests/Defs.h"
 
-#include "riglogic/joints/bpcm/LODRegion.h"
+#include "riglogic/joints/cpu/bpcm/JointGroup.h"
 
 #include <cstdint>
 
@@ -13,8 +13,8 @@ struct LODTestSetup {
     std::uint32_t lodEndRow;
     std::uint32_t blockHeight;
     std::uint32_t padTo;
-    std::uint32_t lodEndRowAlignedToLastFullBlock;
-    std::uint32_t lodEndRowAlignedToSecondLastFullBlock;
+    std::uint32_t lodEndRowPaddedToLastFullBlock;
+    std::uint32_t lodEndRowPaddedToSecondLastFullBlock;
 };
 
 class LODRegionTest : public ::testing::TestWithParam<LODTestSetup> {
@@ -26,16 +26,16 @@ TEST(LODRegionTest, ConstructFromFixtures) {
     rl4::bpcm::LODRegion lod{1, 2, 3};
 
     ASSERT_EQ(lod.size, 1u);
-    ASSERT_EQ(lod.sizeAlignedToLastFullBlock, 2u);
-    ASSERT_EQ(lod.sizeAlignedToSecondLastFullBlock, 3u);
+    ASSERT_EQ(lod.sizePaddedToLastFullBlock, 2u);
+    ASSERT_EQ(lod.sizePaddedToSecondLastFullBlock, 3u);
 }
 
 TEST_P(LODRegionTest, ConstructLODRegionFromVaryingParameters) {
     auto params = GetParam();
     rl4::bpcm::LODRegion lod{params.lodEndRow, params.rowCount, params.blockHeight, params.padTo};
     ASSERT_EQ(lod.size, params.lodEndRow);
-    ASSERT_EQ(lod.sizeAlignedToLastFullBlock, params.lodEndRowAlignedToLastFullBlock);
-    ASSERT_EQ(lod.sizeAlignedToSecondLastFullBlock, params.lodEndRowAlignedToSecondLastFullBlock);
+    ASSERT_EQ(lod.sizePaddedToLastFullBlock, params.lodEndRowPaddedToLastFullBlock);
+    ASSERT_EQ(lod.sizePaddedToSecondLastFullBlock, params.lodEndRowPaddedToSecondLastFullBlock);
 }
 
 INSTANTIATE_TEST_SUITE_P(
diff --git a/Engine/Plugins/Animation/RigLogic/Source/RigLogicLibTest/Private/rltests/joints/quaternions/TestQuaternionEvaluator.cpp b/Engine/Plugins/Animation/RigLogic/Source/RigLogicLibTest/Private/rltests/joints/quaternions/TestQuaternionEvaluator.cpp
new file mode 100644
--- /dev/null
+++ b/Engine/Plugins/Animation/RigLogic/Source/RigLogicLibTest/Private/rltests/joints/quaternions/TestQuaternionEvaluator.cpp
@@ -0,0 +1,172 @@
+// Copyright Epic Games, Inc. All Rights Reserved.
+
+#include "rltests/Defs.h"
+#include "rltests/StorageValueType.h"
+#include "rltests/controls/ControlFixtures.h"
+#include "rltests/joints/Helpers.h"
+#include "rltests/joints/quaternions/QuaternionFixtures.h"
+
+#include "riglogic/TypeDefs.h"
+#include "riglogic/joints/cpu/CPUJointsOutputInstance.h"
+#include "riglogic/joints/cpu/quaternions/QuaternionJointsEvaluator.h"
+#include "riglogic/joints/cpu/quaternions/RotationAdapters.h"
+#include "riglogic/system/simd/Detect.h"
+#include "riglogic/system/simd/SIMD.h"
+
+#include <tuple>
+
+#ifdef _MSC_VER
+    #pragma warning(push)
+    #pragma warning(disable : 4324)
+#endif
+
+namespace {
+
+template<typename TTestTypes>
+class QuaternionEvaluatorTest : public ::testing::Test {
+    protected:
+        void SetUp() override {
+            QuaternionEvaluatorTest::SetUpImpl();
+        }
+
+        template<typename TestTypes = TTestTypes>
+        typename std::enable_if<std::tuple_size<TestTypes>::value != 0ul, void>::type SetUpImpl() {
+            using T = typename std::tuple_element<0, TestTypes>::type;
+            using TF256 = typename std::tuple_element<1, TestTypes>::type;
+            using TF128 = typename std::tuple_element<2, TestTypes>::type;
+            using TRotationAdapter = typename std::tuple_element<3, TestTypes>::type;
+            using QJEvaluator = rl4::QuaternionJointsEvaluator<T>;
+            using CalculationStrategy = rl4::VectorizedJointGroupQuaternionCalculationStrategy<T, TF256, TF128, TRotationAdapter>;
+
+            rotationSelectorIndex = RotationOutputTypeSelector<TRotationAdapter>::value();
+            rotationType = RotationOutputTypeSelector<TRotationAdapter>::rotation();
+            const auto& values = rltests::qs::optimized::Values<T>::get();
+            rl4::Vector<rl4::JointGroup<T> > jointGroups{values.size(), rl4::JointGroup<T>{&memRes}, &memRes};
+            for (std::size_t jgi = {}; jgi < jointGroups.size(); ++jgi) {
+                auto& jointGroup = jointGroups[jgi];
+                jointGroup.values = values[jgi];
+                jointGroup.inputIndices = rltests::qs::optimized::inputIndices[jgi];
+                jointGroup.outputIndices = rltests::qs::optimized::outputIndices[rotationSelectorIndex][jgi];
+                jointGroup.lods = rltests::qs::optimized::lodRegions[jgi];
+            }
+            auto factory = rl4::UniqueInstance<QJEvaluator, rl4::JointsEvaluator>::with(&memRes);
+            auto strategy = rl4::UniqueInstance<CalculationStrategy, rl4::JointGroupQuaternionCalculationStrategy<T> >::with(
+                &memRes).create();
+            evaluator = factory.create(std::move(strategy), std::move(jointGroups), nullptr, &memRes);
+        }
+
+        template<typename TestTypes = TTestTypes>
+        typename std::enable_if<std::tuple_size<TestTypes>::value == 0ul, void>::type SetUpImpl() {
+            GTEST_SKIP();
+        }
+
+    protected:
+        pma::AlignedMemoryResource memRes;
+        rltests::qs::QuaternionReader reader;
+        rl4::JointsEvaluator::Pointer evaluator;
+        std::size_t rotationSelectorIndex;
+        rl4::RotationType rotationType;
+
+};
+
+}  // namespace
+
+#if defined(RL_BUILD_WITH_AVX) && defined(RL_BUILD_WITH_SSE)
+    using QuaternionEvaluatorTypeList = ::testing::Types<
+    #ifdef RL_USE_HALF_FLOATS
+            std::tuple<StorageValueType, trimd::avx::F256, trimd::sse::F128, rl4::QuaternionsToEulerAngles<tdm::frad,
+                                                                                                           tdm::rot_seq::xyz> >,
+            std::tuple<StorageValueType, trimd::avx::F256, trimd::sse::F128, rl4::PassthroughAdapter>,
+            std::tuple<StorageValueType, trimd::sse::F256, trimd::sse::F128, rl4::QuaternionsToEulerAngles<tdm::frad,
+                                                                                                           tdm::rot_seq::xyz> >,
+            std::tuple<StorageValueType, trimd::sse::F256, trimd::sse::F128, rl4::PassthroughAdapter>
+    #else
+            std::tuple<StorageValueType, trimd::avx::F256, trimd::sse::F128, rl4::QuaternionsToEulerAngles<tdm::frad,
+                                                                                                           tdm::rot_seq::xyz> >,
+            std::tuple<StorageValueType, trimd::avx::F256, trimd::sse::F128, rl4::PassthroughAdapter>,
+            std::tuple<StorageValueType, trimd::sse::F256, trimd::sse::F128, rl4::QuaternionsToEulerAngles<tdm::frad,
+                                                                                                           tdm::rot_seq::xyz> >,
+            std::tuple<StorageValueType, trimd::sse::F256, trimd::sse::F128, rl4::PassthroughAdapter>,
+            std::tuple<StorageValueType, trimd::scalar::F256, trimd::scalar::F128, rl4::QuaternionsToEulerAngles<tdm::frad,
+                                                                                                                 tdm::rot_seq::xyz> >,
+            std::tuple<StorageValueType, trimd::scalar::F256, trimd::scalar::F128, rl4::PassthroughAdapter>
+    #endif  // RL_USE_HALF_FLOATS
+        >;
+#elif defined(RL_BUILD_WITH_AVX)
+    using QuaternionEvaluatorTypeList = ::testing::Types<
+    #ifdef RL_USE_HALF_FLOATS
+            std::tuple<StorageValueType, trimd::avx::F256, trimd::sse::F128, rl4::QuaternionsToEulerAngles<tdm::frad,
+                                                                                                           tdm::rot_seq::xyz> >,
+            std::tuple<StorageValueType, trimd::avx::F256, trimd::sse::F128, rl4::PassthroughAdapter>
+    #else
+            std::tuple<StorageValueType, trimd::avx::F256, trimd::sse::F128, rl4::QuaternionsToEulerAngles<tdm::frad,
+                                                                                                           tdm::rot_seq::xyz> >,
+            std::tuple<StorageValueType, trimd::avx::F256, trimd::sse::F128, rl4::PassthroughAdapter>,
+            std::tuple<StorageValueType, trimd::scalar::F256, trimd::scalar::F128, rl4::QuaternionsToEulerAngles<tdm::frad,
+                                                                                                                 tdm::rot_seq::xyz> >,
+            std::tuple<StorageValueType, trimd::scalar::F256, trimd::scalar::F128, rl4::PassthroughAdapter>
+    #endif  // RL_USE_HALF_FLOATS
+        >;
+#elif defined(RL_BUILD_WITH_SSE)
+    using QuaternionEvaluatorTypeList = ::testing::Types<
+    #ifdef RL_USE_HALF_FLOATS
+            std::tuple<StorageValueType, trimd::sse::F256, trimd::sse::F128, rl4::QuaternionsToEulerAngles<tdm::frad,
+                                                                                                           tdm::rot_seq::xyz> >,
+            std::tuple<StorageValueType, trimd::sse::F256, trimd::sse::F128, rl4::PassthroughAdapter>
+    #else
+            std::tuple<StorageValueType, trimd::sse::F256, trimd::sse::F128, rl4::QuaternionsToEulerAngles<tdm::frad,
+                                                                                                           tdm::rot_seq::xyz> >,
+            std::tuple<StorageValueType, trimd::sse::F256, trimd::sse::F128, rl4::PassthroughAdapter>,
+            std::tuple<StorageValueType, trimd::scalar::F256, trimd::scalar::F128, rl4::QuaternionsToEulerAngles<tdm::frad,
+                                                                                                                 tdm::rot_seq::xyz> >,
+            std::tuple<StorageValueType, trimd::scalar::F256, trimd::scalar::F128, rl4::PassthroughAdapter>
+    #endif  // RL_USE_HALF_FLOATS
+        >;
+#else
+    #ifndef RL_USE_HALF_FLOATS
+        using QuaternionEvaluatorTypeList = ::testing::Types<
+            std::tuple<StorageValueType, trimd::scalar::F256, trimd::scalar::F128, rl4::QuaternionsToEulerAngles<tdm::frad,
+                                                                                                                 tdm::rot_seq::xyz> >,
+            std::tuple<StorageValueType, trimd::scalar::F256, trimd::scalar::F128, rl4::PassthroughAdapter>
+            >;
+    #else
+        using QuaternionEvaluatorTypeList = ::testing::Types<std::tuple<> >;
+    #endif  // RL_USE_HALF_FLOATS
+#endif
+
+TYPED_TEST_SUITE(QuaternionEvaluatorTest, QuaternionEvaluatorTypeList, );
+
+TYPED_TEST(QuaternionEvaluatorTest, EvaluatePerLOD) {
+    const auto jointAttrCount =
+        static_cast<std::uint16_t>(rltests::qs::output::valuesPerLODPerConfig[this->rotationSelectorIndex][0].size());
+    rl4::CPUJointsOutputInstance outputInstance{jointAttrCount,
+                                                rl4::TranslationType::Vector,
+                                                this->rotationType,
+                                                rl4::ScaleType::Vector,
+                                                &this->memRes};
+    auto outputBuffer = outputInstance.getOutputBuffer();
+    auto inputInstanceFactory = ControlsFactory::getInstanceFactory(0,
+                                                                    static_cast<std::uint16_t>(rltests::qs::input::values.size()),
+                                                                    0,
+                                                                    0,
+                                                                    0);
+    auto inputInstance = inputInstanceFactory(&this->memRes);
+    auto inputBuffer = inputInstance->getInputBuffer();
+    std::copy(rltests::qs::input::values.begin(), rltests::qs::input::values.end(), inputBuffer.begin());
+
+    for (std::uint16_t lod = 0u; lod < rltests::qs::unoptimized::lodCount; ++lod) {
+        std::fill(outputBuffer.begin(), outputBuffer.end(), 0.0f);
+        this->evaluator->calculate(inputInstance.get(), &outputInstance, lod);
+        const auto& expected = rltests::qs::output::valuesPerLODPerConfig[this->rotationSelectorIndex][lod];
+        #ifdef RL_USE_HALF_FLOATS
+            static constexpr float threshold = 0.05f;
+        #else
+            static constexpr float threshold = 0.002f;
+        #endif  // RL_USE_HALF_FLOATS
+        ASSERT_ELEMENTS_NEAR(outputBuffer, expected, expected.size(), threshold);
+    }
+}
+
+#ifdef _MSC_VER
+    #pragma warning(pop)
+#endif
diff --git a/Engine/Plugins/Animation/RigLogic/Source/RigLogicLibTest/Private/rltests/joints/quaternions/TestQuaternionStorageBuilder.cpp b/Engine/Plugins/Animation/RigLogic/Source/RigLogicLibTest/Private/rltests/joints/quaternions/TestQuaternionStorageBuilder.cpp
new file mode 100644
--- /dev/null
+++ b/Engine/Plugins/Animation/RigLogic/Source/RigLogicLibTest/Private/rltests/joints/quaternions/TestQuaternionStorageBuilder.cpp
@@ -0,0 +1,163 @@
+// Copyright Epic Games, Inc. All Rights Reserved.
+
+#include "rltests/Defs.h"
+#include "rltests/StorageValueType.h"
+#include "rltests/joints/Helpers.h"
+#include "rltests/joints/quaternions/QuaternionFixtures.h"
+
+#include "riglogic/TypeDefs.h"
+#include "riglogic/joints/cpu/quaternions/QuaternionJointsBuilder.h"
+#include "riglogic/system/simd/Detect.h"
+
+#include <tuple>
+
+#ifdef _MSC_VER
+    #pragma warning(push)
+    #pragma warning(disable : 4324)
+#endif
+
+namespace rl4 {
+
+template<typename T>
+struct QuaternionJointsEvaluator<T>::Accessor {
+
+    static void assertRawDataEqual(const QuaternionJointsEvaluator<T>& result, std::size_t rotationSelectorIndex) {
+        ASSERT_EQ(result.jointGroups.size(), rltests::qs::unoptimized::subMatrices.size());
+        const auto& values = rltests::qs::optimized::Values<T>::get();
+        const auto& inputIndices = rltests::qs::optimized::inputIndices;
+        const auto& outputIndices = rltests::qs::optimized::outputIndices[rotationSelectorIndex];
+        const auto& lods = rltests::qs::optimized::lodRegions;
+        for (std::size_t jgi = {}; jgi < result.jointGroups.size(); ++jgi) {
+            const auto& jointGroup = result.jointGroups[jgi];
+            ASSERT_ELEMENTS_NEAR(jointGroup.values, values[jgi], values[jgi].size(), 0.0002f);
+            ASSERT_ELEMENTS_EQ(jointGroup.inputIndices, inputIndices[jgi], inputIndices[jgi].size());
+            ASSERT_ELEMENTS_EQ(jointGroup.outputIndices, outputIndices[jgi], outputIndices[jgi].size());
+            ASSERT_EQ(jointGroup.lods.size(), lods[jgi].size());
+            for (std::size_t lod = {}; lod < lods[jgi].size(); ++lod) {
+                const auto& resultLOD = jointGroup.lods[lod];
+                const auto& expectedLOD = lods[jgi][lod];
+                ASSERT_EQ(resultLOD.size, expectedLOD.size);
+                ASSERT_EQ(resultLOD.sizePaddedToLastFullBlock, expectedLOD.sizePaddedToLastFullBlock);
+                ASSERT_EQ(resultLOD.sizePaddedToSecondLastFullBlock, expectedLOD.sizePaddedToSecondLastFullBlock);
+            }
+        }
+    }
+
+};
+
+}  // namespace rl4
+
+namespace {
+
+template<typename TTestTypes>
+class QuaternionJointStorageBuilderTest : public ::testing::Test {
+    protected:
+        template<typename TestTypes = TTestTypes>
+        typename std::enable_if<std::tuple_size<TestTypes>::value != 0ul, void>::type buildStorage() {
+            using T = typename std::tuple_element<0, TTestTypes>::type;
+            using TF256 = typename std::tuple_element<1, TTestTypes>::type;
+            using TF128 = typename std::tuple_element<2, TTestTypes>::type;
+            using TRotationAdapter = typename std::tuple_element<3, TTestTypes>::type;
+            using QJEvaluator = rl4::QuaternionJointsEvaluator<T>;
+
+            rl4::Configuration config{};
+            config.rotationType = RotationOutputTypeSelector<TRotationAdapter>::rotation();
+            rl4::QuaternionJointsBuilder<T, TF256, TF128> builder(config, &memRes);
+
+            rl4::JointBehaviorFilter filter{&reader, &memRes};
+            filter.include(dna::RotationRepresentation::Quaternion);
+
+            builder.computeStorageRequirements(filter);
+            builder.allocateStorage(filter);
+            builder.fillStorage(filter);
+            auto joints = builder.build();
+
+            const auto rotationSelectorIndex = RotationOutputTypeSelector<TRotationAdapter>::value();
+            rl4::QuaternionJointsEvaluator<T>::Accessor::assertRawDataEqual(*static_cast<QJEvaluator*>(joints.get()),
+                                                                            rotationSelectorIndex);
+        }
+
+        template<typename TestTypes = TTestTypes>
+        typename std::enable_if<std::tuple_size<TestTypes>::value == 0ul, void>::type buildStorage() {
+        }
+
+    protected:
+        pma::AlignedMemoryResource memRes;
+        rltests::qs::QuaternionReader reader;
+
+};
+
+}  // namespace
+
+#if defined(RL_BUILD_WITH_AVX) && defined(RL_BUILD_WITH_SSE)
+    using StorageValueTypeList = ::testing::Types<
+    #ifdef RL_USE_HALF_FLOATS
+            std::tuple<StorageValueType, trimd::avx::F256, trimd::sse::F128, rl4::PassthroughAdapter>,
+            std::tuple<StorageValueType, trimd::avx::F256, trimd::sse::F128, rl4::QuaternionsToEulerAngles<tdm::frad,
+                                                                                                           tdm::rot_seq::xyz> >,
+            std::tuple<StorageValueType, trimd::sse::F256, trimd::sse::F128, rl4::PassthroughAdapter>,
+            std::tuple<StorageValueType, trimd::sse::F256, trimd::sse::F128, rl4::QuaternionsToEulerAngles<tdm::frad,
+                                                                                                           tdm::rot_seq::xyz> >
+    #else
+            std::tuple<StorageValueType, trimd::avx::F256, trimd::sse::F128, rl4::PassthroughAdapter>,
+            std::tuple<StorageValueType, trimd::avx::F256, trimd::sse::F128, rl4::QuaternionsToEulerAngles<tdm::frad,
+                                                                                                           tdm::rot_seq::xyz> >,
+            std::tuple<StorageValueType, trimd::sse::F256, trimd::sse::F128, rl4::PassthroughAdapter>,
+            std::tuple<StorageValueType, trimd::sse::F256, trimd::sse::F128, rl4::QuaternionsToEulerAngles<tdm::frad,
+                                                                                                           tdm::rot_seq::xyz> >,
+            std::tuple<StorageValueType, trimd::scalar::F256, trimd::scalar::F128, rl4::PassthroughAdapter>,
+            std::tuple<StorageValueType, trimd::scalar::F256, trimd::scalar::F128, rl4::QuaternionsToEulerAngles<tdm::frad,
+                                                                                                                 tdm::rot_seq::xyz> >
+    #endif  // RL_USE_HALF_FLOATS
+        >;
+#elif defined(RL_BUILD_WITH_AVX)
+    using StorageValueTypeList = ::testing::Types<
+    #ifdef RL_USE_HALF_FLOATS
+            std::tuple<StorageValueType, trimd::avx::F256, trimd::sse::F128, rl4::PassthroughAdapter>,
+            std::tuple<StorageValueType, trimd::avx::F256, trimd::sse::F128, rl4::QuaternionsToEulerAngles<tdm::frad,
+                                                                                                           tdm::rot_seq::xyz> >
+    #else
+            std::tuple<StorageValueType, trimd::avx::F256, trimd::sse::F128, rl4::PassthroughAdapter>,
+            std::tuple<StorageValueType, trimd::avx::F256, trimd::sse::F128, rl4::QuaternionsToEulerAngles<tdm::frad,
+                                                                                                           tdm::rot_seq::xyz> >,
+            std::tuple<StorageValueType, trimd::scalar::F256, trimd::scalar::F128, rl4::PassthroughAdapter>,
+            std::tuple<StorageValueType, trimd::scalar::F256, trimd::scalar::F128, rl4::QuaternionsToEulerAngles<tdm::frad,
+                                                                                                                 tdm::rot_seq::xyz> >
+    #endif  // RL_USE_HALF_FLOATS
+        >;
+#elif defined(RL_BUILD_WITH_SSE)
+    using StorageValueTypeList = ::testing::Types<
+    #ifdef RL_USE_HALF_FLOATS
+            std::tuple<StorageValueType, trimd::sse::F256, trimd::sse::F128, rl4::PassthroughAdapter>,
+            std::tuple<StorageValueType, trimd::sse::F256, trimd::sse::F128, rl4::QuaternionsToEulerAngles<tdm::frad,
+                                                                                                           tdm::rot_seq::xyz> >
+    #else
+            std::tuple<StorageValueType, trimd::sse::F256, trimd::sse::F128, rl4::PassthroughAdapter>,
+            std::tuple<StorageValueType, trimd::sse::F256, trimd::sse::F128, rl4::QuaternionsToEulerAngles<tdm::frad,
+                                                                                                           tdm::rot_seq::xyz> >,
+            std::tuple<StorageValueType, trimd::scalar::F256, trimd::scalar::F128, rl4::PassthroughAdapter>,
+            std::tuple<StorageValueType, trimd::scalar::F256, trimd::scalar::F128, rl4::QuaternionsToEulerAngles<tdm::frad,
+                                                                                                                 tdm::rot_seq::xyz> >
+    #endif  // RL_USE_HALF_FLOATS
+        >;
+#else
+    #ifndef RL_USE_HALF_FLOATS
+        using StorageValueTypeList = ::testing::Types<
+            std::tuple<StorageValueType, trimd::scalar::F256, trimd::scalar::F128, rl4::PassthroughAdapter>,
+            std::tuple<StorageValueType, trimd::scalar::F256, trimd::scalar::F128, rl4::QuaternionsToEulerAngles<tdm::frad,
+                                                                                                                 tdm::rot_seq::xyz> >
+            >;
+    #else
+        using StorageValueTypeList = ::testing::Types<std::tuple<> >;
+    #endif  // RL_USE_HALF_FLOATS
+#endif
+
+TYPED_TEST_SUITE(QuaternionJointStorageBuilderTest, StorageValueTypeList, );
+
+TYPED_TEST(QuaternionJointStorageBuilderTest, LayoutOptimization) {
+    this->buildStorage();
+}
+
+#ifdef _MSC_VER
+    #pragma warning(pop)
+#endif
diff --git a/Engine/Plugins/Animation/RigLogic/Source/RigLogicLibTest/Private/rltests/joints/twistswing/TwistSwingFixtures.cpp b/Engine/Plugins/Animation/RigLogic/Source/RigLogicLibTest/Private/rltests/joints/twistswing/TwistSwingFixtures.cpp
new file mode 100644
--- /dev/null
+++ b/Engine/Plugins/Animation/RigLogic/Source/RigLogicLibTest/Private/rltests/joints/twistswing/TwistSwingFixtures.cpp
@@ -0,0 +1,1025 @@
+// Copyright Epic Games, Inc. All Rights Reserved.
+
+#include "rltests/joints/twistswing/TwistSwingFixtures.h"
+
+#include "riglogic/TypeDefs.h"
+#include "riglogic/joints/cpu/CPUJointsOutputInstance.h"
+#include "riglogic/joints/cpu/twistswing/TwistSwingJointsEvaluator.h"
+#include "riglogic/types/Extent.h"
+
+namespace rltests {
+
+namespace twsw {
+
+#ifdef __clang__
+    #pragma clang diagnostic push
+    #pragma clang diagnostic ignored "-Wglobal-constructors"
+    #pragma clang diagnostic ignored "-Wexit-time-destructors"
+#endif
+
+namespace unoptimized {
+
+const std::uint16_t lodCount = 4u;
+
+const pma::Matrix<float> swingBlendWeights = {
+    {0.3f, 0.7f},
+    {0.5f, 0.5f},
+    {1.0f}
+};
+
+const pma::Matrix<std::uint16_t> swingOutputJointIndices = {
+    {2u, 3u},
+    {7u, 9u},
+    {8u}
+};
+
+const pma::Matrix<std::uint16_t> swingInputControlIndices = {
+    {0u, 1u, 2u, 3u},
+    {4u, 5u, 6u, 7u},
+    {8u, 9u, 10u, 11u}
+};
+
+const pma::Vector<dna::TwistAxis> swingTwistAxes = {
+    dna::TwistAxis::X,
+    dna::TwistAxis::X,
+    dna::TwistAxis::X
+};
+
+const pma::Matrix<float> twistBlendWeights = {
+    {0.2f, 0.8f},
+    {0.6f, 0.4f},
+    {1.0f}
+};
+
+const pma::Matrix<std::uint16_t> twistOutputJointIndices = {
+    {0u, 1u},
+    {4u, 6u},
+    {5u}
+};
+
+const pma::Matrix<std::uint16_t> twistInputControlIndices = {
+    {0u, 1u, 2u, 3u},
+    {4u, 5u, 6u, 7u},
+    {8u, 9u, 10u, 11u}
+};
+
+const pma::Vector<dna::TwistAxis> twistTwistAxes = {
+    dna::TwistAxis::X,
+    dna::TwistAxis::X,
+    dna::TwistAxis::X
+};
+
+}  // namespace unoptimized
+
+namespace optimized {
+
+const std::uint16_t setupCount = 3u;
+
+const pma::Matrix<float> swingBlendWeights = {
+    {0.3f, 0.7f},
+    {0.5f, 0.5f},
+    {1.0f}
+};
+
+const pma::Vector<pma::Matrix<std::uint16_t> > swingOutputIndices = {
+    {  // Quaternion outputs
+        {23, 24, 25, 26, 33, 34, 35, 36},
+        {73, 74, 75, 76, 93, 94, 95, 96},
+        {83, 84, 85, 86}
+    },
+    {  // Euler-angle outputs
+        {21, 22, 23, 0, 30, 31, 32, 0},
+        {66, 67, 68, 0, 84, 85, 86, 0},
+        {75, 76, 77, 0}
+    }
+};
+
+const pma::Matrix<std::uint16_t> swingInputIndices = {
+    {0u, 1u, 2u, 3u},
+    {4u, 5u, 6u, 7u},
+    {8u, 9u, 10u, 11u}
+};
+
+const pma::Matrix<float> twistBlendWeights = {
+    {0.2f, 0.8f},
+    {0.6f, 0.4f},
+    {1.0f}
+};
+
+const pma::Vector<pma::Matrix<std::uint16_t> > twistOutputIndices = {
+    {  // Quaternion outputs
+        {3, 4, 5, 6, 13, 14, 15, 16},
+        {43, 44, 45, 46, 63, 64, 65, 66},
+        {53, 54, 55, 56}
+    },
+    {  // Euler-angle outputs
+        {3, 4, 5, 0, 12, 13, 14, 0},
+        {39, 40, 41, 0, 57, 58, 59, 0},
+        {48, 49, 50, 0}
+    }
+};
+
+const pma::Matrix<std::uint16_t> twistInputIndices = {
+    {0u, 1u, 2u, 3u},
+    {4u, 5u, 6u, 7u},
+    {8u, 9u, 10u, 11u}
+};
+
+}  // namespace optimized
+
+namespace input {
+
+// Calculation input values
+const rl4::Vector<float> values = {
+    0.308434369465179f,
+    0.3353868017170958f,
+    0.14504533469235772f,
+    0.8782629354871905f,
+    0.24610924459196282f,
+    0.19641565253465743f,
+    0.40743077341317624f,
+    0.8572346796774493f,
+    0.24550619238838187f,
+    0.5267005234522829f,
+    0.3547646971819078f,
+    0.7324310737041945f,
+    0.0f
+};
+
+}  // namespace input
+
+namespace output {
+
+// Expected output results for each LOD
+const rl4::Vector<rl4::Matrix<float> > valuesPerLODPerConfig = {
+    {  // Quaternion outputs
+        {
+            // LOD-0
+            0.0f,
+            0.0f,
+            0.0f,
+            -0.26691f,
+            0.0f,
+            0.0f,
+            0.963722f,
+            0.0f,
+            0.0f,
+            0.0f,
+            0.0f,
+            0.0f,
+            0.0f,
+            -0.067495f,
+            0.0f,
+            0.0f,
+            0.99772f,
+            0.0f,
+            0.0f,
+            0.0f,
+            0.0f,
+            0.0f,
+            0.0f,
+            -0.320053f,
+            -0.2376f,
+            -0.102755f,
+            0.911347f,
+            0.0f,
+            0.0f,
+            0.0f,
+            0.0f,
+            0.0f,
+            0.0f,
+            -0.329264f,
+            -0.102785f,
+            -0.0444518f,
+            0.937574f,
+            0.0f,
+            0.0f,
+            0.0f,
+            0.0f,
+            0.0f,
+            0.0f,
+            -0.111598f,
+            0.0f,
+            0.0f,
+            0.993753f,
+            0.0f,
+            0.0f,
+            0.0f,
+            0.0f,
+            0.0f,
+            0.0f,
+            0.0f,
+            0.0f,
+            0.0f,
+            1.0f,
+            0.0f,
+            0.0f,
+            0.0f,
+            0.0f,
+            0.0f,
+            0.0f,
+            -0.166961f,
+            0.0f,
+            0.0f,
+            0.985964f,
+            0.0f,
+            0.0f,
+            0.0f,
+            0.0f,
+            0.0f,
+            0.0f,
+            -0.268386f,
+            -0.100976f,
+            -0.209457f,
+            0.934827f,
+            0.0f,
+            0.0f,
+            0.0f,
+            0.0f,
+            0.0f,
+            0.0f,
+            -0.317815f,
+            0.0f,
+            0.0f,
+            0.948153f,
+            0.0f,
+            0.0f,
+            0.0f,
+            0.0f,
+            0.0f,
+            0.0f,
+            -0.268386f,
+            -0.100976f,
+            -0.209457f,
+            0.934827f,
+            0.0f,
+            0.0f,
+            0.0f,
+            0.0f,
+            0.0f,
+            0.0f,
+            0.0f,
+            0.0f,
+            0.0f,
+            0.0f,
+            0.0f,
+            0.0f,
+            0.0f
+        }, {
+            // LOD-1
+            0.0f,
+            0.0f,
+            0.0f,
+            -0.26691f,
+            0.0f,
+            0.0f,
+            0.963722f,
+            0.0f,
+            0.0f,
+            0.0f,
+            0.0f,
+            0.0f,
+            0.0f,
+            -0.067495f,
+            0.0f,
+            0.0f,
+            0.99772f,
+            0.0f,
+            0.0f,
+            0.0f,
+            0.0f,
+            0.0f,
+            0.0f,
+            -0.320053f,
+            -0.2376f,
+            -0.102755f,
+            0.911347f,
+            0.0f,
+            0.0f,
+            0.0f,
+            0.0f,
+            0.0f,
+            0.0f,
+            -0.329264f,
+            -0.102785f,
+            -0.0444518f,
+            0.937574f,
+            0.0f,
+            0.0f,
+            0.0f,
+            0.0f,
+            0.0f,
+            0.0f,
+            -0.111598f,
+            0.0f,
+            0.0f,
+            0.993753f,
+            0.0f,
+            0.0f,
+            0.0f,
+            0.0f,
+            0.0f,
+            0.0f,
+            0.0f,
+            0.0f,
+            0.0f,
+            1.0f,
+            0.0f,
+            0.0f,
+            0.0f,
+            0.0f,
+            0.0f,
+            0.0f,
+            -0.166961f,
+            0.0f,
+            0.0f,
+            0.985964f,
+            0.0f,
+            0.0f,
+            0.0f,
+            0.0f,
+            0.0f,
+            0.0f,
+            -0.268386f,
+            -0.100976f,
+            -0.209457f,
+            0.934827f,
+            0.0f,
+            0.0f,
+            0.0f,
+            0.0f,
+            0.0f,
+            0.0f,
+            -0.317815f,
+            0.0f,
+            0.0f,
+            0.948153f,
+            0.0f,
+            0.0f,
+            0.0f,
+            0.0f,
+            0.0f,
+            0.0f,
+            -0.268386f,
+            -0.100976f,
+            -0.209457f,
+            0.934827f,
+            0.0f,
+            0.0f,
+            0.0f,
+            0.0f,
+            0.0f,
+            0.0f,
+            0.0f,
+            0.0f,
+            0.0f,
+            0.0f,
+            0.0f,
+            0.0f,
+            0.0f
+        }, {
+            // LOD-2
+            0.0f,
+            0.0f,
+            0.0f,
+            -0.26691f,
+            0.0f,
+            0.0f,
+            0.963722f,
+            0.0f,
+            0.0f,
+            0.0f,
+            0.0f,
+            0.0f,
+            0.0f,
+            -0.067495f,
+            0.0f,
+            0.0f,
+            0.99772f,
+            0.0f,
+            0.0f,
+            0.0f,
+            0.0f,
+            0.0f,
+            0.0f,
+            -0.320053f,
+            -0.2376f,
+            -0.102755f,
+            0.911347f,
+            0.0f,
+            0.0f,
+            0.0f,
+            0.0f,
+            0.0f,
+            0.0f,
+            -0.329264f,
+            -0.102785f,
+            -0.0444518f,
+            0.937574f,
+            0.0f,
+            0.0f,
+            0.0f,
+            0.0f,
+            0.0f,
+            0.0f,
+            -0.111598f,
+            0.0f,
+            0.0f,
+            0.993753f,
+            0.0f,
+            0.0f,
+            0.0f,
+            0.0f,
+            0.0f,
+            0.0f,
+            0.0f,
+            0.0f,
+            0.0f,
+            1.0f,
+            0.0f,
+            0.0f,
+            0.0f,
+            0.0f,
+            0.0f,
+            0.0f,
+            -0.166961f,
+            0.0f,
+            0.0f,
+            0.985964f,
+            0.0f,
+            0.0f,
+            0.0f,
+            0.0f,
+            0.0f,
+            0.0f,
+            -0.268386f,
+            -0.100976f,
+            -0.209457f,
+            0.934827f,
+            0.0f,
+            0.0f,
+            0.0f,
+            0.0f,
+            0.0f,
+            0.0f,
+            -0.317815f,
+            0.0f,
+            0.0f,
+            0.948153f,
+            0.0f,
+            0.0f,
+            0.0f,
+            0.0f,
+            0.0f,
+            0.0f,
+            -0.268386f,
+            -0.100976f,
+            -0.209457f,
+            0.934827f,
+            0.0f,
+            0.0f,
+            0.0f,
+            0.0f,
+            0.0f,
+            0.0f,
+            0.0f,
+            0.0f,
+            0.0f,
+            0.0f,
+            0.0f,
+            0.0f,
+            0.0f
+        }, {
+            // LOD-3
+            0.0f,
+            0.0f,
+            0.0f,
+            -0.26691f,
+            0.0f,
+            0.0f,
+            0.963722f,
+            0.0f,
+            0.0f,
+            0.0f,
+            0.0f,
+            0.0f,
+            0.0f,
+            -0.067495f,
+            0.0f,
+            0.0f,
+            0.99772f,
+            0.0f,
+            0.0f,
+            0.0f,
+            0.0f,
+            0.0f,
+            0.0f,
+            -0.320053f,
+            -0.2376f,
+            -0.102755f,
+            0.911347f,
+            0.0f,
+            0.0f,
+            0.0f,
+            0.0f,
+            0.0f,
+            0.0f,
+            -0.329264f,
+            -0.102785f,
+            -0.0444518f,
+            0.937574f,
+            0.0f,
+            0.0f,
+            0.0f,
+            0.0f,
+            0.0f,
+            0.0f,
+            -0.111598f,
+            0.0f,
+            0.0f,
+            0.993753f,
+            0.0f,
+            0.0f,
+            0.0f,
+            0.0f,
+            0.0f,
+            0.0f,
+            0.0f,
+            0.0f,
+            0.0f,
+            1.0f,
+            0.0f,
+            0.0f,
+            0.0f,
+            0.0f,
+            0.0f,
+            0.0f,
+            -0.166961f,
+            0.0f,
+            0.0f,
+            0.985964f,
+            0.0f,
+            0.0f,
+            0.0f,
+            0.0f,
+            0.0f,
+            0.0f,
+            -0.268386f,
+            -0.100976f,
+            -0.209457f,
+            0.934827f,
+            0.0f,
+            0.0f,
+            0.0f,
+            0.0f,
+            0.0f,
+            0.0f,
+            -0.317815f,
+            0.0f,
+            0.0f,
+            0.948153f,
+            0.0f,
+            0.0f,
+            0.0f,
+            0.0f,
+            0.0f,
+            0.0f,
+            -0.268386f,
+            -0.100976f,
+            -0.209457f,
+            0.934827f,
+            0.0f,
+            0.0f,
+            0.0f,
+            0.0f,
+            0.0f,
+            0.0f,
+            0.0f,
+            0.0f,
+            0.0f,
+            0.0f,
+            0.0f,
+            0.0f,
+            0.0f
+        }
+    },
+    {  // Euler-angle outputs
+        {
+            // LOD-0
+            0.0f,
+            0.0f,
+            0.0f,
+            -0.540371f,
+            0.0f,
+            0.0f,
+            0.0f,
+            0.0f,
+            0.0f,
+            0.0f,
+            0.0f,
+            0.0f,
+            -0.135093f,
+            0.0f,
+            0.0f,
+            0.0f,
+            0.0f,
+            0.0f,
+            0.0f,
+            0.0f,
+            0.0f,
+            -0.747349f,
+            -0.376101f,
+            -0.373509f,
+            0.0f,
+            0.0f,
+            0.0f,
+            0.0f,
+            0.0f,
+            0.0f,
+            -0.688136f,
+            -0.164202f,
+            -0.153704f,
+            0.0f,
+            0.0f,
+            0.0f,
+            0.0f,
+            0.0f,
+            0.0f,
+            -0.223662f,
+            0.0f,
+            0.0f,
+            0.0f,
+            0.0f,
+            0.0f,
+            0.0f,
+            0.0f,
+            0.0f,
+            0.0f,
+            0.0f,
+            0.0f,
+            0.0f,
+            0.0f,
+            0.0f,
+            0.0f,
+            0.0f,
+            0.0f,
+            -0.335493f,
+            0.0f,
+            0.0f,
+            0.0f,
+            0.0f,
+            0.0f,
+            0.0f,
+            0.0f,
+            0.0f,
+            -0.577202f,
+            -0.0764333f,
+            -0.46354f,
+            0.0f,
+            0.0f,
+            0.0f,
+            0.0f,
+            0.0f,
+            0.0f,
+            -0.646848f,
+            0.0f,
+            0.0f,
+            0.0f,
+            0.0f,
+            0.0f,
+            0.0f,
+            0.0f,
+            0.0f,
+            -0.577202f,
+            -0.0764333f,
+            -0.46354f,
+            0.0f,
+            0.0f,
+            0.0f,
+            0.0f,
+            0.0f,
+            0.0f,
+            0.0f,
+            0.0f,
+            0.0f,
+            0.0f,
+            0.0f,
+            0.0f
+        }, {
+            // LOD-1
+            0.0f,
+            0.0f,
+            0.0f,
+            -0.540371f,
+            0.0f,
+            0.0f,
+            0.0f,
+            0.0f,
+            0.0f,
+            0.0f,
+            0.0f,
+            0.0f,
+            -0.135093f,
+            0.0f,
+            0.0f,
+            0.0f,
+            0.0f,
+            0.0f,
+            0.0f,
+            0.0f,
+            0.0f,
+            -0.747349f,
+            -0.376101f,
+            -0.373509f,
+            0.0f,
+            0.0f,
+            0.0f,
+            0.0f,
+            0.0f,
+            0.0f,
+            -0.688136f,
+            -0.164202f,
+            -0.153704f,
+            0.0f,
+            0.0f,
+            0.0f,
+            0.0f,
+            0.0f,
+            0.0f,
+            -0.223662f,
+            0.0f,
+            0.0f,
+            0.0f,
+            0.0f,
+            0.0f,
+            0.0f,
+            0.0f,
+            0.0f,
+            0.0f,
+            0.0f,
+            0.0f,
+            0.0f,
+            0.0f,
+            0.0f,
+            0.0f,
+            0.0f,
+            0.0f,
+            -0.335493f,
+            0.0f,
+            0.0f,
+            0.0f,
+            0.0f,
+            0.0f,
+            0.0f,
+            0.0f,
+            0.0f,
+            -0.577202f,
+            -0.0764333f,
+            -0.46354f,
+            0.0f,
+            0.0f,
+            0.0f,
+            0.0f,
+            0.0f,
+            0.0f,
+            -0.646848f,
+            0.0f,
+            0.0f,
+            0.0f,
+            0.0f,
+            0.0f,
+            0.0f,
+            0.0f,
+            0.0f,
+            -0.577202f,
+            -0.0764333f,
+            -0.46354f,
+            0.0f,
+            0.0f,
+            0.0f,
+            0.0f,
+            0.0f,
+            0.0f,
+            0.0f,
+            0.0f,
+            0.0f,
+            0.0f,
+            0.0f,
+            0.0f
+        }, {
+            // LOD-2
+            0.0f,
+            0.0f,
+            0.0f,
+            -0.540371f,
+            0.0f,
+            0.0f,
+            0.0f,
+            0.0f,
+            0.0f,
+            0.0f,
+            0.0f,
+            0.0f,
+            -0.135093f,
+            0.0f,
+            0.0f,
+            0.0f,
+            0.0f,
+            0.0f,
+            0.0f,
+            0.0f,
+            0.0f,
+            -0.747349f,
+            -0.376101f,
+            -0.373509f,
+            0.0f,
+            0.0f,
+            0.0f,
+            0.0f,
+            0.0f,
+            0.0f,
+            -0.688136f,
+            -0.164202f,
+            -0.153704f,
+            0.0f,
+            0.0f,
+            0.0f,
+            0.0f,
+            0.0f,
+            0.0f,
+            -0.223662f,
+            0.0f,
+            0.0f,
+            0.0f,
+            0.0f,
+            0.0f,
+            0.0f,
+            0.0f,
+            0.0f,
+            0.0f,
+            0.0f,
+            0.0f,
+            0.0f,
+            0.0f,
+            0.0f,
+            0.0f,
+            0.0f,
+            0.0f,
+            -0.335493f,
+            0.0f,
+            0.0f,
+            0.0f,
+            0.0f,
+            0.0f,
+            0.0f,
+            0.0f,
+            0.0f,
+            -0.577202f,
+            -0.0764333f,
+            -0.46354f,
+            0.0f,
+            0.0f,
+            0.0f,
+            0.0f,
+            0.0f,
+            0.0f,
+            -0.646848f,
+            0.0f,
+            0.0f,
+            0.0f,
+            0.0f,
+            0.0f,
+            0.0f,
+            0.0f,
+            0.0f,
+            -0.577202f,
+            -0.0764333f,
+            -0.46354f,
+            0.0f,
+            0.0f,
+            0.0f,
+            0.0f,
+            0.0f,
+            0.0f,
+            0.0f,
+            0.0f,
+            0.0f,
+            0.0f,
+            0.0f,
+            0.0f
+        }, {
+            // LOD-3
+            0.0f,
+            0.0f,
+            0.0f,
+            -0.540371f,
+            0.0f,
+            0.0f,
+            0.0f,
+            0.0f,
+            0.0f,
+            0.0f,
+            0.0f,
+            0.0f,
+            -0.135093f,
+            0.0f,
+            0.0f,
+            0.0f,
+            0.0f,
+            0.0f,
+            0.0f,
+            0.0f,
+            0.0f,
+            -0.747349f,
+            -0.376101f,
+            -0.373509f,
+            0.0f,
+            0.0f,
+            0.0f,
+            0.0f,
+            0.0f,
+            0.0f,
+            -0.688136f,
+            -0.164202f,
+            -0.153704f,
+            0.0f,
+            0.0f,
+            0.0f,
+            0.0f,
+            0.0f,
+            0.0f,
+            -0.223662f,
+            0.0f,
+            0.0f,
+            0.0f,
+            0.0f,
+            0.0f,
+            0.0f,
+            0.0f,
+            0.0f,
+            0.0f,
+            0.0f,
+            0.0f,
+            0.0f,
+            0.0f,
+            0.0f,
+            0.0f,
+            0.0f,
+            0.0f,
+            -0.335493f,
+            0.0f,
+            0.0f,
+            0.0f,
+            0.0f,
+            0.0f,
+            0.0f,
+            0.0f,
+            0.0f,
+            -0.577202f,
+            -0.0764333f,
+            -0.46354f,
+            0.0f,
+            0.0f,
+            0.0f,
+            0.0f,
+            0.0f,
+            0.0f,
+            -0.646848f,
+            0.0f,
+            0.0f,
+            0.0f,
+            0.0f,
+            0.0f,
+            0.0f,
+            0.0f,
+            0.0f,
+            -0.577202f,
+            -0.0764333f,
+            -0.46354f,
+            0.0f,
+            0.0f,
+            0.0f,
+            0.0f,
+            0.0f,
+            0.0f,
+            0.0f,
+            0.0f,
+            0.0f,
+            0.0f,
+            0.0f,
+            0.0f
+        }
+    }
+};
+
+}  // namespace output
+
+#ifdef __clang__
+    #pragma clang diagnostic pop
+#endif
+
+TwistSwingReader::~TwistSwingReader() = default;
+
+}  // namespace twsw
+
+}  // namespace rltests
diff --git a/Engine/Plugins/Animation/RigLogic/Source/RigLogicLibTest/Private/rltests/ml/cpu/FixturesBlock4.cpp b/Engine/Plugins/Animation/RigLogic/Source/RigLogicLibTest/Private/rltests/ml/cpu/FixturesBlock4.cpp
--- a/Engine/Plugins/Animation/RigLogic/Source/RigLogicLibTest/Private/rltests/ml/cpu/FixturesBlock4.cpp
+++ b/Engine/Plugins/Animation/RigLogic/Source/RigLogicLibTest/Private/rltests/ml/cpu/FixturesBlock4.cpp
@@ -1,13 +1,10 @@
 // Copyright Epic Games, Inc. All Rights Reserved.
 
-// *INDENT-OFF*
-#ifdef RL_BUILD_WITH_ML_EVALUATOR
-
 #include "rltests/ml/cpu/FixturesBlock4.h"
 
 #include "riglogic/TypeDefs.h"
-#include "riglogic/joints/bpcm/BPCMOutputInstance.h"
-#include "riglogic/joints/bpcm/Evaluator.h"
+#include "riglogic/joints/cpu/CPUJointsOutputInstance.h"
+#include "riglogic/joints/cpu/bpcm/BPCMJointsEvaluator.h"
 #include "riglogic/types/Extent.h"
 
 namespace rltests {
@@ -978,7 +975,7 @@ const Vector<AlignedMatrix<std::uint16_t> > mlbNetBiasesHalfFloat = {
     }
 };
 
-const LODSpec lods = {
+const LODSpec<std::uint32_t> lods = {
     {  // NetIndices
         {  // LOD-0
             0u, 1u, 2u, 3u
@@ -1103,6 +1100,3 @@ CanonicalReader::~CanonicalReader() = default;
 }  // namespace ml
 
 }  // namespace rltests
-
-#endif  // RL_BUILD_WITH_ML_EVALUATOR
-// *INDENT-ON*
diff --git a/Engine/Plugins/Animation/RigLogic/Source/RigLogicLibTest/Private/rltests/ml/cpu/FixturesBlock4.h b/Engine/Plugins/Animation/RigLogic/Source/RigLogicLibTest/Private/rltests/ml/cpu/FixturesBlock4.h
--- a/Engine/Plugins/Animation/RigLogic/Source/RigLogicLibTest/Private/rltests/ml/cpu/FixturesBlock4.h
+++ b/Engine/Plugins/Animation/RigLogic/Source/RigLogicLibTest/Private/rltests/ml/cpu/FixturesBlock4.h
@@ -2,15 +2,12 @@
 
 #pragma once
 
-// *INDENT-OFF*
-#ifdef RL_BUILD_WITH_ML_EVALUATOR
-
 #include "rltests/Defs.h"
 #include "rltests/StorageValueType.h"
 #include "rltests/dna/FakeReader.h"
 
 #include "riglogic/TypeDefs.h"
-#include "riglogic/ml/LODSpec.h"
+#include "riglogic/types/LODSpec.h"
 
 #ifdef _MSC_VER
     #pragma warning(push)
@@ -60,7 +57,7 @@ extern const Vector<AlignedMatrix<float> > mlbNetWeightsFloat;
 extern const Vector<AlignedMatrix<std::uint16_t> > mlbNetWeightsHalfFloat;
 extern const Vector<AlignedMatrix<float> > mlbNetBiasesFloat;
 extern const Vector<AlignedMatrix<std::uint16_t> > mlbNetBiasesHalfFloat;
-extern const LODSpec lods;
+extern const LODSpec<std::uint32_t> lods;
 
 template<typename TValue>
 struct Values {
@@ -138,6 +135,3 @@ class CanonicalReader : public dna::FakeReader {
 }  // namespace ml
 
 }  // namespace rltests
-
-#endif  // RL_BUILD_WITH_ML_EVALUATOR
-// *INDENT-ON*
diff --git a/Engine/Plugins/Animation/RigLogic/Source/RigLogicLibTest/Private/rltests/ml/cpu/TestInference.cpp b/Engine/Plugins/Animation/RigLogic/Source/RigLogicLibTest/Private/rltests/ml/cpu/TestInference.cpp
--- a/Engine/Plugins/Animation/RigLogic/Source/RigLogicLibTest/Private/rltests/ml/cpu/TestInference.cpp
+++ b/Engine/Plugins/Animation/RigLogic/Source/RigLogicLibTest/Private/rltests/ml/cpu/TestInference.cpp
@@ -1,15 +1,12 @@
 // Copyright Epic Games, Inc. All Rights Reserved.
 
-// *INDENT-OFF*
-#ifdef RL_BUILD_WITH_ML_EVALUATOR
-
 #include "rltests/Defs.h"
 #include "rltests/StorageValueType.h"
 #include "rltests/controls/ControlFixtures.h"
 #include "rltests/ml/cpu/FixturesBlock4.h"
 
 #include "riglogic/TypeDefs.h"
-#include "riglogic/ml/cpu/Factory.h"
+#include "riglogic/ml/cpu/CPUMachineLearnedBehaviorFactory.h"
 #include "riglogic/system/simd/Detect.h"
 #include "riglogic/system/simd/SIMD.h"
 
@@ -97,7 +94,8 @@ TYPED_TEST(MLBInferenceTest, InferencePerLOD) {
     auto inputInstanceFactory = ControlsFactory::getInstanceFactory(0,
                                                                     rltests::ml::block4::unoptimized::rawControlCount,
                                                                     0,
-                                                                    rltests::ml::block4::unoptimized::mlControlCount);
+                                                                    rltests::ml::block4::unoptimized::mlControlCount,
+                                                                    0);
     auto inputInstance = inputInstanceFactory(&this->memRes);
     auto inputBuffer = inputInstance->getInputBuffer();
     auto outputBuffer = inputBuffer.subview(rltests::ml::block4::unoptimized::rawControlCount,
@@ -122,6 +120,3 @@ TYPED_TEST(MLBInferenceTest, InferencePerLOD) {
 #ifdef _MSC_VER
     #pragma warning(pop)
 #endif
-
-#endif  // RL_BUILD_WITH_ML_EVALUATOR
-// *INDENT-ON*
diff --git a/Engine/Plugins/Animation/RigLogic/Source/RigLogicLibTest/Private/rltests/ml/cpu/TestStorageBuilder.cpp b/Engine/Plugins/Animation/RigLogic/Source/RigLogicLibTest/Private/rltests/ml/cpu/TestStorageBuilder.cpp
--- a/Engine/Plugins/Animation/RigLogic/Source/RigLogicLibTest/Private/rltests/ml/cpu/TestStorageBuilder.cpp
+++ b/Engine/Plugins/Animation/RigLogic/Source/RigLogicLibTest/Private/rltests/ml/cpu/TestStorageBuilder.cpp
@@ -1,14 +1,11 @@
 // Copyright Epic Games, Inc. All Rights Reserved.
 
-// *INDENT-OFF*
-#ifdef RL_BUILD_WITH_ML_EVALUATOR
-
 #include "rltests/Defs.h"
 #include "rltests/StorageValueType.h"
 #include "rltests/ml/cpu/FixturesBlock4.h"
 
 #include "riglogic/TypeDefs.h"
-#include "riglogic/ml/cpu/Factory.h"
+#include "riglogic/ml/cpu/CPUMachineLearnedBehaviorFactory.h"
 #include "riglogic/system/simd/Detect.h"
 
 #include <tuple>
@@ -28,9 +25,9 @@ template<typename T, typename TF256, typename TF128>
 struct Evaluator<T, TF256, TF128>::Accessor {
 
     static void assertRawDataEqual(const Evaluator<T, TF256, TF128>& result) {
-        ASSERT_EQ(result.lods.netIndicesPerLOD, rltests::ml::block4::optimized::lods.netIndicesPerLOD);
-        ASSERT_EQ(result.lods.netCount, rltests::ml::block4::optimized::lods.netCount);
-        ASSERT_EQ(result.neuralNets.size(), rltests::ml::block4::optimized::lods.netCount);
+        ASSERT_EQ(result.lods.indicesPerLOD, rltests::ml::block4::optimized::lods.indicesPerLOD);
+        ASSERT_EQ(result.lods.count, rltests::ml::block4::optimized::lods.count);
+        ASSERT_EQ(result.neuralNets.size(), rltests::ml::block4::optimized::lods.count);
         for (std::uint16_t netIndex = {}; netIndex < rltests::ml::block4::unoptimized::neuralNetworkCount; ++netIndex) {
             const auto& inference = result.neuralNets[netIndex];
             ASSERT_EQ(inference.layerEvaluators.size(),
@@ -136,6 +133,3 @@ TYPED_TEST(MLBSStorageBuilderTest, LayoutOptimization) {
 #ifdef _MSC_VER
     #pragma warning(pop)
 #endif
-
-#endif  // RL_BUILD_WITH_ML_EVALUATOR
-// *INDENT-ON*
diff --git a/Engine/Plugins/Animation/RigLogic/Source/RigLogicLibTest/Private/rltests/rbf/cpu/RBFFixtures.cpp b/Engine/Plugins/Animation/RigLogic/Source/RigLogicLibTest/Private/rltests/rbf/cpu/RBFFixtures.cpp
new file mode 100644
--- /dev/null
+++ b/Engine/Plugins/Animation/RigLogic/Source/RigLogicLibTest/Private/rltests/rbf/cpu/RBFFixtures.cpp
@@ -0,0 +1,288 @@
+// Copyright Epic Games, Inc. All Rights Reserved.
+
+#include "rltests/rbf/cpu/RBFFixtures.h"
+
+namespace rltests {
+
+namespace rbf {
+
+#ifdef __clang__
+    #pragma clang diagnostic push
+    #pragma clang diagnostic ignored "-Wglobal-constructors"
+    #pragma clang diagnostic ignored "-Wexit-time-destructors"
+#endif
+
+namespace unoptimized {
+
+const std::uint16_t lodCount = 3u;
+
+const std::uint16_t rawControlCount = 8u;
+const std::uint16_t rbfControlCount = 15u;
+
+const pma::Matrix<std::uint16_t> solverIndicesPerLOD = {
+    {0u, 1u},
+    {0u},
+    {}
+};
+
+const pma::Vector<float> poseScales = {
+    1.0f, 1.0f, 1.0f, 1.0f, 1.0f, 1.0f, 1.0f, 1.0f, 1.0f, 1.0f, 1.0f, 1.0f, 1.0f, 1.0f, 1.0f
+};
+
+const pma::Vector<dna::RBFSolverType> solverTypes = {
+    dna::RBFSolverType::Interpolative,
+    dna::RBFSolverType::Interpolative
+};
+
+const pma::Vector<dna::RBFDistanceMethod> solverDistanceMethods = {
+    dna::RBFDistanceMethod::SwingAngle,
+    dna::RBFDistanceMethod::SwingAngle
+};
+
+const pma::Vector<dna::RBFFunctionType> solverFunctionType = {
+    dna::RBFFunctionType::Gaussian,
+    dna::RBFFunctionType::Gaussian
+};
+
+const pma::Vector<dna::RBFNormalizeMethod> solverNormalizeMethods = {
+    dna::RBFNormalizeMethod::AlwaysNormalize,
+    dna::RBFNormalizeMethod::AlwaysNormalize
+};
+
+const pma::Vector<dna::TwistAxis> solverTwistAxis = {
+    dna::TwistAxis::X,
+    dna::TwistAxis::X
+};
+
+const pma::Vector<dna::AutomaticRadius> solverAutomaticRadius = {
+    dna::AutomaticRadius::On,
+    dna::AutomaticRadius::On
+};
+
+const pma::Vector<float> solverRadius = {
+    0.0f, 0.0f
+};
+
+const pma::Vector<float> solverWeightThreshold = {
+    0.001f, 0.001f, 0.001f
+};
+
+const pma::Matrix<std::uint16_t> solverPoseIndices = {
+    {0u, 1u, 2u, 3u, 4u, 5u, 6u, 7u, 8u, 9u, 10u, 11u},
+    {0u, 12u, 13u, 14u}
+};
+
+const pma::Matrix<std::uint16_t> solverRawControlIndices = {
+    {0u, 1u, 2u, 3u},
+    {4u, 5u, 6u, 7u},
+};
+
+const pma::Matrix<float> solverRawControlValues = {
+    {
+        0.000000000000000f, 0.000000000000000f, 0.000000000000000f, 1.000000000000000f,
+        -0.003081271657720f, -0.118239738047123f, -0.009329595603049f, 0.992936491966248f,
+        -0.008705757558346f, 0.009779179468751f, -0.141530960798264f, 0.989847242832184f,
+        0.026532903313637f, -0.811531901359558f, -0.024197027087212f, 0.583203732967377f,
+        -0.000952127971686f, 0.013058164156973f, 0.076173260807991f, 0.997008621692657f,
+        -0.044993601739407f, -0.664866507053375f, 0.044108338654041f, 0.744300007820129f,
+        -0.005394733510911f, 0.099454566836357f, -0.012115634977818f, 0.994953811168671f,
+        0.009781738743186f, 0.008702844381332f, 0.372627735137939f, 0.927888572216034f,
+        -0.009282855316997f, 0.312406390905380f, -0.014897738583386f, 0.949786365032196f,
+        -0.003883346682414f, -0.450696706771851f, -0.001544478582218f, 0.892667353153229f,
+        -0.005706345662475f, -0.011783968657255f, -0.714682221412659f, 0.699326753616333f,
+        0.000949318520725f, -0.013058470562100f, -0.825225293636322f, 0.564651966094971f
+    },
+    {
+        0.0f, 0.0f, 0.0f, 1.0f,
+        -0.0936718f, -0.12003f, 0.663135f, 0.732851f,
+        0.123443f, 0.0891258f, -0.199695f, 0.967957f,
+        -0.12003f, 0.0936719f, -0.732851f, 0.663135f
+    }
+};
+
+}  // namespace unoptimized
+
+namespace optimized {
+
+const LODSpec<std::uint16_t> lods = {
+    {{0, 1}, {0}, {}},
+    2u  // Solver count
+};
+
+const Matrix<std::uint16_t> solverRawControlInputIndices = {{0u, 1u, 2u, 3u}, {4u, 5u, 6u, 7u}};
+
+const Matrix<std::uint16_t> solverRawControlOutputIndices =
+{{8u, 9u, 10u, 11u, 12u, 13u, 14u, 15u, 16u, 17u, 18u, 19u}, {8u, 20u, 21u, 22u}};
+
+const std::uint16_t maximumInputCount = 4u;
+
+const std::uint16_t maxTargetCount = 12u;
+
+const Vector<Matrix<float> > targetValues = {
+    {
+        {0.0f, -0.0f, -0.0f, -1.0f},
+        {0.0f, 0.118210219f, 0.00969646964f, -0.99294126f},
+        {0.0f, -0.0110235251f, 0.141439483f, -0.989885509f},
+        {0.0f, 0.811793089f, -0.0127105378f, -0.583806992f},
+        {0.0f, -0.0129854148f, -0.0761857033f, -0.997009039f},
+        {0.0f, 0.666316569f, -0.00390942581f, -0.745658696f},
+        {0.0f, -0.0995188057f, 0.0115762129f, -0.994968414f},
+        {0.0f, -0.0126303593f, -0.372515291f, -0.92794013f},
+        {0.0f, -0.312537074f, 0.0118438303f, -0.949831724f},
+        {0.0f, 0.45068571f, 0.0035050991f, -0.892675817f},
+        {0.0f, 0.00595212681f, 0.714754522f, -0.699350059f},
+        {0.0f, 0.0144458571f, 0.825202227f, -0.564652741f}
+    },
+    {
+        {0.0f, -0.0f, -0.0f, -1.0f},
+        {0.0f, 0.203138143f, -0.64256525f, -0.738813281f},
+        {0.0f, -0.063147366f, 0.209365502f, -0.97579658f},
+        {0.0f, -0.222702071f, 0.704449356f, -0.673910379f}
+    }
+};
+
+const Vector<Matrix<float> > coefficients = {
+    {
+        {10.4131117f, -2.50109625f, -1.68585944f, 0.00226313062f, -4.21029043f, 0.00851259567f, -2.70837259f, 0.19591637f,
+         0.28792578f, 0.11638914f, 0.0584293604f, 0.00188750029f
+        },
+        {
+            -2.50109625f, 5.29161167f, -1.25280869f, 0.0132418955f, -1.07533145f, 0.0170967579f, 0.630707622f, -0.224908382f,
+            0.209685102f, -1.06077302f, -0.053258799f, 0.0111311972f
+        },
+        {-1.68586028f, -1.25280893f, 4.58401918f, 0.00919136312f, 0.76852268f, -0.00866004918f, -1.57939589f, 0.237331465f,
+         -0.384880483f, -0.169630617f, -0.486707956f, 0.0163090285f
+        },
+        {
+            0.00226316857f, 0.0132418117f, 0.00919136778f, 2.45947075f, 0.00964977313f, -1.88095915f, 0.00817291997f,
+            -0.0562427454f, 0.0234976001f, 0.0126816928f, -0.0217441544f, -0.0700526312f
+        },
+        {
+            -4.210289f, -1.07533169f, 0.768522859f, 0.00964973494f, 7.6117754f, 0.00983878039f, -1.76371253f, -1.1520133f,
+            -0.220182136f, -0.0616565943f, 0.0662193298f, 0.0089058429f
+        },
+        {
+            0.00851258356f, 0.0170967635f, -0.00866001379f, -1.88095927f, 0.00983874034f, 3.59918761f, 0.013124981f,
+            -0.0629767478f,
+            0.022830274f, -1.55218458f, -0.0439120345f, -0.0208300687f
+        },
+        {
+            -2.70837283f, 0.630707622f, -1.57939577f, 0.00817293487f, -1.76371217f, 0.0131249651f, 6.99775982f, -0.0686918572f,
+            -1.6557734f, 0.102362826f, 0.00317596481f, 0.00825250242f
+        },
+        {
+            0.195916459f, -0.224908367f, 0.237331495f, -0.0562427156f, -1.15201342f, -0.0629767776f, -0.0686918348f, 2.05959892f,
+            -0.373588949f, -0.214859918f, 0.0606168285f, 0.0199813042f
+        },
+        {
+            0.287925869f, 0.209685087f, -0.384880453f, 0.0234975964f, -0.220182076f, 0.0228302795f, -1.65577352f, -0.373588949f,
+            2.60091233f, 0.0772639886f, -0.170895383f, -0.0364022851f
+        },
+        {
+            0.116389424f, -1.06077302f, -0.169630632f, 0.0126816807f, -0.0616565533f, -1.55218446f, 0.102362826f, -0.214859903f,
+            0.0772639811f, 2.94548035f, -0.107176304f, -0.00518302247f
+        },
+        {
+            0.0584293269f, -0.0532588698f, -0.486707866f, -0.0217442084f, 0.0662193596f, -0.0439120308f, 0.00317585957f,
+            0.0606168173f, -0.170895383f, -0.107176282f, 3.21423078f, -2.36085558f
+        },
+        {
+            0.00188759307f, 0.0111312028f, 0.0163090061f, -0.0700525865f, 0.00890581496f, -0.0208300594f, 0.00825257134f,
+            0.0199813209f, -0.0364022776f, -0.00518305739f, -2.36085558f, 2.9289515f
+        }
+    },
+    {
+        {4.09646654f, -0.869968235f, -3.05232024f, -6.16908073e-05f},
+        {-0.869968295f, 1.50343823f, -3.51387607e-05f, 8.92129538e-06f},
+        {-3.05232024f, -3.51387644e-05f, 4.27379131f, -1.06966126f},
+        {-6.17128608e-05f, 8.92129538e-06f, -1.06966126f, 1.68079543f}
+    }
+};
+
+const pma::Vector<float> solverRadius = {
+    1.29129f, 1.651938f
+};
+const pma::Matrix<float> solverPoseScales = {
+    {1.0f, 1.0f, 1.0f, 1.0f, 1.0f, 1.0f, 1.0f, 1.0f, 1.0f, 1.0f, 1.0f, 1.0f},
+    {1.0f, 1.0f, 1.0f, 1.0f}
+};
+
+}  // namespace optimized
+
+namespace input {
+
+// Calculation input values
+const rl4::Vector<float> values = {
+    0.620591878890991f, 0.382426619529724f, -0.683809995651245f, 0.031930625438690f,
+    0.0f, 0.0f, -0.47362f, 0.880729f
+};
+
+}  // namespace input
+
+namespace output {
+
+const Matrix<float> valuesPerLOD = {
+    {
+        0.0657254f,
+        0.0f,
+        0.0f,
+        0.212102f,
+        0.0116427f,
+        0.0f,
+        0.0f,
+        0.279036f,
+        0.473921f,
+        0.0f,
+        0.0033996f,
+        0.0198993f,
+        0.0f,
+        0.453696f,
+        0.480578f,
+    },
+    {
+        0.0f,
+        0.0f,
+        0.0f,
+        0.212102f,
+        0.0116427f,
+        0.0f,
+        0.0f,
+        0.279036f,
+        0.473921f,
+        0.0f,
+        0.0033996f,
+        0.0198993f,
+        0.0f,
+        0.0f,
+        0.0f
+    },
+    {
+        0.0f,
+        0.0f,
+        0.0f,
+        0.0f,
+        0.0f,
+        0.0f,
+        0.0f,
+        0.0f,
+        0.0f,
+        0.0f,
+        0.0f,
+        0.0f,
+        0.0f,
+        0.0f,
+        0.0f,
+    }
+};
+
+}  // namespace output
+
+#ifdef __clang__
+    #pragma clang diagnostic pop
+#endif
+
+RBFReader::~RBFReader() = default;
+
+}  // namespace rbf
+
+}  // namespace rltests
diff --git a/Engine/Plugins/Animation/RigLogic/Source/RigLogicLibTest/Private/rltests/rbf/cpu/TestRBFSolver.cpp b/Engine/Plugins/Animation/RigLogic/Source/RigLogicLibTest/Private/rltests/rbf/cpu/TestRBFSolver.cpp
new file mode 100644
--- /dev/null
+++ b/Engine/Plugins/Animation/RigLogic/Source/RigLogicLibTest/Private/rltests/rbf/cpu/TestRBFSolver.cpp
@@ -0,0 +1,149 @@
+// Copyright Epic Games, Inc. All Rights Reserved.
+
+#include "rltests/Defs.h"
+
+#include "riglogic/rbf/cpu/RBFSolver.h"
+
+#ifdef _MSC_VER
+    #pragma warning(push)
+    #pragma warning(disable : 4324)
+#endif
+
+namespace {
+
+class RBFSolverTest : public ::testing::Test {
+    protected:
+        void SetUp() override {
+            recipe.normalizeMethod = rl4::RBFNormalizeMethod::AlwaysNormalize;
+            recipe.isAutomaticRadius = true;
+            recipe.radius = 0.0f;
+            recipe.twistAxis = dna::TwistAxis::X;
+            recipe.weightThreshold = 0.001f;
+        }
+
+    protected:
+        pma::AlignedMemoryResource memRes;
+        rl4::RBFSolverRecipe recipe;
+};
+
+}  // namespace
+
+TEST_F(RBFSolverTest, InterpolativeGaussianSwingAngle0) {
+    recipe.solverType = rl4::RBFSolverType::Interpolative;
+    recipe.distanceMethod = rl4::RBFDistanceMethod::SwingAngle;
+    recipe.weightFunction = rl4::RBFFunctionType::Gaussian;
+    std::size_t targetCount = 4;
+    pma::Vector<float> scaleFactors{targetCount, 1.0f, &memRes};
+    recipe.targetScales = scaleFactors;
+    recipe.rawControlCount = 4u;
+    pma::Vector<float> targetValues{
+        0.0f, 0.0f, 0.0f, 1.0f,
+        -0.0936718f, -0.12003f, 0.663135f, 0.732851f,
+        0.123443f, 0.0891258f, -0.199695f, 0.967957f,
+        -0.12003f, 0.0936719f, -0.732851f, 0.663135f
+    };
+    recipe.targetValues = targetValues;
+    auto solver = rl4::RBFSolver::create(recipe, &memRes);
+    pma::Vector<float> input {-0.0936718f, -0.12003f, 0.663135f, 0.732851f};
+    pma::Vector<float> buffer {targetCount, 0.0f, &memRes};
+    pma::Vector<float> result {targetCount, 0.0f, &memRes};
+
+    solver->solve(input, buffer, result);
+    pma::Vector<float> expected{0.0f, 1.0f, 0.0f, 0.0f};
+    EXPECT_ELEMENTS_NEAR(result, expected, targetCount, 0.0001f);
+    input = {0.0f, 0.0f, -0.47362f, 0.880729f};
+    solver->solve(input, buffer, result);
+    expected = {0.0657254f, 0.0f, 0.453696f, 0.480578f};
+    EXPECT_ELEMENTS_NEAR(result, expected, targetCount, 0.0001f);
+}
+
+TEST_F(RBFSolverTest, InterpolativeGaussianSwingAngle1) {
+    recipe.solverType = rl4::RBFSolverType::Interpolative;
+    recipe.distanceMethod = rl4::RBFDistanceMethod::SwingAngle;
+    recipe.weightFunction = rl4::RBFFunctionType::Gaussian;
+    std::size_t targetCount = 12u;
+    pma::Vector<float> scaleFactors{targetCount, 1.0f, &memRes};
+    recipe.targetScales = scaleFactors;
+    recipe.rawControlCount = 4u;
+    pma::Vector<float> targetValues{
+        0.000000000000000f, 0.000000000000000f, 0.000000000000000f, 1.000000000000000f,
+        -0.003081271657720f, -0.118239738047123f, -0.009329595603049f, 0.992936491966248f,
+        -0.008705757558346f, 0.009779179468751f, -0.141530960798264f, 0.989847242832184f,
+        0.026532903313637f, -0.811531901359558f, -0.024197027087212f, 0.583203732967377f,
+        -0.000952127971686f, 0.013058164156973f, 0.076173260807991f, 0.997008621692657f,
+        -0.044993601739407f, -0.664866507053375f, 0.044108338654041f, 0.744300007820129f,
+        -0.005394733510911f, 0.099454566836357f, -0.012115634977818f, 0.994953811168671f,
+        0.009781738743186f, 0.008702844381332f, 0.372627735137939f, 0.927888572216034f,
+        -0.009282855316997f, 0.312406390905380f, -0.014897738583386f, 0.949786365032196f,
+        -0.003883346682414f, -0.450696706771851f, -0.001544478582218f, 0.892667353153229f,
+        -0.005706345662475f, -0.011783968657255f, -0.714682221412659f, 0.699326753616333f,
+        0.000949318520725f, -0.013058470562100f, -0.825225293636322f, 0.564651966094971f
+    };
+
+    recipe.targetValues = targetValues;
+    auto solver = rl4::RBFSolver::create(recipe, &memRes);
+    pma::Vector<float> input {-0.005706345662475f, -0.011783968657255f, -0.714682221412659f, 0.699326753616333f};
+    pma::Vector<float> buffer {targetCount, 0.0f, &memRes};
+    pma::Vector<float> result {targetCount, 0.0f, &memRes};
+
+    solver->solve(input, buffer, result);
+    pma::Vector<float> expected {0.0f,
+                                 0.0f,
+                                 0.0f,
+                                 0.0f,
+                                 0.0f,
+                                 0.0f,
+                                 0.0f,
+                                 0.0f,
+                                 0.0f,
+                                 0.0f,
+                                 1.0f,
+                                 0.0f};
+    EXPECT_ELEMENTS_NEAR(result, expected, targetCount, 0.0001f);
+
+    input = {0.620591878890991f, 0.382426619529724f, -0.683809995651245f, 0.031930625438690f};
+    solver->solve(input, buffer, result);
+    expected = {0.0f,
+                0.0f,
+                0.0f,
+                0.212102f,
+                0.0116427f,
+                0.0f,
+                0.0f,
+                0.279036f,
+                0.473921f,
+                0.0f,
+                0.0033996f,
+                0.0198993f};
+
+    EXPECT_ELEMENTS_NEAR(result, expected, targetCount, 0.0001f);
+}
+
+TEST_F(RBFSolverTest, AdditiveGaussianSwingAngle0) {
+    recipe.solverType = rl4::RBFSolverType::Additive;
+    recipe.distanceMethod = rl4::RBFDistanceMethod::SwingAngle;
+    recipe.weightFunction = rl4::RBFFunctionType::Gaussian;
+    std::size_t targetCount = 4;
+    pma::Vector<float> scaleFactors{targetCount, 1.0f, &memRes};
+    recipe.targetScales = scaleFactors;
+    recipe.rawControlCount = 4u;
+    pma::Vector<float> targetValues{
+        0.0f, 0.0f, 0.0f, 1.0f,
+        -0.0936718f, -0.12003f, 0.663135f, 0.732851f,
+        0.123443f, 0.0891258f, -0.199695f, 0.967957f,
+        -0.12003f, 0.0936719f, -0.732851f, 0.663135f
+    };
+    recipe.targetValues = targetValues;
+    auto solver = rl4::RBFSolver::create(recipe, &memRes);
+    pma::Vector<float> input {-0.0936718f, -0.12003f, 0.663135f, 0.732851f};
+    pma::Vector<float> buffer {targetCount, 0.0f, &memRes};
+    pma::Vector<float> result {targetCount, 0.0f, &memRes};
+
+    solver->solve(input, buffer, result);
+    pma::Vector<float> expected{0.242209f,
+                                0.41645f,
+                                0.20938f,
+                                0.131961f};
+
+    EXPECT_ELEMENTS_NEAR(result, expected, targetCount, 0.0001f);
+}
diff --git a/Engine/Plugins/Animation/RigLogic/Source/RigLogicLibTest/Private/rltests/riglogic/TestRigLogic.cpp b/Engine/Plugins/Animation/RigLogic/Source/RigLogicLibTest/Private/rltests/riglogic/TestRigLogic.cpp
--- a/Engine/Plugins/Animation/RigLogic/Source/RigLogicLibTest/Private/rltests/riglogic/TestRigLogic.cpp
+++ b/Engine/Plugins/Animation/RigLogic/Source/RigLogicLibTest/Private/rltests/riglogic/TestRigLogic.cpp
@@ -44,15 +44,22 @@ TEST_F(RigLogicTest, EvaluateRigInstance) {
 
     rigLogic->calculate(rigInstance.get());
 
-    ASSERT_EQ(rigInstance->getRawJointOutputs().size(), reader->getJointRowCount());
+    ASSERT_EQ(rigInstance->getJointOutputs().size(), reader->getJointRowCount());
     ASSERT_EQ(rigInstance->getBlendShapeOutputs().size(), reader->getBlendShapeChannelCount());
     ASSERT_EQ(rigInstance->getAnimatedMapOutputs().size(), reader->getAnimatedMapCount());
 }
 
 TEST_F(RigLogicTest, AccessJointVariableAttributeIndices) {
     for (std::uint16_t lod = 0u; lod < rigLogic->getLODCount(); ++lod) {
-        ASSERT_EQ(rigLogic->getJointVariableAttributeIndices(lod),
-                  (rl4::ConstArrayView<std::uint16_t>{rltests::decoded::jointVariableIndices[0ul][lod]}));
+        auto actual = rigLogic->getJointVariableAttributeIndices(lod);
+        auto expected = rl4::ConstArrayView<std::uint16_t>{rltests::decoded::jointVariableIndices[0ul][lod]};
+        ASSERT_EQ(actual.size(), expected.size());
+        // Since implementation relies on std::set which has different implementations across compilers we cannot guarantee order
+        // of elements
+        for (const auto attrIndex : expected) {
+            ASSERT_NE(std::find(actual.begin(), actual.end(), attrIndex),
+                      actual.end());
+        }
     }
 }
 
@@ -76,11 +83,11 @@ TEST_F(RigLogicTest, DumpStateThenRestore) {
         cloneRigLogic->mapGUIToRawControls(cloneRigInstance);
         cloneRigLogic->calculate(cloneRigInstance);
 
-        auto origJointOutputs = rigInstance->getRawJointOutputs();
+        auto origJointOutputs = rigInstance->getJointOutputs();
         auto origBlendShapeOutputs = rigInstance->getBlendShapeOutputs();
         auto origAnimatedMapOutputs = rigInstance->getAnimatedMapOutputs();
 
-        auto cloneJointOutputs = cloneRigInstance->getRawJointOutputs();
+        auto cloneJointOutputs = cloneRigInstance->getJointOutputs();
         auto cloneBlendShapeOutputs = cloneRigInstance->getBlendShapeOutputs();
         auto cloneAnimatedMapOutputs = cloneRigInstance->getAnimatedMapOutputs();
 
diff --git a/Engine/Plugins/Animation/RigLogic/Source/RigLogicLibTest/Private/tdmtests/Helpers.h b/Engine/Plugins/Animation/RigLogic/Source/RigLogicLibTest/Private/tdmtests/Helpers.h
new file mode 100644
--- /dev/null
+++ b/Engine/Plugins/Animation/RigLogic/Source/RigLogicLibTest/Private/tdmtests/Helpers.h
@@ -0,0 +1,21 @@
+// Copyright Epic Games, Inc. All Rights Reserved.
+
+#pragma once
+
+#define ASSERT_VEC_NEAR(v1, v2, threshold)                                     \
+    ASSERT_EQ(v1.dimensions(), v2.dimensions());                               \
+    for (tdm::dim_t ri = 0u; ri < v1.dimensions(); ++ri)                       \
+    {                                                                          \
+        ASSERT_NEAR(v1[ri], v2[ri], threshold);                                \
+    }
+
+#define ASSERT_MAT_NEAR(m1, m2, threshold)                                     \
+    ASSERT_EQ(m1.rows(), m2.rows());                                           \
+    ASSERT_EQ(m1.columns(), m2.columns());                                     \
+    for (tdm::dim_t ri = 0u; ri < m1.rows(); ++ri)                             \
+    {                                                                          \
+        for (tdm::dim_t ci = 0u; ci < m1.columns(); ++ci)                      \
+        {                                                                      \
+            ASSERT_NEAR(m1[ri][ci], m2[ri][ci], threshold);                    \
+        }                                                                      \
+    }
