Update to new RigLogic implementing body rig support

diff --git a/Engine/Plugins/Animation/RigLogic/Source/RigLogicLibTest/Private/tdmtests/TestAxisAngle.cpp b/Engine/Plugins/Animation/RigLogic/Source/RigLogicLibTest/Private/tdmtests/TestAxisAngle.cpp
new file mode 100644
--- /dev/null
+++ b/Engine/Plugins/Animation/RigLogic/Source/RigLogicLibTest/Private/tdmtests/TestAxisAngle.cpp
@@ -0,0 +1,30 @@
+// Copyright Epic Games, Inc. All Rights Reserved.
+
+#include "tdmtests/Defs.h"
+#include "tdmtests/Helpers.h"
+
+#include "tdm/TDM.h"
+
+TEST(TestAxisAngle, ConvertXYZ2AxisAngleAndRotate) {
+    using namespace tdm::ang_literals;
+    auto m = tdm::mat4<float>::identity();
+    tdm::axis_angle<float> axisAngle{tdm::frad{45.0_fdeg}, tdm::frad{0.0_fdeg}, tdm::frad{30.0_fdeg}};
+    m = tdm::rotate<float>(m, axisAngle.axis, axisAngle.angle);
+    tdm::mat4<float> expected{0.866025f, 0.5f, 0.0f, 0.0f,
+                              -0.353553f, 0.612372f, 0.707107f, 0.0f,
+                              0.353553f, -0.612372f, 0.707107f, 0.0f,
+                              0.0f, 0.0f, 0.0f, 1.0f};
+    ASSERT_MAT_NEAR(m, expected, 0.0001f);
+}
+
+TEST(TestAxisAngle, ConvertXYZVector2AxisAngleAndRotate) {
+    using namespace tdm::ang_literals;
+    auto m = tdm::mat4<float>::identity();
+    tdm::axis_angle<float> axisAngle{tdm::frad{45.0_fdeg}, tdm::frad{0.0_fdeg}, tdm::frad{30.0_fdeg}};
+    m = tdm::rotate<float>(m, axisAngle.axis, axisAngle.angle);
+    tdm::mat4<float> expected{0.866025f, 0.5f, 0.0f, 0.0f,
+                              -0.353553f, 0.612372f, 0.707107f, 0.0f,
+                              0.353553f, -0.612372f, 0.707107f, 0.0f,
+                              0.0f, 0.0f, 0.0f, 1.0f};
+    ASSERT_MAT_NEAR(m, expected, 0.0001f);
+}
diff --git a/Engine/Plugins/Animation/RigLogic/Source/RigLogicLibTest/Private/tdmtests/TestTransforms.cpp b/Engine/Plugins/Animation/RigLogic/Source/RigLogicLibTest/Private/tdmtests/TestTransforms.cpp
new file mode 100644
--- /dev/null
+++ b/Engine/Plugins/Animation/RigLogic/Source/RigLogicLibTest/Private/tdmtests/TestTransforms.cpp
@@ -0,0 +1,110 @@
+// Copyright Epic Games, Inc. All Rights Reserved.
+
+#include "tdmtests/Defs.h"
+#include "tdmtests/Helpers.h"
+
+#include "tdm/TDM.h"
+
+TEST(TestTransforms, Rotate) {
+    using namespace tdm::ang_literals;
+    tdm::vec3<float> axis{0.0f, 0.0f, 1.0f};
+    auto m = tdm::rotate(axis, tdm::frad{90.0_fdeg});
+    tdm::vec4<float> v{1.0f, 0.0f, 0.0f, 1.0f};
+    auto result = v * m;
+    tdm::vec4<float> expected{0.0f, -1.0f, 0.0f, 1.0f};
+    ASSERT_VEC_NEAR(result, expected, 0.0001f);
+}
+
+TEST(TestTransforms, RotateLeftHanded) {
+    using namespace tdm::ang_literals;
+    tdm::vec3<float> axis{0.5f, 0.4f, 0.3f};
+    auto m = tdm::rotate(axis, tdm::frad{45.0_fdeg}, tdm::handedness::left);
+    tdm::mat4<float> expected{0.853553f, 0.417157f, -0.312132f, 0.0f,
+                              -0.182843f, 0.800833f, 0.570294f, 0.0f,
+                              0.487868f, -0.429706f, 0.759828f, 0.0f,
+                              0.0f, 0.0f, 0.0f, 1.0f};
+    ASSERT_MAT_NEAR(m, expected, 0.0001f);
+}
+
+TEST(TestTransforms, ScaleUniform) {
+    auto m = tdm::mat4<float>::identity();
+    m = tdm::scale(m, 3.0f);
+    tdm::mat4<float> expected{3.0f, 0.0f, 0.0f, 0.0f,
+                              0.0f, 3.0f, 0.0f, 0.0f,
+                              0.0f, 0.0f, 3.0f, 0.0f,
+                              0.0f, 0.0f, 0.0f, 1.0f};
+    ASSERT_MAT_NEAR(m, expected, 0.0001f);
+}
+
+TEST(TestTransforms, ScaleNonUniform) {
+    auto m = tdm::mat4<float>::identity();
+    m = tdm::scale(m, tdm::vec3<float>{2.0f, 3.0f, 4.0f});
+    tdm::mat4<float> expected{2.0f, 0.0f, 0.0f, 0.0f,
+                              0.0f, 3.0f, 0.0f, 0.0f,
+                              0.0f, 0.0f, 4.0f, 0.0f,
+                              0.0f, 0.0f, 0.0f, 1.0f};
+    ASSERT_MAT_NEAR(m, expected, 0.0001f);
+}
+
+TEST(TestTransforms, Translate) {
+    auto m = tdm::mat4<float>::identity();
+    m = tdm::translate(m, tdm::vec3<float>{2.0f, 3.0f, 4.0f});
+    tdm::mat4<float> expected{1.0f, 0.0f, 0.0f, 0.0f,
+                              0.0f, 1.0f, 0.0f, 0.0f,
+                              0.0f, 0.0f, 1.0f, 0.0f,
+                              2.0f, 3.0f, 4.0f, 1.0f};
+    ASSERT_MAT_NEAR(m, expected, 0.0001f);
+}
+
+TEST(TestTransforms, CreateRotationMatrixFromXYZAngles) {
+    using namespace tdm::ang_literals;
+    tdm::mat4<float> m = tdm::rotate<float>(tdm::frad{90.0_fdeg}, tdm::frad{0.0_fdeg}, tdm::frad{90.0_fdeg});
+    tdm::mat4<float> expected{0.0f, 1.0f, 0.0f, 0.0f,
+                              0.0f, 0.0f, 1.0f, 0.0f,
+                              1.0f, 0.0f, 0.0f, 0.0f,
+                              0.0f, 0.0f, 0.0f, 1.0f};
+    ASSERT_MAT_NEAR(m, expected, 0.0001f);
+}
+
+TEST(TestTransforms, CreateRotationMatrixFromXYZAnglesVector) {
+    using namespace tdm::ang_literals;
+    tdm::mat4<float> m = tdm::rotate<float>({tdm::frad{90.0_fdeg}, tdm::frad{0.0_fdeg}, tdm::frad{90.0_fdeg}});
+    tdm::mat4<float> expected{0.0f, 1.0f, 0.0f, 0.0f,
+                              0.0f, 0.0f, 1.0f, 0.0f,
+                              1.0f, 0.0f, 0.0f, 0.0f,
+                              0.0f, 0.0f, 0.0f, 1.0f};
+    ASSERT_MAT_NEAR(m, expected, 0.0001f);
+}
+
+TEST(TestTransforms, RotateByXYZAngles) {
+    using namespace tdm::ang_literals;
+    auto m = tdm::mat4<float>::identity();
+    m = tdm::rotate<float>(m, tdm::frad{45.0_fdeg}, tdm::frad{0.0_fdeg}, tdm::frad{30.0_fdeg});
+    tdm::mat4<float> expected{0.866025f, 0.5f, 0.0f, 0.0f,
+                              -0.353553f, 0.612372f, 0.707107f, 0.0f,
+                              0.353553f, -0.612372f, 0.707107f, 0.0f,
+                              0.0f, 0.0f, 0.0f, 1.0f};
+    ASSERT_MAT_NEAR(m, expected, 0.0001f);
+}
+
+TEST(TestTransforms, RotateByXYZAnglesVector) {
+    using namespace tdm::ang_literals;
+    auto m = tdm::mat4<float>::identity();
+    m = tdm::rotate<float>(m, {tdm::frad{45.0_fdeg}, tdm::frad{0.0_fdeg}, tdm::frad{30.0_fdeg}});
+    tdm::mat4<float> expected{0.866025f, 0.5f, 0.0f, 0.0f,
+                              -0.353553f, 0.612372f, 0.707107f, 0.0f,
+                              0.353553f, -0.612372f, 0.707107f, 0.0f,
+                              0.0f, 0.0f, 0.0f, 1.0f};
+    ASSERT_MAT_NEAR(m, expected, 0.0001f);
+}
+
+TEST(TestTransforms, RotateByXYZAnglesLeftHanded) {
+    using namespace tdm::ang_literals;
+    auto m = tdm::mat4<float>::identity();
+    m = tdm::rotate<float>(m, tdm::frad{45.0_fdeg}, tdm::frad{0.0_fdeg}, tdm::frad{30.0_fdeg}, tdm::handedness::left);
+    tdm::mat4<float> expected{0.8660254f, -0.5f, 0.0f, 0.0f,
+                              0.3535534f, 0.6123725f, -0.7071068f, 0.0f,
+                              0.3535534f, 0.6123725f, 0.7071068f, 0.0f,
+                              0.0f, 0.0f, 0.0f, 1.0f};
+    ASSERT_MAT_NEAR(m, expected, 0.0001f);
+}
diff --git a/Engine/Plugins/Animation/RigLogic/Source/RigLogicLibTest/Private/trimdtests/Defs.h b/Engine/Plugins/Animation/RigLogic/Source/RigLogicLibTest/Private/trimdtests/Defs.h
--- a/Engine/Plugins/Animation/RigLogic/Source/RigLogicLibTest/Private/trimdtests/Defs.h
+++ b/Engine/Plugins/Animation/RigLogic/Source/RigLogicLibTest/Private/trimdtests/Defs.h
@@ -20,6 +20,15 @@
     #pragma warning(pop)
 #endif
 
+#ifdef _MSC_VER
+    #pragma warning(push)
+    #pragma warning(disable : 4365 4987)
+#endif
+#include <cmath>
+#ifdef _MSC_VER
+    #pragma warning(pop)
+#endif
+
 #ifndef INSTANTIATE_TEST_SUITE_P
     #define INSTANTIATE_TEST_SUITE_P INSTANTIATE_TEST_CASE_P
 #endif
@@ -63,3 +72,17 @@
             EXPECT_NEAR(result[i], expected[i], threshold);             \
         }
 #endif
+
+template<typename TVec>
+inline bool near(const TVec& lhs, const TVec& rhs, float threshold) {
+    float ls[TVec::size()];
+    float rs[TVec::size()];
+    lhs.unalignedStore(ls);
+    rhs.unalignedStore(rs);
+    for (std::size_t i = {}; i < lhs.size(); ++i) {
+        if (std::abs(ls[i] - rs[i]) > threshold) {
+            return false;
+        }
+    }
+    return true;
+}
diff --git a/Engine/Plugins/Animation/RigLogic/Source/RigLogicLibTest/Private/trimdtests/TestT128.cpp b/Engine/Plugins/Animation/RigLogic/Source/RigLogicLibTest/Private/trimdtests/TestT128.cpp
--- a/Engine/Plugins/Animation/RigLogic/Source/RigLogicLibTest/Private/trimdtests/TestT128.cpp
+++ b/Engine/Plugins/Animation/RigLogic/Source/RigLogicLibTest/Private/trimdtests/TestT128.cpp
@@ -388,6 +388,18 @@ TEST(T128Test, AndNotScalar) {
     ASSERT_TRUE(equal(result2, e2));
 }
 
+TEST(T128Test, RsqrtScalar) {
+    trimd::scalar::F128 v{1.0f, 2.0f, 3.0f, 9.0f};
+    v = trimd::scalar::rsqrt(v);
+    trimd::scalar::F128 e{1.0f, 0.70710678f, 0.57735026f, 0.33333333f};
+    #ifdef TRIMD_ENABLE_FAST_INVERSE_SQRT
+    static constexpr float threshold = 0.0004f;
+    #else
+    static constexpr float threshold = 0.0002f;
+    #endif  // TRIMD_ENABLE_FAST_INVERSE_SQRT
+    ASSERT_TRUE(near(v, e, threshold));
+}
+
 #ifdef TRIMD_ENABLE_SSE
     TEST(T128Test, TransposeSquareSSE) {
         trimd::sse::F128 v1{1.0f, 2.0f, 3.0f, 4.0f};
@@ -427,6 +439,18 @@ TEST(T128Test, AndNotScalar) {
         ASSERT_TRUE(equal(result2, e2));
     }
 
+    TEST(T128Test, RsqrtSSE) {
+        trimd::sse::F128 v{1.0f, 2.0f, 3.0f, 9.0f};
+        v = trimd::sse::rsqrt(v);
+        trimd::sse::F128 e{1.0f, 0.70710678f, 0.57735026f, 0.33333333f};
+        #ifdef TRIMD_ENABLE_FAST_INVERSE_SQRT
+        static constexpr float threshold = 0.0004f;
+        #else
+        static constexpr float threshold = 0.0003f;
+        #endif  // TRIMD_ENABLE_FAST_INVERSE_SQRT
+        ASSERT_TRUE(near(v, e, threshold));
+    }
+
     #ifdef TRIMD_ENABLE_F16C
         TEST(T128Test, LoadAlignedHalfFloatsSSE) {
             alignas(trimd::sse::F128::alignment()) const std::uint16_t halfFloats[] = {15360, 16384, 16896, 17408};
@@ -513,6 +537,18 @@ TEST(T128Test, AndNotScalar) {
         ASSERT_TRUE(equal(result2, e2));
     }
 
+    TEST(T128Test, RsqrtNEON) {
+        trimd::neon::F128 v{1.0f, 2.0f, 3.0f, 9.0f};
+        v = trimd::neon::rsqrt(v);
+        trimd::neon::F128 e{1.0f, 0.70710678f, 0.57735026f, 0.33333333f};
+        #ifdef TRIMD_ENABLE_FAST_INVERSE_SQRT
+        static constexpr float threshold = 0.0004f;
+        #else
+        static constexpr float threshold = 0.0002f;
+        #endif  // TRIMD_ENABLE_FAST_INVERSE_SQRT
+        ASSERT_TRUE(near(v, e, threshold));
+    }
+
     #ifdef TRIMD_ENABLE_NEON_FP16
         TEST(T128Test, LoadAlignedHalfFloatsNEON) {
             alignas(trimd::neon::F128::alignment()) const std::uint16_t halfFloats[] = {15360, 16384, 16896, 17408};
diff --git a/Engine/Plugins/Animation/RigLogic/Source/RigLogicLibTest/Private/trimdtests/TestT256.cpp b/Engine/Plugins/Animation/RigLogic/Source/RigLogicLibTest/Private/trimdtests/TestT256.cpp
--- a/Engine/Plugins/Animation/RigLogic/Source/RigLogicLibTest/Private/trimdtests/TestT256.cpp
+++ b/Engine/Plugins/Animation/RigLogic/Source/RigLogicLibTest/Private/trimdtests/TestT256.cpp
@@ -423,7 +423,6 @@ TYPED_TEST(T256Test, GreaterThanOrEqual) {
     ASSERT_TRUE(equal(v1 >= v4, m14));
 }
 
-
 TYPED_TEST(T256Test, BitwiseAND) {
     using F256 = typename TestFixture::T256;
     F256 v{1.0f, 2.0f, 3.0f, 4.0f, 5.0f, 6.0f, 7.0f, 8.0f};
@@ -767,6 +766,18 @@ TEST(T256Test, AndNotScalar) {
     ASSERT_TRUE(equal(result2, e2));
 }
 
+TEST(T256Test, RsqrtScalar) {
+    trimd::scalar::F256 v{1.0f, 2.0f, 3.0f, 9.0f, 9.0f, 3.0f, 2.0f, 1.0,};
+    v = trimd::scalar::rsqrt(v);
+    trimd::scalar::F256 e{1.0f, 0.70710678f, 0.57735026f, 0.33333333f, 0.33333333f, 0.57735026f, 0.70710678f, 1.0f};
+    #ifdef TRIMD_ENABLE_FAST_INVERSE_SQRT
+    static constexpr float threshold = 0.0004f;
+    #else
+    static constexpr float threshold = 0.0002f;
+    #endif  // TRIMD_ENABLE_FAST_INVERSE_SQRT
+    ASSERT_TRUE(near(v, e, threshold));
+}
+
 #ifdef TRIMD_ENABLE_AVX
     TEST(T256Test, TransposeSquareAVX) {
         trimd::avx::F256 v1{1.0f, 2.0f, 3.0f, 4.0f, 5.0f, 6.0f, 7.0f, 8.0f};
@@ -832,6 +843,17 @@ TEST(T256Test, AndNotScalar) {
         ASSERT_TRUE(equal(result2, e2));
     }
 
+    TEST(T256Test, RsqrtAVX) {
+        trimd::avx::F256 v{1.0f, 2.0f, 3.0f, 9.0f, 9.0f, 3.0f, 2.0f, 1.0,};
+        v = trimd::avx::rsqrt(v);
+        trimd::avx::F256 e{1.0f, 0.70710678f, 0.57735026f, 0.33333333f, 0.33333333f, 0.57735026f, 0.70710678f, 1.0f};
+        #ifdef TRIMD_ENABLE_FAST_INVERSE_SQRT
+        static constexpr float threshold = 0.0004f;
+        #else
+        static constexpr float threshold = 0.0003f;
+        #endif  // TRIMD_ENABLE_FAST_INVERSE_SQRT
+        ASSERT_TRUE(near(v, e, threshold));
+    }
 #endif  // TRIMD_ENABLE_AVX
 
 #ifdef TRIMD_ENABLE_F16C
diff --git a/Engine/Plugins/Animation/RigLogic/Source/RigLogicLibTest/RigLogicLibTest.Build.cs b/Engine/Plugins/Animation/RigLogic/Source/RigLogicLibTest/RigLogicLibTest.Build.cs
--- a/Engine/Plugins/Animation/RigLogic/Source/RigLogicLibTest/RigLogicLibTest.Build.cs
+++ b/Engine/Plugins/Animation/RigLogic/Source/RigLogicLibTest/RigLogicLibTest.Build.cs
@@ -12,6 +12,8 @@ namespace UnrealBuildTool.Rules
 			bUseUnity = false; // A windows include is preprocessing some method names causing compile failures.
 			bDisableStaticAnalysis = true;
 
+			PrivateDefinitions.Add("RL_BUILD_WITH_XYZ_ROTATION_ORDER=1");
+
 			if (Target.Platform == UnrealTargetPlatform.Win64 && Target.Architecture == UnrealArch.X64)
 			{
 				PrivateDefinitions.Add("RL_BUILD_WITH_SSE=1");
diff --git a/Engine/Plugins/Animation/RigLogic/Source/RigLogicModule/Private/AnimNode_RigLogic.cpp b/Engine/Plugins/Animation/RigLogic/Source/RigLogicModule/Private/AnimNode_RigLogic.cpp
--- a/Engine/Plugins/Animation/RigLogic/Source/RigLogicModule/Private/AnimNode_RigLogic.cpp
+++ b/Engine/Plugins/Animation/RigLogic/Source/RigLogicModule/Private/AnimNode_RigLogic.cpp
@@ -14,12 +14,14 @@
 #include "HAL/LowLevelMemTracker.h"
 #include "Animation/MorphTarget.h"
 
+#include <tdm/TDM.h>
+
 #include UE_INLINE_GENERATED_CPP_BY_NAME(AnimNode_RigLogic)
 
 LLM_DEFINE_TAG(Animation_RigLogic);
 DEFINE_LOG_CATEGORY(LogRigLogicAnimNode);
 
-static constexpr uint16 ATTR_COUNT_PER_JOINT = 9;
+static constexpr uint16 ATTR_COUNT_PER_JOINT = 10;
 
 FAnimNode_RigLogic::FAnimNode_RigLogic() : RigInstance(nullptr)
 {
@@ -96,8 +98,11 @@ void FAnimNode_RigLogic::CacheBones_AnyThread(const FAnimationCacheBonesContext&
 	const FBoneContainer& RequiredBones = Context.AnimInstanceProxy->GetRequiredBones();
 	if (RequiredBones.IsValid())
 	{
+		// Set current LOD
 		RigInstance->SetLOD(Context.AnimInstanceProxy->GetLODLevel());
 		const uint16 CurrentLOD = RigInstance->GetLOD();
+
+		// Populate mapping of DNA joint indices to CompactPoseBoneIndex entries (used in updating joints with results from RigLogic)
 		const TArray<uint16>& VariableJointIndices = LocalRigRuntimeContext->VariableJointIndicesPerLOD[CurrentLOD].Values;
 		JointsMapDNAIndicesToCompactPoseBoneIndices.Empty();
 		JointsMapDNAIndicesToCompactPoseBoneIndices.Reserve(VariableJointIndices.Num());
@@ -109,6 +114,26 @@ void FAnimNode_RigLogic::CacheBones_AnyThread(const FAnimationCacheBonesContext&
 				JointsMapDNAIndicesToCompactPoseBoneIndices.Add({JointIndex, CompactPoseBoneIndex});
 			}
 		}
+
+		// Populate driver joint to raw control attribute mapping (used to feed RigLogic with inputs from the joint hierarchy)
+		SparseDriverJointsToControlAttributesMap.Empty();
+		DenseDriverJointsToControlAttributesMap.Empty();
+		DenseDriverJointsToControlAttributesMap.Reserve(LocalDNAIndexMapping->DriverJointsToControlAttributesMap.Num());
+		// Sparse mapping will likely remain empty so no reservation happens
+		for (const auto& Mapping : LocalDNAIndexMapping->DriverJointsToControlAttributesMap)
+		{
+			const FCompactPoseBoneIndex CompactPoseBoneIndex = RequiredBones.MakeCompactPoseIndex(Mapping.MeshPoseBoneIndex);
+			if (CompactPoseBoneIndex != INDEX_NONE) {
+				if ((Mapping.RotationX != INDEX_NONE) && (Mapping.RotationY != INDEX_NONE) && (Mapping.RotationZ != INDEX_NONE) && (Mapping.RotationW != INDEX_NONE))
+				{
+					DenseDriverJointsToControlAttributesMap.Add({CompactPoseBoneIndex, Mapping.DNAJointIndex, Mapping.RotationX, Mapping.RotationY, Mapping.RotationZ, Mapping.RotationW});
+				}
+				else
+				{
+					SparseDriverJointsToControlAttributesMap.Add({CompactPoseBoneIndex, Mapping.DNAJointIndex, Mapping.RotationX, Mapping.RotationY, Mapping.RotationZ, Mapping.RotationW});
+				}
+			}
+		}
 	}
 }
 
@@ -135,11 +160,10 @@ void FAnimNode_RigLogic::Evaluate_AnyThread(FPoseContext& OutputContext)
 
 	if (IsLODEnabled(OutputContext.AnimInstanceProxy))
 	{
-		UpdateControlCurves(OutputContext, LocalDNAIndexMapping.Get());
+		UpdateControlCurves(OutputContext, LocalDNAIndexMapping.Get(), LocalRigRuntimeContext->RigLogic->GetNeutralJointValues());
 		CalculateRigLogic(LocalRigRuntimeContext->RigLogic.Get());
 		const uint16 CurrentLOD = RigInstance->GetLOD();
-		TArrayView<const uint16> VariableJointIndices = LocalRigRuntimeContext->VariableJointIndicesPerLOD[CurrentLOD].Values;
-		UpdateJoints(VariableJointIndices, LocalRigRuntimeContext->RigLogic->GetRawNeutralJointValues(), RigInstance->GetRawJointOutputs(), OutputContext);
+		UpdateJoints(LocalRigRuntimeContext->RigLogic->GetNeutralJointValues(), RigInstance->GetJointOutputs(), OutputContext);
 		UpdateBlendShapeCurves(LocalDNAIndexMapping.Get(), RigInstance->GetBlendShapeOutputs(), OutputContext);
 		UpdateAnimMapCurves(LocalDNAIndexMapping.Get(), RigInstance->GetAnimatedMapOutputs(), OutputContext);
 	}
@@ -150,7 +174,7 @@ void FAnimNode_RigLogic::GatherDebugData(FNodeDebugData& DebugData)
 	AnimSequence.GatherDebugData(DebugData);
 }
 
-void FAnimNode_RigLogic::UpdateControlCurves(const FPoseContext& InputContext, const FDNAIndexMapping* DNAIndexMapping)
+void FAnimNode_RigLogic::UpdateControlCurves(const FPoseContext& InputContext, const FDNAIndexMapping* DNAIndexMapping, TArrayView<const float> NeutralJointValues)
 {
 	DECLARE_SCOPE_HIERARCHICAL_COUNTER_FUNC()
 	QUICK_SCOPE_CYCLE_COUNTER(STAT_AnimNode_RigLogic_UpdateControlCurves);
@@ -166,6 +190,55 @@ void FAnimNode_RigLogic::UpdateControlCurves(const FPoseContext& InputContext, c
 			}
 		});
 
+	// The sparse mapping is NOT guaranteed to supply all quaternion attributes, so checks for each attribute mapping are present
+	for (int32 MappingIndex = 0; MappingIndex < SparseDriverJointsToControlAttributesMap.Num(); ++MappingIndex)
+	{
+		const FCompactPoseBoneControlAttributeMapping& Mapping = SparseDriverJointsToControlAttributesMap[MappingIndex];
+		const FTransform& CompactPose = InputContext.Pose[Mapping.CompactPoseBoneIndex];
+		// Translation and Scale is currently not used here, so to avoid the overhead of checking them, they are simply ignored.
+		// Should the need arise to use them as well, this code will need adjustment.
+		const FQuat Rotation = CompactPose.GetRotation();
+		const int32 AttrIndex = Mapping.DNAJointIndex * ATTR_COUNT_PER_JOINT;
+		const tdm::fquat NeutralRotation{NeutralJointValues[AttrIndex + 3], NeutralJointValues[AttrIndex + 4], NeutralJointValues[AttrIndex + 5], NeutralJointValues[AttrIndex + 6]};
+		const tdm::fquat AbsPoseRotation{static_cast<float>(Rotation.X), static_cast<float>(Rotation.Y), static_cast<float>(Rotation.Z), static_cast<float>(Rotation.W)};
+		const tdm::fquat DeltaPoseRotation = tdm::inverse(NeutralRotation) * AbsPoseRotation;
+
+		if (Mapping.RotationX != INDEX_NONE)
+		{
+			RigInstance->SetRawControl(Mapping.RotationX, DeltaPoseRotation.x);
+		}
+		if (Mapping.RotationY != INDEX_NONE)
+		{
+			RigInstance->SetRawControl(Mapping.RotationY, DeltaPoseRotation.y);
+		}
+		if (Mapping.RotationZ != INDEX_NONE)
+		{
+			RigInstance->SetRawControl(Mapping.RotationZ, DeltaPoseRotation.z);
+		}
+		if (Mapping.RotationW != INDEX_NONE)
+		{
+			RigInstance->SetRawControl(Mapping.RotationW, DeltaPoseRotation.w);
+		}
+	}
+
+	// The dense mapping is guaranteed to supply all quaternion attributes, so NO checks for each attribute mapping are present
+	for (int32 MappingIndex = 0; MappingIndex < DenseDriverJointsToControlAttributesMap.Num(); ++MappingIndex)
+	{
+		const FCompactPoseBoneControlAttributeMapping& Mapping = DenseDriverJointsToControlAttributesMap[MappingIndex];
+		const FTransform& CompactPose = InputContext.Pose[Mapping.CompactPoseBoneIndex];
+		// Translation and Scale is currently not used here, so to avoid the overhead of checking them, they are simply ignored.
+		// Should the need arise to use them as well, this code will need adjustment.
+		const FQuat Rotation = CompactPose.GetRotation();
+		const int32 AttrIndex = Mapping.DNAJointIndex * ATTR_COUNT_PER_JOINT;
+		const tdm::fquat NeutralRotation{ NeutralJointValues[AttrIndex + 3], NeutralJointValues[AttrIndex + 4], NeutralJointValues[AttrIndex + 5], NeutralJointValues[AttrIndex + 6] };
+		const tdm::fquat AbsPoseRotation{ static_cast<float>(Rotation.X), static_cast<float>(Rotation.Y), static_cast<float>(Rotation.Z), static_cast<float>(Rotation.W) };
+		const tdm::fquat DeltaPoseRotation = tdm::inverse(NeutralRotation) * AbsPoseRotation;
+		RigInstance->SetRawControl(Mapping.RotationX, DeltaPoseRotation.x);
+		RigInstance->SetRawControl(Mapping.RotationY, DeltaPoseRotation.y);
+		RigInstance->SetRawControl(Mapping.RotationZ, DeltaPoseRotation.z);
+		RigInstance->SetRawControl(Mapping.RotationW, DeltaPoseRotation.w);
+	}
+
 	if (RigInstance->GetNeuralNetworkCount() != 0) {
 		UE::Anim::FNamedValueArrayUtils::Union(InputContext.Curve, DNAIndexMapping->NeuralNetworkMaskCurves,
 			[this](const UE::Anim::FCurveElement& InCurveElement, const UE::Anim::FCurveElementIndexed& InControlAttributeCurveElement, UE::Anim::ENamedValueUnionFlags InFlags)
@@ -185,7 +258,7 @@ void FAnimNode_RigLogic::CalculateRigLogic(FRigLogic* RigLogic)
 	RigLogic->Calculate(RigInstance);
 }
 
-void FAnimNode_RigLogic::UpdateJoints(TArrayView<const uint16> VariableJointIndices, TArrayView<const float> NeutralJointValues, TArrayView<const float> DeltaJointValues, FPoseContext& OutputContext)
+void FAnimNode_RigLogic::UpdateJoints(TArrayView<const float> NeutralJointValues, TArrayView<const float> DeltaJointValues, FPoseContext& OutputContext)
 {
 	DECLARE_SCOPE_HIERARCHICAL_COUNTER_FUNC()
 	QUICK_SCOPE_CYCLE_COUNTER(STAT_AnimNode_RigLogic_UpdateJoints);
@@ -196,12 +269,9 @@ void FAnimNode_RigLogic::UpdateJoints(TArrayView<const uint16> VariableJointIndi
 	{
 		const uint16 AttrIndex = Mapping.JointIndex * ATTR_COUNT_PER_JOINT;
 		FTransform& CompactPose = OutputContext.Pose[Mapping.CompactPoseBoneIndex];
-		// Translation: X = X, Y = -Y, Z = Z
-		CompactPose.SetTranslation(FVector((N[AttrIndex + 0] + D[AttrIndex + 0]), -(N[AttrIndex + 1] + D[AttrIndex + 1]), (N[AttrIndex + 2] + D[AttrIndex + 2])));
-		// Rotation: X = -Y, Y = -Z, Z = X
-		CompactPose.SetRotation(FQuat(FRotator(-N[AttrIndex + 4], -N[AttrIndex + 5], N[AttrIndex + 3])) * FQuat(FRotator(-D[AttrIndex + 4], -D[AttrIndex + 5], D[AttrIndex + 3])));
-		// Scale: X = X, Y = Y, Z = Z
-		CompactPose.SetScale3D(FVector((N[AttrIndex + 6] + D[AttrIndex + 6]), (N[AttrIndex + 7] + D[AttrIndex + 7]), (N[AttrIndex + 8] + D[AttrIndex + 8])));
+		CompactPose.SetTranslation(FVector((N[AttrIndex + 0] + D[AttrIndex + 0]), (N[AttrIndex + 1] + D[AttrIndex + 1]), (N[AttrIndex + 2] + D[AttrIndex + 2])));
+		CompactPose.SetRotation(FQuat(N[AttrIndex + 3], N[AttrIndex + 4], N[AttrIndex + 5], N[AttrIndex + 6]) * FQuat(D[AttrIndex + 3], D[AttrIndex + 4], D[AttrIndex + 5], D[AttrIndex + 6]));
+		CompactPose.SetScale3D(FVector((N[AttrIndex + 7] + D[AttrIndex + 7]), (N[AttrIndex + 8] + D[AttrIndex + 8]), (N[AttrIndex + 9] + D[AttrIndex + 9])));
 	}
 }
 
@@ -238,4 +308,3 @@ void FAnimNode_RigLogic::UpdateAnimMapCurves(const FDNAIndexMapping* DNAIndexMap
 		}
 	});
 }
-
diff --git a/Engine/Plugins/Animation/RigLogic/Source/RigLogicModule/Private/DNAAsset.cpp b/Engine/Plugins/Animation/RigLogic/Source/RigLogicModule/Private/DNAAsset.cpp
--- a/Engine/Plugins/Animation/RigLogic/Source/RigLogicModule/Private/DNAAsset.cpp
+++ b/Engine/Plugins/Animation/RigLogic/Source/RigLogicModule/Private/DNAAsset.cpp
@@ -8,9 +8,13 @@
 #include "DNAReaderAdapter.h"
 #include "DNAIndexMapping.h"
 #include "FMemoryResource.h"
+#include "RigLogicDNAReader.h"
 #include "RigLogicMemoryStream.h"
 #include "SharedRigRuntimeContext.h"
 
+#if WITH_EDITORONLY_DATA
+    #include "EditorFramework/AssetImportData.h"
+#endif
 #include "Engine/AssetUserData.h"
 #include "Serialization/BufferArchive.h"
 #include "Serialization/MemoryReader.h"
@@ -21,7 +25,6 @@
 #include "Animation/Skeleton.h"
 #include "Components/SkeletalMeshComponent.h"
 #include "Engine/SkeletalMesh.h"
-#include "EditorFramework/AssetImportData.h"
 
 #include "riglogic/RigLogic.h"
 
@@ -36,7 +39,7 @@ static constexpr uint32 AVG_GEOMETRY_SIZE = 50 * 1024 * 1024;
 
 static TSharedPtr<IDNAReader> ReadDNAFromStream(rl4::BoundedIOStream* Stream, EDNADataLayer Layer, uint16 MaxLOD)
 {
-	auto DNAStreamReader = rl4::makeScoped<dna::BinaryStreamReader>(Stream, static_cast<dna::DataLayer>(Layer), dna::UnknownLayerPolicy::Preserve, MaxLOD, FMemoryResource::Instance());
+	auto DNAStreamReader = rl4::makeScoped<dna::BinaryStreamReader>(Stream, CalculateDNADataLayerBitmask(Layer), dna::UnknownLayerPolicy::Preserve, MaxLOD, FMemoryResource::Instance());
 	DNAStreamReader->read();
 	if (!rl4::Status::isOk())
 	{
@@ -51,7 +54,7 @@ static void WriteDNAToStream(const IDNAReader* Source, EDNADataLayer Layer, rl4:
 	auto DNAWriter = rl4::makeScoped<dna::BinaryStreamWriter>(Destination, FMemoryResource::Instance());
 	if (Source != nullptr)
 	{
-		DNAWriter->setFrom(Source->Unwrap(), static_cast<dna::DataLayer>(Layer), dna::UnknownLayerPolicy::Preserve, FMemoryResource::Instance());
+		DNAWriter->setFrom(Source->Unwrap(), CalculateDNADataLayerBitmask(Layer), dna::UnknownLayerPolicy::Preserve, FMemoryResource::Instance());
 	}
 	DNAWriter->write();
 }
@@ -110,7 +113,12 @@ void UDNAAsset::SetBehaviorReader(TSharedPtr<IDNAReader> SourceDNAReader)
 {
 	FWriteScopeLock DNAScopeLock{DNAUpdateLock};
 	const size_t PredictedSize = (SourceDNAReader->GetNeuralNetworkCount() != 0) ? AVG_BEHAVIOR_SIZE + AVG_MACHINE_LEARNED_BEHAVIOR_SIZE : AVG_BEHAVIOR_SIZE;
-	BehaviorReader = CopyDNALayer(SourceDNAReader.Get(), EDNADataLayer::Behavior | EDNADataLayer::MachineLearnedBehavior, PredictedSize);
+	const EDNADataLayer BehaviorLayers = (
+		EDNADataLayer::Behavior |
+		EDNADataLayer::MachineLearnedBehavior |
+		EDNADataLayer::RBFBehavior
+	);
+	BehaviorReader = CopyDNALayer(SourceDNAReader.Get(), BehaviorLayers, PredictedSize);
 	InvalidateRigRuntimeContext();
 	InitializeRigRuntimeContext();
 }
@@ -179,7 +187,10 @@ void UDNAAsset::InitializeRigRuntimeContext()
 	if (BehaviorReader.IsValid() && (BehaviorReader->GetJointCount() != 0))
 	{
 		NewContext->BehaviorReader = BehaviorReader;
-		NewContext->RigLogic = MakeShared<FRigLogic>(BehaviorReader.Get());
+		// Convert behavior data from DNA to UE space on the fly as RigLogic accesses it
+		RigLogicDNAReader BehaviorReaderInUESpace{BehaviorReader->Unwrap()};
+		FDNAReader<RigLogicDNAReader> BehaviorReaderInUESpaceWrapper{&BehaviorReaderInUESpace};
+		NewContext->RigLogic = MakeShared<FRigLogic>(&BehaviorReaderInUESpaceWrapper);
 		NewContext->CacheVariableJointIndices();
 		{
 			FWriteScopeLock ContextScopeLock{RigRuntimeContextUpdateLock};
@@ -191,6 +202,7 @@ void UDNAAsset::InitializeRigRuntimeContext()
 			BehaviorReader->Unload(EDNADataLayer::Behavior);
 			BehaviorReader->Unload(EDNADataLayer::Geometry);
 			BehaviorReader->Unload(EDNADataLayer::MachineLearnedBehavior);
+			BehaviorReader->Unload(EDNADataLayer::RBFBehavior);
 		}
 #endif  // !WITH_EDITOR
 	}
@@ -204,8 +216,7 @@ TSharedPtr<FSharedRigRuntimeContext> UDNAAsset::GetRigRuntimeContext()
 	return RigRuntimeContext;
 }
 
-TSharedPtr<FDNAIndexMapping> UDNAAsset::GetDNAIndexMapping(const USkeleton* Skeleton,
-														   const USkeletalMesh* SkeletalMesh)
+TSharedPtr<FDNAIndexMapping> UDNAAsset::GetDNAIndexMapping(const USkeleton* Skeleton, const USkeletalMesh* SkeletalMesh)
 {
 	LLM_SCOPE_BYNAME(TEXT("Animation/RigLogic"));
 
@@ -238,6 +249,7 @@ TSharedPtr<FDNAIndexMapping> UDNAAsset::GetDNAIndexMapping(const USkeleton* Skel
 		DNAIndexMapping->MapControlCurves(BehaviorReader.Get(), Skeleton);
 		DNAIndexMapping->MapNeuralNetworkMaskCurves(BehaviorReader.Get(), Skeleton);
 		DNAIndexMapping->MapJoints(BehaviorReader.Get(), SkeletalMesh);
+		DNAIndexMapping->MapDriverJoints(BehaviorReader.Get(), SkeletalMesh);
 		DNAIndexMapping->MapMorphTargets(BehaviorReader.Get(), Skeleton, SkeletalMesh);
 		DNAIndexMapping->MapMaskMultipliers(BehaviorReader.Get(), Skeleton);
 		DNAIndexMappingContainer.Add(SkeletalMesh, DNAIndexMapping);
@@ -254,6 +266,7 @@ bool UDNAAsset::Init(const FString& DNAFilename)
 	{
 		UE_LOG(LogDNAAsset, Warning, TEXT("%s"), ANSI_TO_TCHAR(rl4::Status::get().message));
 	}
+
 #if WITH_EDITORONLY_DATA
 	AssetImportData = NewObject<UAssetImportData>(this, TEXT("AssetImportData"));
 	TArray<FAssetImportInfo::FSourceFile> SourceFiles = { FAssetImportInfo::FSourceFile(DNAFilename) };
@@ -280,7 +293,12 @@ bool UDNAAsset::Init(const FString& DNAFilename)
 	FWriteScopeLock DNAScopeLock{DNAUpdateLock};
 
 	// Load run-time data (behavior) from whole-DNA buffer into BehaviorReader
-	BehaviorReader = ReadDNAFromBuffer(&TempFileBuffer, EDNADataLayer::Behavior | EDNADataLayer::MachineLearnedBehavior, 0u); //0u = MaxLOD
+	const EDNADataLayer BehaviorLayers = (
+		EDNADataLayer::Behavior |
+		EDNADataLayer::MachineLearnedBehavior |
+		EDNADataLayer::RBFBehavior
+	);
+	BehaviorReader = ReadDNAFromBuffer(&TempFileBuffer, BehaviorLayers, 0u); //0u = MaxLOD
 	if (!BehaviorReader.IsValid())
 	{
 		return false;
@@ -321,7 +339,12 @@ void UDNAAsset::Serialize(FArchive& Ar)
 		if (Ar.IsLoading())
 		{
 			FArchiveMemoryStream BehaviorStream{&Ar};
-			BehaviorReader = ReadDNAFromStream(&BehaviorStream, EDNADataLayer::Behavior | EDNADataLayer::MachineLearnedBehavior, 0u); //0u = max LOD
+			const EDNADataLayer BehaviorLayers = (
+				EDNADataLayer::Behavior |
+				EDNADataLayer::MachineLearnedBehavior |
+				EDNADataLayer::RBFBehavior
+			);
+			BehaviorReader = ReadDNAFromStream(&BehaviorStream, BehaviorLayers, 0u); //0u = max LOD
 			// Geometry data is always present (even if only as an empty placeholder), just so the uasset
 			// format remains consistent between editor and non-editor builds
 			FArchiveMemoryStream GeometryStream{&Ar};
@@ -340,7 +363,12 @@ void UDNAAsset::Serialize(FArchive& Ar)
 			TSharedPtr<IDNAReader> EmptyDNA = CreateEmptyDNA(AVG_EMPTY_SIZE);
 			IDNAReader* BehaviorReaderPtr = (BehaviorReader.IsValid() ? static_cast<IDNAReader*>(BehaviorReader.Get()) : EmptyDNA.Get());
 			FArchiveMemoryStream BehaviorStream{&Ar};
-			WriteDNAToStream(BehaviorReaderPtr, EDNADataLayer::Behavior | EDNADataLayer::MachineLearnedBehavior, &BehaviorStream);
+			const EDNADataLayer BehaviorLayers = (
+				EDNADataLayer::Behavior |
+				EDNADataLayer::MachineLearnedBehavior |
+				EDNADataLayer::RBFBehavior
+			);
+			WriteDNAToStream(BehaviorReaderPtr, BehaviorLayers, &BehaviorStream);
 
 			// When cooking (or when there was no Geometry data available), an empty DNA structure is written
 			// into the stream, serving as a placeholder just so uasset files can be conveniently loaded
diff --git a/Engine/Plugins/Animation/RigLogic/Source/RigLogicModule/Private/DNAIndexMapping.cpp b/Engine/Plugins/Animation/RigLogic/Source/RigLogicModule/Private/DNAIndexMapping.cpp
--- a/Engine/Plugins/Animation/RigLogic/Source/RigLogicModule/Private/DNAIndexMapping.cpp
+++ b/Engine/Plugins/Animation/RigLogic/Source/RigLogicModule/Private/DNAIndexMapping.cpp
@@ -77,6 +77,68 @@ void FDNAIndexMapping::MapNeuralNetworkMaskCurves(const IDNAReader* DNAReader, c
 	}
 }
 
+void FDNAIndexMapping::MapDriverJoints(const IDNAReader* DNAReader, const USkeletalMesh* SkeletalMesh)
+{
+	LLM_SCOPE_BYNAME(TEXT("Animation/RigLogic"));
+
+	const FReferenceSkeleton& RefSkeleton = SkeletalMesh->GetRefSkeleton();
+	const uint32 ControlCount = DNAReader->GetRawControlCount();
+
+	DriverJointsToControlAttributesMap.Empty();
+	// This is a correct approximation as long as only 4 (rotation) attributes are used as driver joint attributes
+	// and no regular raw controls are present in the DNA
+	DriverJointsToControlAttributesMap.Reserve(ControlCount / 4);
+
+	for (uint32_t ControlIndex = 0; ControlIndex < ControlCount; ++ControlIndex)
+	{
+		const FString DriverJointAttrName = DNAReader->GetRawControlName(ControlIndex);
+		if (DriverJointAttrName.Len() < 2)
+		{
+			continue;
+		}
+		const FString DriverJointName = DriverJointAttrName.Mid(0, DriverJointAttrName.Len() - 2);
+		const FName BoneName = FName(*DriverJointName);
+		const int32 BoneIndex = RefSkeleton.FindBoneIndex(BoneName);
+		if (BoneIndex == INDEX_NONE)
+		{
+			// Mixed DNAs will contain both driver joints and normal raw controls in this list, and those will
+			// not be found in the joint hierarchy
+			continue;
+		}
+
+		int32 MappingIndex = DriverJointsToControlAttributesMap.FindLastByPredicate([BoneIndex](const FMeshPoseBoneControlAttributeMapping& Element)
+		{
+			return Element.MeshPoseBoneIndex.GetInt() == BoneIndex;
+		});
+		if (MappingIndex == INDEX_NONE)
+		{
+			FMeshPoseBoneControlAttributeMapping NewMapping{FMeshPoseBoneIndex{BoneIndex}, INDEX_NONE, INDEX_NONE , INDEX_NONE , INDEX_NONE, INDEX_NONE};
+			// BoneIndex may be INDEX_NONE, but it's handled properly by the Evaluate method
+			MappingIndex = DriverJointsToControlAttributesMap.Add(NewMapping);
+		}
+
+		FMeshPoseBoneControlAttributeMapping& Mapping = DriverJointsToControlAttributesMap[MappingIndex];
+		Mapping.DNAJointIndex = JointsMapDNAIndicesToMeshPoseBoneIndices.Find(Mapping.MeshPoseBoneIndex);
+
+		if (DriverJointAttrName.EndsWith(TEXT(".x")))
+		{
+			Mapping.RotationX = ControlIndex;
+		}
+		else if (DriverJointAttrName.EndsWith(TEXT(".y")))
+		{
+			Mapping.RotationY = ControlIndex;
+		}
+		else if (DriverJointAttrName.EndsWith(TEXT(".z")))
+		{
+			Mapping.RotationZ = ControlIndex;
+		}
+		else if (DriverJointAttrName.EndsWith(TEXT(".w")))
+		{
+			Mapping.RotationW = ControlIndex;
+		}
+	}
+}
+
 void FDNAIndexMapping::MapJoints(const IDNAReader* DNAReader, const USkeletalMesh* SkeletalMesh)
 {
 	LLM_SCOPE_BYNAME(TEXT("Animation/RigLogic"));
diff --git a/Engine/Plugins/Animation/RigLogic/Source/RigLogicModule/Private/DNAUtils.cpp b/Engine/Plugins/Animation/RigLogic/Source/RigLogicModule/Private/DNAUtils.cpp
--- a/Engine/Plugins/Animation/RigLogic/Source/RigLogicModule/Private/DNAUtils.cpp
+++ b/Engine/Plugins/Animation/RigLogic/Source/RigLogicModule/Private/DNAUtils.cpp
@@ -4,6 +4,7 @@
 
 #include "CoreMinimal.h"
 #include "HAL/LowLevelMemTracker.h"
+#include "DNACommon.h"
 #include "DNAReaderAdapter.h"
 #include "FMemoryResource.h"
 #include "RigLogicMemoryStream.h"
@@ -11,6 +12,16 @@
 #include "riglogic/RigLogic.h"
 
 
+dna::DataLayer CalculateDNADataLayerBitmask(EDNADataLayer Layer)
+{
+	dna::DataLayer mask = static_cast<dna::DataLayer>(Layer);
+	if ((Layer & EDNADataLayer::RBFBehavior) == EDNADataLayer::RBFBehavior) {
+		mask = mask | dna::DataLayer::JointBehaviorMetadata;
+		mask = mask | dna::DataLayer::TwistSwingBehavior;
+	}
+	return mask;
+}
+
 static TSharedPtr<IDNAReader> ReadDNAStream(rl4::ScopedPtr<dna::BinaryStreamReader> DNAStreamReader)
 {
 	DNAStreamReader->read();
@@ -26,7 +37,7 @@ TSharedPtr<IDNAReader> ReadDNAFromFile(const FString& Path, EDNADataLayer Layer,
 {
 	LLM_SCOPE_BYNAME(TEXT("Animation/RigLogic"));
 	auto DNAFileStream = rl4::makeScoped<rl4::MemoryMappedFileStream>(TCHAR_TO_UTF8(*Path), rl4::MemoryMappedFileStream::AccessMode::Read, FMemoryResource::Instance());
-	auto DNAStreamReader = rl4::makeScoped<dna::BinaryStreamReader>(DNAFileStream.get(), static_cast<dna::DataLayer>(Layer), dna::UnknownLayerPolicy::Preserve, MaxLOD, FMemoryResource::Instance());
+	auto DNAStreamReader = rl4::makeScoped<dna::BinaryStreamReader>(DNAFileStream.get(), CalculateDNADataLayerBitmask(Layer), dna::UnknownLayerPolicy::Preserve, MaxLOD, FMemoryResource::Instance());
 	return ReadDNAStream(MoveTemp(DNAStreamReader));
 }
 
@@ -34,7 +45,7 @@ TSharedPtr<IDNAReader> ReadDNAFromFile(const FString& Path, EDNADataLayer Layer,
 {
 	LLM_SCOPE_BYNAME(TEXT("Animation/RigLogic"));
 	auto DNAFileStream = rl4::makeScoped<rl4::MemoryMappedFileStream>(TCHAR_TO_UTF8(*Path), rl4::MemoryMappedFileStream::AccessMode::Read, FMemoryResource::Instance());
-	auto DNAStreamReader = rl4::makeScoped<dna::BinaryStreamReader>(DNAFileStream.get(), static_cast<dna::DataLayer>(Layer), dna::UnknownLayerPolicy::Preserve, LODs.GetData(), static_cast<uint16>(LODs.Num()), FMemoryResource::Instance());
+	auto DNAStreamReader = rl4::makeScoped<dna::BinaryStreamReader>(DNAFileStream.get(), CalculateDNADataLayerBitmask(Layer), dna::UnknownLayerPolicy::Preserve, LODs.GetData(), static_cast<uint16>(LODs.Num()), FMemoryResource::Instance());
 	return ReadDNAStream(MoveTemp(DNAStreamReader));
 }
 
@@ -42,7 +53,7 @@ TSharedPtr<IDNAReader> ReadDNAFromBuffer(TArray<uint8>* DNABuffer, EDNADataLayer
 {
 	LLM_SCOPE_BYNAME(TEXT("Animation/RigLogic"));
 	FRigLogicMemoryStream DNAMemoryStream(DNABuffer);
-	auto DNAStreamReader = rl4::makeScoped<dna::BinaryStreamReader>(&DNAMemoryStream, static_cast<dna::DataLayer>(Layer), dna::UnknownLayerPolicy::Preserve, MaxLOD, FMemoryResource::Instance());
+	auto DNAStreamReader = rl4::makeScoped<dna::BinaryStreamReader>(&DNAMemoryStream, CalculateDNADataLayerBitmask(Layer), dna::UnknownLayerPolicy::Preserve, MaxLOD, FMemoryResource::Instance());
 	return ReadDNAStream(MoveTemp(DNAStreamReader));
 }
 
@@ -50,7 +61,7 @@ TSharedPtr<IDNAReader> ReadDNAFromBuffer(TArray<uint8>* DNABuffer, EDNADataLayer
 {
 	LLM_SCOPE_BYNAME(TEXT("Animation/RigLogic"));
 	FRigLogicMemoryStream DNAMemoryStream(DNABuffer);
-	auto DNAStreamReader = rl4::makeScoped<dna::BinaryStreamReader>(&DNAMemoryStream, static_cast<dna::DataLayer>(Layer), dna::UnknownLayerPolicy::Preserve, LODs.GetData(), static_cast<uint16>(LODs.Num()), FMemoryResource::Instance());
+	auto DNAStreamReader = rl4::makeScoped<dna::BinaryStreamReader>(&DNAMemoryStream, CalculateDNADataLayerBitmask(Layer), dna::UnknownLayerPolicy::Preserve, LODs.GetData(), static_cast<uint16>(LODs.Num()), FMemoryResource::Instance());
 	return ReadDNAStream(MoveTemp(DNAStreamReader));
 }
 
@@ -60,7 +71,7 @@ TArray<uint8> ReadStreamFromDNA(const IDNAReader* Reader, EDNADataLayer Layer)
 	TArray<char> DNABuffer;
 	auto DeltaDnaStream = rl4::makeScoped<trio::MemoryStream>();
 	auto DNAStreamWriter = rl4::makeScoped<dna::BinaryStreamWriter>(DeltaDnaStream.get(), FMemoryResource::Instance());
-	DNAStreamWriter->setFrom(Reader->Unwrap(), static_cast<dna::DataLayer>(Layer), dna::UnknownLayerPolicy::Preserve, FMemoryResource::Instance());
+	DNAStreamWriter->setFrom(Reader->Unwrap(), CalculateDNADataLayerBitmask(Layer), dna::UnknownLayerPolicy::Preserve, FMemoryResource::Instance());
 	DNAStreamWriter->write();
 	DNABuffer.AddZeroed(DeltaDnaStream->size());
 	DeltaDnaStream->read(DNABuffer.GetData(), DeltaDnaStream->size());
@@ -72,6 +83,6 @@ void WriteDNAToFile(const IDNAReader* Reader, EDNADataLayer Layer, const FString
 	LLM_SCOPE_BYNAME(TEXT("Animation/RigLogic"));
 	auto DNAFileStream = rl4::makeScoped<rl4::FileStream>(TCHAR_TO_UTF8(*Path), rl4::FileStream::AccessMode::Write, rl4::FileStream::OpenMode::Binary, FMemoryResource::Instance());
 	auto DNAStreamWriter = rl4::makeScoped<dna::BinaryStreamWriter>(DNAFileStream.get(), FMemoryResource::Instance());
-	DNAStreamWriter->setFrom(Reader->Unwrap(), static_cast<dna::DataLayer>(Layer), dna::UnknownLayerPolicy::Preserve, FMemoryResource::Instance());
+	DNAStreamWriter->setFrom(Reader->Unwrap(), CalculateDNADataLayerBitmask(Layer), dna::UnknownLayerPolicy::Preserve, FMemoryResource::Instance());
 	DNAStreamWriter->write();
 }
diff --git a/Engine/Plugins/Animation/RigLogic/Source/RigLogicModule/Private/RigInstance.cpp b/Engine/Plugins/Animation/RigLogic/Source/RigLogicModule/Private/RigInstance.cpp
--- a/Engine/Plugins/Animation/RigLogic/Source/RigLogicModule/Private/RigInstance.cpp
+++ b/Engine/Plugins/Animation/RigLogic/Source/RigLogicModule/Private/RigInstance.cpp
@@ -119,16 +119,26 @@ void FRigInstance::SetNeuralNetworkMask(uint16 NeuralNetIndex, float Value)
 	RigInstance->setNeuralNetworkMask(NeuralNetIndex, Value);
 }
 
-TArrayView<const float> FRigInstance::GetRawJointOutputs() const
+uint16 FRigInstance::GetRBFControlCount() const
 {
-	rl4::ConstArrayView<float> Outputs = RigInstance->getRawJointOutputs();
-	return TArrayView<const float>{Outputs.data(), static_cast<int32>(Outputs.size())};
+	return RigInstance->getRBFControlCount();
 }
 
-FTransformArrayView FRigInstance::GetJointOutputs() const
+float FRigInstance::GetRBFControl(uint16 Index) const
 {
-	rl4::ConstArrayView<float> Outputs = RigInstance->getRawJointOutputs();
-	return FTransformArrayView(Outputs.data(), Outputs.size());
+	return RigInstance->getRBFControl(Index);
+}
+
+TArrayView<const float> FRigInstance::GetRBFControlValues() const
+{
+	rl4::ConstArrayView<float> Values = RigInstance->getRBFControlValues();
+	return TArrayView<const float>{Values.data(), static_cast<int32>(Values.size())};
+}
+
+TArrayView<const float> FRigInstance::GetJointOutputs() const
+{
+	rl4::ConstArrayView<float> Outputs = RigInstance->getJointOutputs();
+	return TArrayView<const float>{Outputs.data(), static_cast<int32>(Outputs.size())};
 }
 
 TArrayView<const float> FRigInstance::GetBlendShapeOutputs() const
diff --git a/Engine/Plugins/Animation/RigLogic/Source/RigLogicModule/Private/RigLogic.cpp b/Engine/Plugins/Animation/RigLogic/Source/RigLogicModule/Private/RigLogic.cpp
--- a/Engine/Plugins/Animation/RigLogic/Source/RigLogicModule/Private/RigLogic.cpp
+++ b/Engine/Plugins/Animation/RigLogic/Source/RigLogicModule/Private/RigLogic.cpp
@@ -13,9 +13,26 @@ void FRigLogic::FRigLogicDeleter::operator()(rl4::RigLogic* Pointer)
 	rl4::RigLogic::destroy(Pointer);
 }
 
-FRigLogic::FRigLogic(const IDNAReader* Reader, ERigLogicCalculationType CalculationType) :
+static rl4::Configuration AdaptRigLogicConfig(const FRigLogicConfiguration& Config)
+{
+	rl4::Configuration Copy = {};
+	Copy.calculationType = static_cast<rl4::CalculationType>(Config.CalculationType);
+	Copy.loadJoints = Config.LoadJoints;
+	Copy.loadBlendShapes = Config.LoadBlendShapes;
+	Copy.loadAnimatedMaps = Config.LoadAnimatedMaps;
+	Copy.loadMachineLearnedBehavior = Config.LoadMachineLearnedBehavior;
+	Copy.loadRBFBehavior = Config.LoadRBFBehavior;
+	Copy.loadTwistSwingBehavior = Config.LoadTwistSwingBehavior;
+	Copy.translationType = static_cast<rl4::TranslationType>(Config.TranslationType);
+	Copy.rotationType = static_cast<rl4::RotationType>(Config.RotationType);
+	Copy.rotationOrder = static_cast<rl4::RotationOrder>(Config.RotationOrder);
+	Copy.scaleType = static_cast<rl4::ScaleType>(Config.ScaleType);
+	return Copy;
+}
+
+FRigLogic::FRigLogic(const IDNAReader* Reader, FRigLogicConfiguration Config) :
 	MemoryResource{FMemoryResource::SharedInstance()},
-	RigLogic{rl4::RigLogic::create(Reader->Unwrap(), rl4::Configuration{static_cast<rl4::CalculationType>(CalculationType)}, FMemoryResource::Instance())}
+	RigLogic{rl4::RigLogic::create(Reader->Unwrap(), AdaptRigLogicConfig(Config), FMemoryResource::Instance())}
 {
 }
 
@@ -26,18 +43,12 @@ uint16 FRigLogic::GetLODCount() const
 	return RigLogic->getLODCount();
 }
 
-TArrayView<const float> FRigLogic::GetRawNeutralJointValues() const
+TArrayView<const float> FRigLogic::GetNeutralJointValues() const
 {
-	rl4::ConstArrayView<float> Values = RigLogic->getRawNeutralJointValues();
+	rl4::ConstArrayView<float> Values = RigLogic->getNeutralJointValues();
 	return TArrayView<const float>{Values.data(), static_cast<int32>(Values.size())};
 }
 
-FTransformArrayView FRigLogic::GetNeutralJointValues() const
-{
-	rl4::ConstArrayView<float> Values = RigLogic->getRawNeutralJointValues();
-	return FTransformArrayView{Values.data(), Values.size()};
-}
-
 TArrayView<const uint16> FRigLogic::GetJointVariableAttributeIndices(uint16 LOD) const
 {
 	rl4::ConstArrayView<uint16> Indices = RigLogic->getJointVariableAttributeIndices(LOD);
@@ -54,6 +65,11 @@ uint16 FRigLogic::GetNeuralNetworkCount() const
 	return RigLogic->getNeuralNetworkCount();
 }
 
+uint16 FRigLogic::GetRBFSolverCount() const
+{
+	return RigLogic->getRBFSolverCount();
+}
+
 uint16 FRigLogic::GetMeshCount() const
 {
 	return RigLogic->getMeshCount();
@@ -95,6 +111,16 @@ void FRigLogic::CalculateMachineLearnedBehaviorControls(FRigInstance* Instance,
 	RigLogic->calculateMachineLearnedBehaviorControls(Instance->Unwrap(), NeuralNetIndex);
 }
 
+void FRigLogic::CalculateRBFControls(FRigInstance* Instance) const
+{
+	RigLogic->calculateRBFControls(Instance->Unwrap());
+}
+
+void FRigLogic::CalculateRBFControls(FRigInstance* Instance, uint16 SolverIndex) const
+{
+	RigLogic->calculateRBFControls(Instance->Unwrap(), SolverIndex);
+}
+
 void FRigLogic::CalculateJoints(FRigInstance* Instance) const
 {
 	RigLogic->calculateJoints(Instance->Unwrap());
diff --git a/Engine/Plugins/Animation/RigLogic/Source/RigLogicModule/Private/RigUnit_RigLogic.cpp b/Engine/Plugins/Animation/RigLogic/Source/RigLogicModule/Private/RigUnit_RigLogic.cpp
--- a/Engine/Plugins/Animation/RigLogic/Source/RigLogicModule/Private/RigUnit_RigLogic.cpp
+++ b/Engine/Plugins/Animation/RigLogic/Source/RigLogicModule/Private/RigUnit_RigLogic.cpp
@@ -12,11 +12,13 @@
 #include "Math/TransformNonVectorized.h"
 #include "Units/RigUnitContext.h"
 
+#include <tdm/TDM.h>
+
 #include UE_INLINE_GENERATED_CPP_BY_NAME(RigUnit_RigLogic)
 
 DEFINE_LOG_CATEGORY(LogRigLogicUnit);
 
-const uint8 FRigUnit_RigLogic_Data::MAX_ATTRS_PER_JOINT = 9;
+const uint8 FRigUnit_RigLogic_Data::MAX_ATTRS_PER_JOINT = 10;
 
 /** Constructs curve name from nameToSplit using formatString of form x<obj>y<attr>z **/
 static FString ConstructCurveName(const FString& NameToSplit, const FString& FormatString)
@@ -179,6 +181,74 @@ void FRigUnit_RigLogic_Data::MapJoints(const URigHierarchy* Hierarchy)
 		const int32 BoneIndex = Hierarchy->GetIndex(FRigElementKey(JointFName, ERigElementType::Bone));
 		HierarchyBoneIndices.Add(BoneIndex);
 	}
+
+	auto FindJointIndex = [DNABehavior](const FString& JointName)
+		{
+			for (uint16 JointIndex = 0; JointIndex < DNABehavior->GetJointCount(); ++JointIndex)
+			{
+				const FString& RLJointName = DNABehavior->GetJointName(JointIndex);
+				if (RLJointName == JointName)
+				{
+					return JointIndex;
+				}
+			}
+			return static_cast<uint16>(-1);
+		};
+
+	const uint32 ControlCount = DNABehavior->GetRawControlCount();
+	DriverJointsToControlAttributesMap.Empty();
+	// This is a correct approximation as long as only 4 (rotation) attributes are used as driver joint attributes
+	// and no regular raw controls are present in the DNA
+	DriverJointsToControlAttributesMap.Reserve(ControlCount / 4);
+	for (uint32_t ControlIndex = 0; ControlIndex < ControlCount; ++ControlIndex)
+	{
+		const FString DriverJointAttrName = DNABehavior->GetRawControlName(ControlIndex);
+		if (DriverJointAttrName.Len() < 2)
+		{
+			continue;
+		}
+
+		const FString DriverJointName = DriverJointAttrName.Mid(0, DriverJointAttrName.Len() - 2);
+		const FName BoneName = FName(*DriverJointName);
+		const int32 BoneIndex = Hierarchy->GetIndex(FRigElementKey(BoneName, ERigElementType::Bone));
+		if (BoneIndex == INDEX_NONE)
+		{
+			// Mixed DNAs will contain both driver joints and normal raw controls in this list, and those will
+			// not be found in the joint hierarchy
+			continue;
+		}
+
+		int32 MappingIndex = DriverJointsToControlAttributesMap.FindLastByPredicate([BoneIndex](const FBoneIndexControlAttributeMapping& Element)
+			{
+				return Element.BoneIndex == BoneIndex;
+			});
+		if (MappingIndex == INDEX_NONE)
+		{
+			FBoneIndexControlAttributeMapping NewMapping{BoneIndex, INDEX_NONE, INDEX_NONE , INDEX_NONE , INDEX_NONE, INDEX_NONE};
+			// BoneIndex may be INDEX_NONE, but it's handled properly by the Evaluate method
+			MappingIndex = DriverJointsToControlAttributesMap.Add(NewMapping);
+		}
+
+		FBoneIndexControlAttributeMapping& Mapping = DriverJointsToControlAttributesMap[MappingIndex];
+		Mapping.DNAJointIndex = FindJointIndex(DriverJointName);
+
+		if (DriverJointAttrName.EndsWith(TEXT(".x")))
+		{
+			Mapping.RotationX = ControlIndex;
+		}
+		else if (DriverJointAttrName.EndsWith(TEXT(".y")))
+		{
+			Mapping.RotationY = ControlIndex;
+		}
+		else if (DriverJointAttrName.EndsWith(TEXT(".z")))
+		{
+			Mapping.RotationZ = ControlIndex;
+		}
+		else if (DriverJointAttrName.EndsWith(TEXT(".w")))
+		{
+			Mapping.RotationW = ControlIndex;
+		}
+	}
 }
 
 void FRigUnit_RigLogic_Data::MapMorphTargets(const URigHierarchy* InHierarchy)
@@ -249,7 +319,7 @@ void FRigUnit_RigLogic_Data::MapMaskMultipliers(const URigHierarchy* InHierarchy
 	}
 }
 
-void FRigUnit_RigLogic_Data::CalculateRigLogic(const URigHierarchy* InHierarchy)
+void FRigUnit_RigLogic_Data::CalculateRigLogic(const URigHierarchy* InHierarchy, TArrayView<const float> NeutralJointValues)
 {
 	DECLARE_SCOPE_HIERARCHICAL_COUNTER_FUNC()
 	QUICK_SCOPE_CYCLE_COUNTER(STAT_RigUnit_RigLogic_Calculate);
@@ -266,6 +336,36 @@ void FRigUnit_RigLogic_Data::CalculateRigLogic(const URigHierarchy* InHierarchy)
 		RigInstance->SetRawControl(ControlIndex, Value);
 	}
 
+	for (int32 MappingIndex = 0; MappingIndex < DriverJointsToControlAttributesMap.Num(); ++MappingIndex)
+	{
+		const FBoneIndexControlAttributeMapping& Mapping = DriverJointsToControlAttributesMap[MappingIndex];
+		const FTransform& PoseTransform = InHierarchy->GetLocalTransform(Mapping.BoneIndex);
+		// Translation and Scale is currently not used here, so to avoid the overhead of checking them, they are simply ignored.
+		// Should the need arise to use them as well, this code will need adjustment.
+		const FQuat Rotation = PoseTransform.GetRotation();
+		const int32 AttrIndex = Mapping.DNAJointIndex * MAX_ATTRS_PER_JOINT;
+		const tdm::fquat NeutralRotation{NeutralJointValues[AttrIndex + 3], NeutralJointValues[AttrIndex + 4], NeutralJointValues[AttrIndex + 5], NeutralJointValues[AttrIndex + 6]};
+		const tdm::fquat AbsPoseRotation{static_cast<float>(Rotation.X), static_cast<float>(Rotation.Y), static_cast<float>(Rotation.Z), static_cast<float>(Rotation.W)};
+		const tdm::fquat DeltaPoseRotation = tdm::inverse(NeutralRotation) * AbsPoseRotation;
+
+		if (Mapping.RotationX != INDEX_NONE)
+		{
+			RigInstance->SetRawControl(Mapping.RotationX, DeltaPoseRotation.x);
+		}
+		if (Mapping.RotationY != INDEX_NONE)
+		{
+			RigInstance->SetRawControl(Mapping.RotationY, DeltaPoseRotation.y);
+		}
+		if (Mapping.RotationZ != INDEX_NONE)
+		{
+			RigInstance->SetRawControl(Mapping.RotationZ, DeltaPoseRotation.z);
+		}
+		if (Mapping.RotationW != INDEX_NONE)
+		{
+			RigInstance->SetRawControl(Mapping.RotationW, DeltaPoseRotation.w);
+		}
+	}
+
 	const int32 NeuralNetworkCount = RigInstance->GetNeuralNetworkCount();
 	for (int32 NeuralNetworkIndex = 0; NeuralNetworkIndex < NeuralNetworkCount; ++NeuralNetworkIndex)
 	{
@@ -296,12 +396,9 @@ void FRigUnit_RigLogic_Data::UpdateJoints(URigHierarchy* Hierarchy, TArrayView<c
 			const uint16 AttrIndex = JointIndex * MAX_ATTRS_PER_JOINT;
 			const FTransform Transform
 			{
-				// Rotation: X = -Y, Y = -Z, Z = X
-				FQuat(FRotator(-N[AttrIndex + 4], -N[AttrIndex + 5], N[AttrIndex + 3])) * FQuat(FRotator(-D[AttrIndex + 4], -D[AttrIndex + 5], D[AttrIndex + 3])),
-				// Translation: X = X, Y = -Y, Z = Z
-				FVector((N[AttrIndex + 0] + D[AttrIndex + 0]), -(N[AttrIndex + 1] + D[AttrIndex + 1]), (N[AttrIndex + 2] + D[AttrIndex + 2])),
-				// Scale: X = X, Y = Y, Z = Z
-				FVector((N[AttrIndex + 6] + D[AttrIndex + 6]), (N[AttrIndex + 7] + D[AttrIndex + 7]), (N[AttrIndex + 8] + D[AttrIndex + 8]))
+				FQuat(N[AttrIndex + 3], N[AttrIndex + 4], N[AttrIndex + 5], N[AttrIndex + 6]) * FQuat(D[AttrIndex + 3], D[AttrIndex + 4], D[AttrIndex + 5], D[AttrIndex + 6]),
+				FVector((N[AttrIndex + 0] + D[AttrIndex + 0]), (N[AttrIndex + 1] + D[AttrIndex + 1]), (N[AttrIndex + 2] + D[AttrIndex + 2])),
+				FVector((N[AttrIndex + 7] + D[AttrIndex + 7]), (N[AttrIndex + 8] + D[AttrIndex + 8]), (N[AttrIndex + 9] + D[AttrIndex + 9]))
 			};
 			Hierarchy->SetLocalTransform(BoneIndex, Transform);
 		}
@@ -419,11 +516,12 @@ FRigUnit_RigLogic_Execute()
 			return;
 		}
 		Data.CurrentLOD = Data.SkelMeshComponent->GetPredictedLODLevel();
-		Data.CalculateRigLogic(Hierarchy);
 
-		//Filing a struct so we can call the same method for updating joints from tests
-		TArrayView<const float> NeutralJointValues = Data.LocalRigRuntimeContext->RigLogic->GetRawNeutralJointValues();
-		TArrayView<const float> DeltaJointValues = Data.RigInstance->GetRawJointOutputs();
+		TArrayView<const float> NeutralJointValues = Data.LocalRigRuntimeContext->RigLogic->GetNeutralJointValues();
+		TArrayView<const float> DeltaJointValues = Data.RigInstance->GetJointOutputs();
+
+		Data.CalculateRigLogic(Hierarchy, NeutralJointValues);
+
 		Data.UpdateJoints(Hierarchy, NeutralJointValues, DeltaJointValues);
 
 		TArrayView<const float> BlendShapeValues = Data.RigInstance->GetBlendShapeOutputs();
diff --git a/Engine/Plugins/Animation/RigLogic/Source/RigLogicModule/Private/RigUnit_RigLogic_Test.cpp b/Engine/Plugins/Animation/RigLogic/Source/RigLogicModule/Private/RigUnit_RigLogic_Test.cpp
--- a/Engine/Plugins/Animation/RigLogic/Source/RigLogicModule/Private/RigUnit_RigLogic_Test.cpp
+++ b/Engine/Plugins/Animation/RigLogic/Source/RigLogicModule/Private/RigUnit_RigLogic_Test.cpp
@@ -18,7 +18,7 @@
 #include "Units/RigUnitTest.h"
 #include "RigUnit_RigLogic.h"
 
-const uint8 FRigUnit_RigLogic::TestAccessor::MAX_ATTRS_PER_JOINT = 9;
+const uint8 FRigUnit_RigLogic::TestAccessor::MAX_ATTRS_PER_JOINT = 10;
 
 FRigUnit_RigLogic::TestAccessor::TestAccessor(FRigUnit_RigLogic* Unit)
 {
@@ -186,7 +186,7 @@ void FRigUnit_RigLogic::TestAccessor::AddToTransformArray(float* InArray, FTrans
 }
 
 
-FTransformArrayView FRigUnit_RigLogic::TestAccessor::CreateTwoJointNeutralTransforms(float *InValueArray)
+TArrayView<const float> FRigUnit_RigLogic::TestAccessor::CreateTwoJointNeutralTransforms(float *InValueArray)
 {
 	float* Transform1Ptr = InValueArray;
 
@@ -204,7 +204,7 @@ FTransformArrayView FRigUnit_RigLogic::TestAccessor::CreateTwoJointNeutralTransf
 	AddToTransformArray(Transform2Ptr, Joint2Transform);
 
 	const float* ValuesPtr = InValueArray;
-	return FTransformArrayView(ValuesPtr, sizeof(FTransform));
+	return TArrayView<const float>(ValuesPtr, 2 * MAX_ATTRS_PER_JOINT);
 }
 
 TArrayView<const uint16> FRigUnit_RigLogic::TestAccessor::CreateTwoJointVariableAttributes(uint16* InVariableAttributeIndices, uint8 LOD)
@@ -757,7 +757,7 @@ IMPLEMENT_RIGUNIT_AUTOMATION_TEST(FRigUnit_RigLogic)
 	//create neutral transforms for two bones
 	const uint8 TransformArraySize = 2 * FRigUnit_RigLogic::TestAccessor::MAX_ATTRS_PER_JOINT;
 	float NeutralValues[TransformArraySize];  //two bones, nine attributes
-	FTransformArrayView TwoJointNeutralTransforms = Test.CreateTwoJointNeutralTransforms(NeutralValues);
+	TArrayView<const float> TwoJointNeutralTransforms = Test.CreateTwoJointNeutralTransforms(NeutralValues);
 	//create delta transforms
 	float DeltaTransformData[TransformArraySize] = { 0.f };
 	//first bone translation
@@ -768,7 +768,7 @@ IMPLEMENT_RIGUNIT_AUTOMATION_TEST(FRigUnit_RigLogic)
 	DeltaTransformData[9] = 1.f;
 	DeltaTransformData[10] = 2.f;
 	DeltaTransformData[11] = 7.f;
-	FTransformArrayView DeltaTransforms = FTransformArrayView(DeltaTransformData, TransformArraySize);
+	TArrayView<const float> DeltaTransforms = TArrayView<const float>(DeltaTransformData, TransformArraySize);
 	//create variable joint index arrays for two bones
 	SharedRigRuntimeContext->VariableJointIndicesPerLOD.Reset();
 	SharedRigRuntimeContext->VariableJointIndicesPerLOD.AddDefaulted(1);
diff --git a/Engine/Plugins/Animation/RigLogic/Source/RigLogicModule/Private/SharedRigRuntimeContext.cpp b/Engine/Plugins/Animation/RigLogic/Source/RigLogicModule/Private/SharedRigRuntimeContext.cpp
--- a/Engine/Plugins/Animation/RigLogic/Source/RigLogicModule/Private/SharedRigRuntimeContext.cpp
+++ b/Engine/Plugins/Animation/RigLogic/Source/RigLogicModule/Private/SharedRigRuntimeContext.cpp
@@ -7,8 +7,6 @@
 #include "DNAReader.h"
 #include "RigLogic.h"
 
-static constexpr uint16 NUM_ATTRS_PER_JOINT = 9;
-
 void FSharedRigRuntimeContext::CacheVariableJointIndices()
 {
 	LLM_SCOPE_BYNAME(TEXT("Animation/RigLogic"));
@@ -24,6 +22,10 @@ void FSharedRigRuntimeContext::CacheVariableJointIndices()
 		DistinctVariableJointIndices.Reserve(VariableAttributeIndices.Num());
 		for (const uint16 AttrIndex : VariableAttributeIndices)
 		{
+			// In DNA, the number of joint attributes is always 9 (only RigLogic has the ability to switch this)
+			// and since the variable indices are queried from the DNA here, we deal with 9 as well, regardless of
+			// the state elsewhere where we switched to 10 since the introduction of quaternion outputs from RigLogic.
+			static constexpr uint16 NUM_ATTRS_PER_JOINT = 9;
 			const uint16 JointIndex = AttrIndex / NUM_ATTRS_PER_JOINT;
 			DistinctVariableJointIndices.Add(JointIndex);
 		}
diff --git a/Engine/Plugins/Animation/RigLogic/Source/RigLogicModule/Private/SkelMeshDNAReader.cpp b/Engine/Plugins/Animation/RigLogic/Source/RigLogicModule/Private/SkelMeshDNAReader.cpp
--- a/Engine/Plugins/Animation/RigLogic/Source/RigLogicModule/Private/SkelMeshDNAReader.cpp
+++ b/Engine/Plugins/Animation/RigLogic/Source/RigLogicModule/Private/SkelMeshDNAReader.cpp
@@ -698,6 +698,181 @@ TArrayView<const float> FSkelMeshDNAReader::GetNeuralNetworkLayerWeights(uint16
 	return BehaviorReader->GetNeuralNetworkLayerWeights(NetIndex, LayerIndex);
 }
 
+ETranslationRepresentation FSkelMeshDNAReader::GetJointTranslationRepresentation(uint16 JointIndex) const
+{
+	return BehaviorReader->GetJointTranslationRepresentation(JointIndex);
+}
+
+ERotationRepresentation FSkelMeshDNAReader::GetJointRotationRepresentation(uint16 JointIndex) const
+{
+	return BehaviorReader->GetJointRotationRepresentation(JointIndex);
+}
+
+EScaleRepresentation FSkelMeshDNAReader::GetJointScaleRepresentation(uint16 JointIndex) const
+{
+	return BehaviorReader->GetJointScaleRepresentation(JointIndex);
+}
+
+uint16 FSkelMeshDNAReader::GetRBFPoseCount() const
+{
+	return BehaviorReader->GetRBFPoseCount();
+}
+
+FString FSkelMeshDNAReader::GetRBFPoseName(uint16 PoseIndex) const
+{
+	return BehaviorReader->GetRBFPoseName(PoseIndex);
+}
+
+TArrayView<const uint16> FSkelMeshDNAReader::GetRBFPoseJointOutputIndices(uint16 PoseIndex) const
+{
+	return BehaviorReader->GetRBFPoseJointOutputIndices(PoseIndex);
+}
+
+TArrayView<const uint16> FSkelMeshDNAReader::GetRBFPoseBlendShapeChannelOutputIndices(uint16 PoseIndex) const
+{
+	return BehaviorReader->GetRBFPoseBlendShapeChannelOutputIndices(PoseIndex);
+}
+
+TArrayView<const uint16> FSkelMeshDNAReader::GetRBFPoseAnimatedMapOutputIndices(uint16 PoseIndex) const
+{
+	return BehaviorReader->GetRBFPoseAnimatedMapOutputIndices(PoseIndex);
+}
+
+TArrayView<const float> FSkelMeshDNAReader::GetRBFPoseJointOutputValues(uint16 PoseIndex) const
+{
+	return BehaviorReader->GetRBFPoseJointOutputValues(PoseIndex);
+}
+
+float FSkelMeshDNAReader::GetRBFPoseScale(uint16 PoseIndex) const
+{
+	return BehaviorReader->GetRBFPoseScale(PoseIndex);
+}
+
+uint16 FSkelMeshDNAReader::GetRBFSolverCount() const
+{
+	return BehaviorReader->GetRBFSolverCount();
+}
+
+uint16 FSkelMeshDNAReader::GetRBFSolverIndexListCount() const
+{
+	return BehaviorReader->GetRBFSolverIndexListCount();
+}
+
+TArrayView<const uint16> FSkelMeshDNAReader::GetRBFSolverIndicesForLOD(uint16 LOD) const
+{
+	return BehaviorReader->GetRBFSolverIndicesForLOD(LOD);
+}
+
+FString FSkelMeshDNAReader::GetRBFSolverName(uint16 SolverIndex) const
+{
+	return BehaviorReader->GetRBFSolverName(SolverIndex);
+}
+
+TArrayView<const uint16> FSkelMeshDNAReader::GetRBFSolverRawControlIndices(uint16 SolverIndex) const
+{
+	return BehaviorReader->GetRBFSolverRawControlIndices(SolverIndex);
+}
+
+TArrayView<const uint16> FSkelMeshDNAReader::GetRBFSolverPoseIndices(uint16 SolverIndex) const
+{
+	return BehaviorReader->GetRBFSolverPoseIndices(SolverIndex);
+}
+
+TArrayView<const float> FSkelMeshDNAReader::GetRBFSolverRawControlValues(uint16 SolverIndex) const
+{
+	return BehaviorReader->GetRBFSolverRawControlValues(SolverIndex);
+}
+
+ERBFSolverType FSkelMeshDNAReader::GetRBFSolverType(uint16 SolverIndex) const
+{
+	return BehaviorReader->GetRBFSolverType(SolverIndex);
+}
+
+float FSkelMeshDNAReader::GetRBFSolverRadius(uint16 SolverIndex) const
+{
+	return BehaviorReader->GetRBFSolverRadius(SolverIndex);
+}
+
+EAutomaticRadius FSkelMeshDNAReader::GetRBFSolverAutomaticRadius(uint16 SolverIndex) const
+{
+	return BehaviorReader->GetRBFSolverAutomaticRadius(SolverIndex);
+}
+
+float FSkelMeshDNAReader::GetRBFSolverWeightThreshold(uint16 SolverIndex) const
+{
+	return BehaviorReader->GetRBFSolverWeightThreshold(SolverIndex);
+}
+
+ERBFDistanceMethod FSkelMeshDNAReader::GetRBFSolverDistanceMethod(uint16 SolverIndex) const
+{
+	return BehaviorReader->GetRBFSolverDistanceMethod(SolverIndex);
+}
+
+ERBFNormalizeMethod FSkelMeshDNAReader::GetRBFSolverNormalizeMethod(uint16 SolverIndex) const
+{
+	return BehaviorReader->GetRBFSolverNormalizeMethod(SolverIndex);
+}
+
+ERBFFunctionType FSkelMeshDNAReader::GetRBFSolverFunctionType(uint16 SolverIndex) const
+{
+	return BehaviorReader->GetRBFSolverFunctionType(SolverIndex);
+}
+
+ETwistAxis FSkelMeshDNAReader::GetRBFSolverTwistAxis(uint16 SolverIndex) const
+{
+	return BehaviorReader->GetRBFSolverTwistAxis(SolverIndex);
+}
+
+uint16 FSkelMeshDNAReader::GetTwistCount() const
+{
+	return BehaviorReader->GetTwistCount();
+}
+
+ETwistAxis FSkelMeshDNAReader::GetTwistSetupTwistAxis(uint16 TwistIndex) const
+{
+	return BehaviorReader->GetTwistSetupTwistAxis(TwistIndex);
+}
+
+TArrayView<const uint16> FSkelMeshDNAReader::GetTwistInputControlIndices(uint16 TwistIndex) const
+{
+	return BehaviorReader->GetTwistInputControlIndices(TwistIndex);
+}
+
+TArrayView<const uint16> FSkelMeshDNAReader::GetTwistOutputJointIndices(uint16 TwistIndex) const
+{
+	return BehaviorReader->GetTwistOutputJointIndices(TwistIndex);
+}
+
+TArrayView<const float> FSkelMeshDNAReader::GetTwistBlendWeights(uint16 TwistIndex) const
+{
+	return BehaviorReader->GetTwistBlendWeights(TwistIndex);
+}
+
+uint16 FSkelMeshDNAReader::GetSwingCount() const
+{
+	return BehaviorReader->GetSwingCount();
+}
+
+ETwistAxis FSkelMeshDNAReader::GetSwingSetupTwistAxis(uint16 SwingIndex) const
+{
+	return BehaviorReader->GetSwingSetupTwistAxis(SwingIndex);
+}
+
+TArrayView<const uint16> FSkelMeshDNAReader::GetSwingInputControlIndices(uint16 SwingIndex) const
+{
+	return BehaviorReader->GetSwingInputControlIndices(SwingIndex);
+}
+
+TArrayView<const uint16> FSkelMeshDNAReader::GetSwingOutputJointIndices(uint16 SwingIndex) const
+{
+	return BehaviorReader->GetSwingOutputJointIndices(SwingIndex);
+}
+
+TArrayView<const float> FSkelMeshDNAReader::GetSwingBlendWeights(uint16 SwingIndex) const
+{
+	return BehaviorReader->GetSwingBlendWeights(SwingIndex);
+}
+
 void FSkelMeshDNAReader::Unload(EDNADataLayer Layer)
 {
 	ensureMsgf(false, TEXT("Assest are not unloadable"));
diff --git a/Engine/Plugins/Animation/RigLogic/Source/RigLogicModule/Public/AnimNode_RigLogic.h b/Engine/Plugins/Animation/RigLogic/Source/RigLogicModule/Public/AnimNode_RigLogic.h
--- a/Engine/Plugins/Animation/RigLogic/Source/RigLogicModule/Public/AnimNode_RigLogic.h
+++ b/Engine/Plugins/Animation/RigLogic/Source/RigLogicModule/Public/AnimNode_RigLogic.h
@@ -37,9 +37,9 @@ public:
 	FPoseLink AnimSequence;
 
 private:
-	void UpdateControlCurves(const FPoseContext& OutputContext, const FDNAIndexMapping* DNAIndexMapping);
+	void UpdateControlCurves(const FPoseContext& OutputContext, const FDNAIndexMapping* DNAIndexMapping, TArrayView<const float> NeutralJointValues);
 	void CalculateRigLogic(FRigLogic* RigLogic);
-	void UpdateJoints(TArrayView<const uint16> VariableJointIndices, TArrayView<const float> NeutralJointValues, TArrayView<const float> DeltaJointValues, FPoseContext& OutputContext);
+	void UpdateJoints(TArrayView<const float> NeutralJointValues, TArrayView<const float> DeltaJointValues, FPoseContext& OutputContext);
 	void UpdateBlendShapeCurves(const FDNAIndexMapping* DNAIndexMapping, TArrayView<const float> BlendShapeValues, FPoseContext& OutputContext);
 	void UpdateAnimMapCurves(const FDNAIndexMapping* DNAIndexMapping, TArrayView<const float> AnimMapOutputs, FPoseContext& OutputContext);
 
@@ -49,6 +49,15 @@ private:
 		FCompactPoseBoneIndex CompactPoseBoneIndex;
 	};
 
+	struct FCompactPoseBoneControlAttributeMapping {
+		FCompactPoseBoneIndex CompactPoseBoneIndex;
+		int32 DNAJointIndex;
+		int32 RotationX;
+		int32 RotationY;
+		int32 RotationZ;
+		int32 RotationW;
+	};
+
 private:
 	/*
 	 * Max LOD level that post-process AnimBPs are evaluated.
@@ -63,4 +72,6 @@ private:
 	TSharedPtr<FDNAIndexMapping> LocalDNAIndexMapping;
 	FRigInstance* RigInstance;
 	TArray<FJointCompactPoseBoneMapping> JointsMapDNAIndicesToCompactPoseBoneIndices;
+	TArray<FCompactPoseBoneControlAttributeMapping> SparseDriverJointsToControlAttributesMap;
+	TArray<FCompactPoseBoneControlAttributeMapping> DenseDriverJointsToControlAttributesMap;
 };
diff --git a/Engine/Plugins/Animation/RigLogic/Source/RigLogicModule/Public/DNAAsset.h b/Engine/Plugins/Animation/RigLogic/Source/RigLogicModule/Public/DNAAsset.h
--- a/Engine/Plugins/Animation/RigLogic/Source/RigLogicModule/Public/DNAAsset.h
+++ b/Engine/Plugins/Animation/RigLogic/Source/RigLogicModule/Public/DNAAsset.h
@@ -91,6 +91,8 @@ public:
 	  * deleted without affecting the other.
 	 **/
 	void InitializeForRuntimeFrom(UDNAAsset* Other);
+	TSharedPtr<FSharedRigRuntimeContext> GetRigRuntimeContext();
+	TSharedPtr<FDNAIndexMapping> GetDNAIndexMapping(const USkeleton* Skeleton, const USkeletalMesh* SkeletalMesh);
 
 	TSharedPtr<FSharedRigRuntimeContext> GetRigRuntimeContext();
 	TSharedPtr<FDNAIndexMapping> GetDNAIndexMapping(const USkeleton* Skeleton,
diff --git a/Engine/Plugins/Animation/RigLogic/Source/RigLogicModule/Public/DNACommon.h b/Engine/Plugins/Animation/RigLogic/Source/RigLogicModule/Public/DNACommon.h
--- a/Engine/Plugins/Animation/RigLogic/Source/RigLogicModule/Public/DNACommon.h
+++ b/Engine/Plugins/Animation/RigLogic/Source/RigLogicModule/Public/DNACommon.h
@@ -5,6 +5,7 @@
 #include "CoreMinimal.h"
 #include "UObject/NoExportTypes.h"
 #include "UObject/ObjectMacros.h"
+#include "RBF/RBFSolver.h"
 
 #include "DNACommon.generated.h"
 
@@ -52,6 +53,75 @@ enum class EDirection: uint8
 	Back
 };
 
+UENUM(BlueprintType)
+enum class ETranslationRepresentation : uint8
+{
+	Vector
+};
+
+UENUM(BlueprintType)
+enum class ERotationRepresentation : uint8
+{
+	EulerAngles,
+	Quaternion
+};
+
+UENUM(BlueprintType)
+enum class EScaleRepresentation : uint8
+{
+	Vector
+};
+
+/*
+UENUM(BlueprintType)
+enum class ERBFSolverType : uint8
+{
+	Additive,
+	Interpolative
+};
+
+UENUM(BlueprintType)
+enum class ERBFFunctionType : uint8
+{
+	Gaussian,
+	Exponential,
+	Linear,
+	Cubic,
+	Quintic,
+};
+
+UENUM(BlueprintType)
+enum class ERBFDistanceMethod : uint8
+{
+	Euclidean,
+	Quaternion,
+	SwingAngle,
+	TwistAngle,
+};
+
+UENUM(BlueprintType)
+enum class ERBFNormalizeMethod : uint8
+{
+	OnlyNormalizeAboveOne,
+	AlwaysNormalize
+};
+*/
+
+UENUM(BlueprintType)
+enum class EAutomaticRadius : uint8
+{
+	On,
+	Off
+};
+
+UENUM(BlueprintType)
+enum class ETwistAxis : uint8
+{
+	X,
+	Y,
+	Z
+};
+
 UENUM(BlueprintType)
 enum class EDNADataLayer : uint8
 {
@@ -62,13 +132,10 @@ enum class EDNADataLayer : uint8
 	Geometry = 8 | Definition,  // Implicitly loads Descriptor and Definition
 	GeometryWithoutBlendShapes = 16 | Definition,  // Implicitly loads Descriptor and Definition
 	MachineLearnedBehavior = 32 | Definition,  // Implicitly loads Definition
-	All = Behavior | Geometry | MachineLearnedBehavior
+	RBFBehavior = 64 | Behavior,  // Implicitly loads Behavior and all body-rig related layers
+	All = RBFBehavior | Geometry | MachineLearnedBehavior
 };
-
-inline EDNADataLayer operator|(EDNADataLayer LHS, EDNADataLayer RHS)
-{
-	return static_cast<EDNADataLayer>(static_cast<uint8>(LHS) | static_cast<uint8>(RHS));
-}
+ENUM_CLASS_FLAGS(EDNADataLayer);
 
 UENUM(BlueprintType)
 enum class EActivationFunction : uint8
diff --git a/Engine/Plugins/Animation/RigLogic/Source/RigLogicModule/Public/DNAIndexMapping.h b/Engine/Plugins/Animation/RigLogic/Source/RigLogicModule/Public/DNAIndexMapping.h
--- a/Engine/Plugins/Animation/RigLogic/Source/RigLogicModule/Public/DNAIndexMapping.h
+++ b/Engine/Plugins/Animation/RigLogic/Source/RigLogicModule/Public/DNAIndexMapping.h
@@ -19,6 +19,16 @@ struct FDNAIndexMapping
 		TArray<T> Values;
 	};
 	
+	struct FMeshPoseBoneControlAttributeMapping
+	{
+		FMeshPoseBoneIndex MeshPoseBoneIndex;
+		int32 DNAJointIndex;
+		int32 RotationX;
+		int32 RotationY;
+		int32 RotationZ;
+		int32 RotationW;
+	};
+
 	using FCachedIndexedCurve = TBaseBlendedCurve<FDefaultAllocator, UE::Anim::FCurveElementIndexed>; 
 
 	FGuid SkeletonGuid;
@@ -26,13 +36,17 @@ struct FDNAIndexMapping
 	// all the control attributes that we will need to extract, alongside their control index
 	FCachedIndexedCurve ControlAttributeCurves;
 	FCachedIndexedCurve NeuralNetworkMaskCurves;
+	TArray<FMeshPoseBoneControlAttributeMapping> DriverJointsToControlAttributesMap;
 	TArray<FMeshPoseBoneIndex> JointsMapDNAIndicesToMeshPoseBoneIndices;
 	TArray<FCachedIndexedCurve> MorphTargetCurvesPerLOD;
 	TArray<FCachedIndexedCurve> MaskMultiplierCurvesPerLOD;
 
 	void MapControlCurves(const IDNAReader* DNAReader, const USkeleton* Skeleton);
 	void MapNeuralNetworkMaskCurves(const IDNAReader* DNAReader, const USkeleton* Skeleton);
+	void MapDriverJoints(const IDNAReader* DNAReader, const USkeletalMesh* SkeletalMesh);
 	void MapJoints(const IDNAReader* DNAReader, const USkeletalMesh* SkeletalMesh);
 	void MapMorphTargets(const IDNAReader* DNAReader, const USkeleton* Skeleton, const USkeletalMesh* SkeletalMesh);
 	void MapMaskMultipliers(const IDNAReader* DNAReader, const USkeleton* Skeleton);
-};
\ No newline at end of file
+
+};
+
diff --git a/Engine/Plugins/Animation/RigLogic/Source/RigLogicModule/Public/DNAReader.h b/Engine/Plugins/Animation/RigLogic/Source/RigLogicModule/Public/DNAReader.h
--- a/Engine/Plugins/Animation/RigLogic/Source/RigLogicModule/Public/DNAReader.h
+++ b/Engine/Plugins/Animation/RigLogic/Source/RigLogicModule/Public/DNAReader.h
@@ -192,10 +192,72 @@ public:
 
 };
 
+class IJointBehaviorMetadataReader : public virtual IDefinitionReader
+{
+protected:
+	virtual ~IJointBehaviorMetadataReader() = default;
+
+public:
+	virtual ETranslationRepresentation GetJointTranslationRepresentation(uint16 JointIndex) const = 0;
+	virtual ERotationRepresentation GetJointRotationRepresentation(uint16 JointIndex) const = 0;
+	virtual EScaleRepresentation GetJointScaleRepresentation(uint16 JointIndex) const = 0;
+
+};
+
+class IRBFBehaviorReader : public virtual IBehaviorReader
+{
+protected:
+	virtual ~IRBFBehaviorReader() = default;
+
+public:
+	virtual uint16 GetRBFPoseCount() const = 0;
+	virtual FString GetRBFPoseName(uint16 PoseIndex) const = 0;
+	virtual TArrayView<const uint16> GetRBFPoseJointOutputIndices(uint16 PoseIndex) const = 0;
+	virtual TArrayView<const uint16> GetRBFPoseBlendShapeChannelOutputIndices(uint16 PoseIndex) const = 0;
+	virtual TArrayView<const uint16> GetRBFPoseAnimatedMapOutputIndices(uint16 PoseIndex) const = 0;
+	virtual TArrayView<const float> GetRBFPoseJointOutputValues(uint16 PoseIndex) const = 0;
+	virtual float GetRBFPoseScale(uint16 PoseIndex) const = 0;
+	virtual uint16 GetRBFSolverCount() const = 0;
+	virtual uint16 GetRBFSolverIndexListCount() const = 0;
+	virtual TArrayView<const uint16> GetRBFSolverIndicesForLOD(uint16 LOD) const = 0;
+	virtual FString GetRBFSolverName(uint16 SolverIndex) const = 0;
+	virtual TArrayView<const uint16> GetRBFSolverRawControlIndices(uint16 SolverIndex) const = 0;
+	virtual TArrayView<const uint16> GetRBFSolverPoseIndices(uint16 SolverIndex) const = 0;
+	virtual TArrayView<const float> GetRBFSolverRawControlValues(uint16 SolverIndex) const = 0;
+	virtual ERBFSolverType GetRBFSolverType(uint16 SolverIndex) const = 0;
+	virtual float GetRBFSolverRadius(uint16 SolverIndex) const = 0;
+	virtual EAutomaticRadius GetRBFSolverAutomaticRadius(uint16 SolverIndex) const = 0;
+	virtual float GetRBFSolverWeightThreshold(uint16 SolverIndex) const = 0;
+	virtual ERBFDistanceMethod GetRBFSolverDistanceMethod(uint16 SolverIndex) const = 0;
+	virtual ERBFNormalizeMethod GetRBFSolverNormalizeMethod(uint16 SolverIndex) const = 0;
+	virtual ERBFFunctionType GetRBFSolverFunctionType(uint16 SolverIndex) const = 0;
+	virtual ETwistAxis GetRBFSolverTwistAxis(uint16 SolverIndex) const = 0;
+
+};
+
+class ITwistSwingBehaviorReader : public virtual IDefinitionReader
+{
+protected:
+	virtual ~ITwistSwingBehaviorReader() = default;
+
+public:
+	virtual uint16 GetTwistCount() const = 0;
+	virtual ETwistAxis GetTwistSetupTwistAxis(uint16 TwistIndex) const = 0;
+	virtual TArrayView<const uint16> GetTwistInputControlIndices(uint16 TwistIndex) const = 0;
+	virtual TArrayView<const uint16> GetTwistOutputJointIndices(uint16 TwistIndex) const = 0;
+	virtual TArrayView<const float> GetTwistBlendWeights(uint16 TwistIndex) const = 0;
+	virtual uint16 GetSwingCount() const = 0;
+	virtual ETwistAxis GetSwingSetupTwistAxis(uint16 SwingIndex) const = 0;
+	virtual TArrayView<const uint16> GetSwingInputControlIndices(uint16 SwingIndex) const = 0;
+	virtual TArrayView<const uint16> GetSwingOutputJointIndices(uint16 SwingIndex) const = 0;
+	virtual TArrayView<const float> GetSwingBlendWeights(uint16 SwingIndex) const = 0;
+
+};
+
 /**
 	@brief UE interface for DNA Reader wrappers.
 */
-class IDNAReader : public IBehaviorReader, public IGeometryReader, public IMachineLearnedBehaviorReader
+class IDNAReader : public IRBFBehaviorReader, public IGeometryReader, public IMachineLearnedBehaviorReader, public IJointBehaviorMetadataReader, public ITwistSwingBehaviorReader
 {
 public:
 	virtual ~IDNAReader() = default;
diff --git a/Engine/Plugins/Animation/RigLogic/Source/RigLogicModule/Public/DNAReaderAdapter.h b/Engine/Plugins/Animation/RigLogic/Source/RigLogicModule/Public/DNAReaderAdapter.h
--- a/Engine/Plugins/Animation/RigLogic/Source/RigLogicModule/Public/DNAReaderAdapter.h
+++ b/Engine/Plugins/Animation/RigLogic/Source/RigLogicModule/Public/DNAReaderAdapter.h
@@ -3,6 +3,7 @@
 #pragma once
 
 #include "DNAReader.h"
+#include "DNAUtils.h"
 #include "FMemoryResource.h"
 
 template <class TWrappedReader>
@@ -135,6 +136,44 @@ public:
 	TArrayView<const float> GetNeuralNetworkLayerActivationFunctionParameters(uint16 NetIndex, uint16 LayerIndex) const override;
 	TArrayView<const float> GetNeuralNetworkLayerBiases(uint16 NetIndex, uint16 LayerIndex) const override;
 	TArrayView<const float> GetNeuralNetworkLayerWeights(uint16 NetIndex, uint16 LayerIndex) const override;
+	// JointBehaviorMetadataReader
+	ETranslationRepresentation GetJointTranslationRepresentation(uint16 JointIndex) const override;
+	ERotationRepresentation GetJointRotationRepresentation(uint16 JointIndex) const override;
+	EScaleRepresentation GetJointScaleRepresentation(uint16 JointIndex) const override;
+	// RBFBehavior
+	uint16 GetRBFPoseCount() const override;
+	FString GetRBFPoseName(uint16 PoseIndex) const override;
+	TArrayView<const uint16> GetRBFPoseJointOutputIndices(uint16 PoseIndex) const override;
+	TArrayView<const uint16> GetRBFPoseBlendShapeChannelOutputIndices(uint16 PoseIndex) const override;
+	TArrayView<const uint16> GetRBFPoseAnimatedMapOutputIndices(uint16 PoseIndex) const override;
+	TArrayView<const float> GetRBFPoseJointOutputValues(uint16 PoseIndex) const override;
+	float GetRBFPoseScale(uint16 PoseIndex) const override;
+	uint16 GetRBFSolverCount() const override;
+	uint16 GetRBFSolverIndexListCount() const override;
+	TArrayView<const uint16> GetRBFSolverIndicesForLOD(uint16 LOD) const override;
+	FString GetRBFSolverName(uint16 SolverIndex) const override;
+	TArrayView<const uint16> GetRBFSolverRawControlIndices(uint16 SolverIndex) const override;
+	TArrayView<const uint16> GetRBFSolverPoseIndices(uint16 SolverIndex) const override;
+	TArrayView<const float> GetRBFSolverRawControlValues(uint16 SolverIndex) const override;
+	ERBFSolverType GetRBFSolverType(uint16 SolverIndex) const override;
+	float GetRBFSolverRadius(uint16 SolverIndex) const override;
+	EAutomaticRadius GetRBFSolverAutomaticRadius(uint16 SolverIndex) const override;
+	float GetRBFSolverWeightThreshold(uint16 SolverIndex) const override;
+	ERBFDistanceMethod GetRBFSolverDistanceMethod(uint16 SolverIndex) const override;
+	ERBFNormalizeMethod GetRBFSolverNormalizeMethod(uint16 SolverIndex) const override;
+	ERBFFunctionType GetRBFSolverFunctionType(uint16 SolverIndex) const override;
+	ETwistAxis GetRBFSolverTwistAxis(uint16 SolverIndex) const override;
+	// TwistSwingBehavior
+	uint16 GetTwistCount() const override;
+	ETwistAxis GetTwistSetupTwistAxis(uint16 TwistIndex) const override;
+	TArrayView<const uint16> GetTwistInputControlIndices(uint16 TwistIndex) const override;
+	TArrayView<const uint16> GetTwistOutputJointIndices(uint16 TwistIndex) const override;
+	TArrayView<const float> GetTwistBlendWeights(uint16 TwistIndex) const override;
+	uint16 GetSwingCount() const override;
+	ETwistAxis GetSwingSetupTwistAxis(uint16 SwingIndex) const override;
+	TArrayView<const uint16> GetSwingInputControlIndices(uint16 SwingIndex) const override;
+	TArrayView<const uint16> GetSwingOutputJointIndices(uint16 SwingIndex) const override;
+	TArrayView<const float> GetSwingBlendWeights(uint16 SwingIndex) const override;
 
 	void Unload(EDNADataLayer Layer) override;
 
@@ -167,7 +206,7 @@ dna::Reader* FDNAReader<TWrappedReader>::Unwrap() const
 template <class TWrappedReader>
 void FDNAReader<TWrappedReader>::Unload(EDNADataLayer Layer)
 {
-	ReaderPtr->unload(static_cast<dna::DataLayer>(Layer));
+	ReaderPtr->unload(CalculateDNADataLayerBitmask(Layer));
 }
 
 template <class TWrappedReader>
@@ -249,7 +288,8 @@ template <class TWrappedReader>
 FCoordinateSystem FDNAReader<TWrappedReader>::GetCoordinateSystem() const
 {
 	const auto System = ReaderPtr->getCoordinateSystem();
-	return FCoordinateSystem{
+	return FCoordinateSystem
+	{
 		static_cast<EDirection>(System.xAxis),
 		static_cast<EDirection>(System.yAxis),
 		static_cast<EDirection>(System.zAxis)
@@ -916,3 +956,228 @@ TArrayView<const float> FDNAReader<TWrappedReader>::GetNeuralNetworkLayerWeights
 	const auto Values = ReaderPtr->getNeuralNetworkLayerWeights(NetIndex, LayerIndex);
 	return TArrayView<const float>(Values.data(), static_cast<int32>(Values.size()));
 }
+
+template <class TWrappedReader>
+ETranslationRepresentation FDNAReader<TWrappedReader>::GetJointTranslationRepresentation(uint16 JointIndex) const
+{
+	return static_cast<ETranslationRepresentation>(ReaderPtr->getJointTranslationRepresentation(JointIndex));
+}
+
+template <class TWrappedReader>
+ERotationRepresentation FDNAReader<TWrappedReader>::GetJointRotationRepresentation(uint16 JointIndex) const
+{
+	return static_cast<ERotationRepresentation>(ReaderPtr->getJointRotationRepresentation(JointIndex));
+}
+
+template <class TWrappedReader>
+EScaleRepresentation FDNAReader<TWrappedReader>::GetJointScaleRepresentation(uint16 JointIndex) const
+{
+	return static_cast<EScaleRepresentation>(ReaderPtr->getJointScaleRepresentation(JointIndex));
+}
+
+template <class TWrappedReader>
+uint16 FDNAReader<TWrappedReader>::GetRBFPoseCount() const
+{
+	return ReaderPtr->getRBFPoseCount();
+}
+
+template <class TWrappedReader>
+FString FDNAReader<TWrappedReader>::GetRBFPoseName(uint16 PoseIndex) const
+{
+	return FString(ANSI_TO_TCHAR(ReaderPtr->getRBFPoseName(PoseIndex).data()));
+}
+
+template <class TWrappedReader>
+TArrayView<const uint16> FDNAReader<TWrappedReader>::GetRBFPoseJointOutputIndices(uint16 PoseIndex) const
+{
+	const auto Indices = ReaderPtr->getRBFPoseJointOutputIndices(PoseIndex);
+	return TArrayView<const uint16>(Indices.data(), static_cast<int32>(Indices.size()));
+}
+
+template <class TWrappedReader>
+TArrayView<const uint16> FDNAReader<TWrappedReader>::GetRBFPoseBlendShapeChannelOutputIndices(uint16 PoseIndex) const
+{
+	const auto Indices = ReaderPtr->getRBFPoseBlendShapeChannelOutputIndices(PoseIndex);
+	return TArrayView<const uint16>(Indices.data(), static_cast<int32>(Indices.size()));
+}
+
+template <class TWrappedReader>
+TArrayView<const uint16> FDNAReader<TWrappedReader>::GetRBFPoseAnimatedMapOutputIndices(uint16 PoseIndex) const
+{
+	const auto Indices = ReaderPtr->getRBFPoseAnimatedMapOutputIndices(PoseIndex);
+	return TArrayView<const uint16>(Indices.data(), static_cast<int32>(Indices.size()));
+}
+
+template <class TWrappedReader>
+TArrayView<const float> FDNAReader<TWrappedReader>::GetRBFPoseJointOutputValues(uint16 PoseIndex) const
+{
+	const auto Values = ReaderPtr->getRBFPoseJointOutputValues(PoseIndex);
+	return TArrayView<const float>(Values.data(), static_cast<int32>(Values.size()));
+}
+
+template <class TWrappedReader>
+float FDNAReader<TWrappedReader>::GetRBFPoseScale(uint16 PoseIndex) const
+{
+	return ReaderPtr->getRBFPoseScale(PoseIndex);
+}
+
+template <class TWrappedReader>
+uint16 FDNAReader<TWrappedReader>::GetRBFSolverCount() const
+{
+	return ReaderPtr->getRBFSolverCount();
+}
+
+template <class TWrappedReader>
+uint16 FDNAReader<TWrappedReader>::GetRBFSolverIndexListCount() const
+{
+	return ReaderPtr->getRBFSolverIndexListCount();
+}
+
+template <class TWrappedReader>
+TArrayView<const uint16> FDNAReader<TWrappedReader>::GetRBFSolverIndicesForLOD(uint16 LOD) const
+{
+	const auto Indices = ReaderPtr->getRBFSolverIndicesForLOD(LOD);
+	return TArrayView<const uint16>(Indices.data(), static_cast<int32>(Indices.size()));
+}
+
+template <class TWrappedReader>
+FString FDNAReader<TWrappedReader>::GetRBFSolverName(uint16 SolverIndex) const
+{
+	return FString(ANSI_TO_TCHAR(ReaderPtr->getRBFSolverName(SolverIndex).data()));
+}
+
+template <class TWrappedReader>
+TArrayView<const uint16> FDNAReader<TWrappedReader>::GetRBFSolverRawControlIndices(uint16 SolverIndex) const
+{
+	const auto Indices = ReaderPtr->getRBFSolverRawControlIndices(SolverIndex);
+	return TArrayView<const uint16>(Indices.data(), static_cast<int32>(Indices.size()));
+}
+
+template <class TWrappedReader>
+TArrayView<const uint16> FDNAReader<TWrappedReader>::GetRBFSolverPoseIndices(uint16 SolverIndex) const
+{
+	const auto Indices = ReaderPtr->getRBFSolverPoseIndices(SolverIndex);
+	return TArrayView<const uint16>(Indices.data(), static_cast<int32>(Indices.size()));
+}
+
+template <class TWrappedReader>
+TArrayView<const float> FDNAReader<TWrappedReader>::GetRBFSolverRawControlValues(uint16 SolverIndex) const
+{
+	const auto Values = ReaderPtr->getRBFSolverRawControlValues(SolverIndex);
+	return TArrayView<const float>(Values.data(), static_cast<int32>(Values.size()));
+}
+
+template <class TWrappedReader>
+ERBFSolverType FDNAReader<TWrappedReader>::GetRBFSolverType(uint16 SolverIndex) const
+{
+	return static_cast<ERBFSolverType>(ReaderPtr->getRBFSolverType(SolverIndex));
+}
+
+template <class TWrappedReader>
+float FDNAReader<TWrappedReader>::GetRBFSolverRadius(uint16 SolverIndex) const
+{
+	return ReaderPtr->getRBFSolverRadius(SolverIndex);
+}
+
+template <class TWrappedReader>
+EAutomaticRadius FDNAReader<TWrappedReader>::GetRBFSolverAutomaticRadius(uint16 SolverIndex) const
+{
+	return static_cast<EAutomaticRadius>(ReaderPtr->getRBFSolverAutomaticRadius(SolverIndex));
+}
+
+template <class TWrappedReader>
+float FDNAReader<TWrappedReader>::GetRBFSolverWeightThreshold(uint16 SolverIndex) const
+{
+	return ReaderPtr->getRBFSolverWeightThreshold(SolverIndex);
+}
+
+template <class TWrappedReader>
+ERBFDistanceMethod FDNAReader<TWrappedReader>::GetRBFSolverDistanceMethod(uint16 SolverIndex) const
+{
+	return static_cast<ERBFDistanceMethod>(ReaderPtr->getRBFSolverDistanceMethod(SolverIndex));
+}
+
+template <class TWrappedReader>
+ERBFNormalizeMethod FDNAReader<TWrappedReader>::GetRBFSolverNormalizeMethod(uint16 SolverIndex) const
+{
+	return static_cast<ERBFNormalizeMethod>(ReaderPtr->getRBFSolverNormalizeMethod(SolverIndex));
+}
+
+template <class TWrappedReader>
+ERBFFunctionType FDNAReader<TWrappedReader>::GetRBFSolverFunctionType(uint16 SolverIndex) const
+{
+	return static_cast<ERBFFunctionType>(ReaderPtr->getRBFSolverFunctionType(SolverIndex));
+}
+
+template <class TWrappedReader>
+ETwistAxis FDNAReader<TWrappedReader>::GetRBFSolverTwistAxis(uint16 SolverIndex) const
+{
+	return static_cast<ETwistAxis>(ReaderPtr->getRBFSolverTwistAxis(SolverIndex));
+}
+
+template <class TWrappedReader>
+uint16 FDNAReader<TWrappedReader>::GetTwistCount() const
+{
+	return ReaderPtr->getTwistCount();
+}
+
+template <class TWrappedReader>
+ETwistAxis FDNAReader<TWrappedReader>::GetTwistSetupTwistAxis(uint16 TwistIndex) const
+{
+	return static_cast<ETwistAxis>(ReaderPtr->getTwistSetupTwistAxis(TwistIndex));
+}
+
+template <class TWrappedReader>
+TArrayView<const uint16> FDNAReader<TWrappedReader>::GetTwistInputControlIndices(uint16 TwistIndex) const
+{
+	const auto Indices = ReaderPtr->getTwistInputControlIndices(TwistIndex);
+	return TArrayView<const uint16>(Indices.data(), static_cast<int32>(Indices.size()));
+}
+
+template <class TWrappedReader>
+TArrayView<const uint16> FDNAReader<TWrappedReader>::GetTwistOutputJointIndices(uint16 TwistIndex) const
+{
+	const auto Indices = ReaderPtr->getTwistOutputJointIndices(TwistIndex);
+	return TArrayView<const uint16>(Indices.data(), static_cast<int32>(Indices.size()));
+}
+
+template <class TWrappedReader>
+TArrayView<const float> FDNAReader<TWrappedReader>::GetTwistBlendWeights(uint16 TwistIndex) const
+{
+	const auto Values = ReaderPtr->getTwistBlendWeights(TwistIndex);
+	return TArrayView<const float>(Values.data(), static_cast<int32>(Values.size()));
+}
+
+template <class TWrappedReader>
+uint16 FDNAReader<TWrappedReader>::GetSwingCount() const
+{
+	return ReaderPtr->getSwingCount();
+}
+
+template <class TWrappedReader>
+ETwistAxis FDNAReader<TWrappedReader>::GetSwingSetupTwistAxis(uint16 SwingIndex) const
+{
+	return static_cast<ETwistAxis>(ReaderPtr->getSwingSetupTwistAxis(SwingIndex));
+
+}
+
+template <class TWrappedReader>
+TArrayView<const uint16> FDNAReader<TWrappedReader>::GetSwingInputControlIndices(uint16 SwingIndex) const
+{
+	const auto Indices = ReaderPtr->getSwingInputControlIndices(SwingIndex);
+	return TArrayView<const uint16>(Indices.data(), static_cast<int32>(Indices.size()));
+}
+
+template <class TWrappedReader>
+TArrayView<const uint16> FDNAReader<TWrappedReader>::GetSwingOutputJointIndices(uint16 SwingIndex) const
+{
+	const auto Indices = ReaderPtr->getSwingOutputJointIndices(SwingIndex);
+	return TArrayView<const uint16>(Indices.data(), static_cast<int32>(Indices.size()));
+}
+
+template <class TWrappedReader>
+TArrayView<const float> FDNAReader<TWrappedReader>::GetSwingBlendWeights(uint16 SwingIndex) const
+{
+	const auto Values = ReaderPtr->getSwingBlendWeights(SwingIndex);
+	return TArrayView<const float>(Values.data(), static_cast<int32>(Values.size()));
+}
diff --git a/Engine/Plugins/Animation/RigLogic/Source/RigLogicModule/Public/DNAUtils.h b/Engine/Plugins/Animation/RigLogic/Source/RigLogicModule/Public/DNAUtils.h
--- a/Engine/Plugins/Animation/RigLogic/Source/RigLogicModule/Public/DNAUtils.h
+++ b/Engine/Plugins/Animation/RigLogic/Source/RigLogicModule/Public/DNAUtils.h
@@ -8,9 +8,17 @@
 #include "UObject/NoExportTypes.h"
 #include "UObject/ObjectMacros.h"
 
+namespace dna
+{
+
+enum class DataLayer : uint32;
+
+}  // namespace dna
+
 RIGLOGICMODULE_API TSharedPtr<IDNAReader> ReadDNAFromFile(const FString& Path, EDNADataLayer Layer = EDNADataLayer::All, uint16_t MaxLOD = 0);
 RIGLOGICMODULE_API TSharedPtr<IDNAReader> ReadDNAFromFile(const FString& Path, EDNADataLayer Layer, TArrayView<uint16_t> LODs);
 RIGLOGICMODULE_API TSharedPtr<IDNAReader> ReadDNAFromBuffer(TArray<uint8>* DNABuffer, EDNADataLayer Layer = EDNADataLayer::All, uint16_t MaxLOD = 0);
 RIGLOGICMODULE_API TSharedPtr<IDNAReader> ReadDNAFromBuffer(TArray<uint8>* DNABuffer, EDNADataLayer Layer, TArrayView<uint16_t> LODs);
 RIGLOGICMODULE_API TArray<uint8> ReadStreamFromDNA(const IDNAReader* Reader, EDNADataLayer Layer);
-RIGLOGICMODULE_API void WriteDNAToFile(const IDNAReader* Reader, EDNADataLayer Layer, const FString& Path);
\ No newline at end of file
+RIGLOGICMODULE_API void WriteDNAToFile(const IDNAReader* Reader, EDNADataLayer Layer, const FString& Path);
+RIGLOGICMODULE_API dna::DataLayer CalculateDNADataLayerBitmask(EDNADataLayer Layer);
diff --git a/Engine/Plugins/Animation/RigLogic/Source/RigLogicModule/Public/RigInstance.h b/Engine/Plugins/Animation/RigLogic/Source/RigLogicModule/Public/RigInstance.h
--- a/Engine/Plugins/Animation/RigLogic/Source/RigLogicModule/Public/RigInstance.h
+++ b/Engine/Plugins/Animation/RigLogic/Source/RigLogicModule/Public/RigInstance.h
@@ -9,7 +9,6 @@
 #include "UObject/ObjectMacros.h"
 
 #include "FMemoryResource.h"
-#include "TransformArrayView.h"
 
 class FRigLogic;
 
@@ -23,7 +22,7 @@ class RigInstance;
 class RIGLOGICMODULE_API FRigInstance
 {
 public:
-	FRigInstance(FRigLogic* RigLogic);
+	explicit FRigInstance(FRigLogic* RigLogic);
 	~FRigInstance();
 
 	FRigInstance(const FRigInstance&) = delete;
@@ -56,8 +55,11 @@ public:
 	float GetNeuralNetworkMask(uint16 NeuralNetIndex) const;
 	void SetNeuralNetworkMask(uint16 NeuralNetIndex, float Value);
 
-	TArrayView<const float> GetRawJointOutputs() const;
-	FTransformArrayView GetJointOutputs() const;
+	uint16 GetRBFControlCount() const;
+	float GetRBFControl(uint16 Index) const;
+	TArrayView<const float> GetRBFControlValues() const;
+
+	TArrayView<const float> GetJointOutputs() const;
 	TArrayView<const float> GetBlendShapeOutputs() const;
 	TArrayView<const float> GetAnimatedMapOutputs() const;
 	uint16 GetLOD() const;
diff --git a/Engine/Plugins/Animation/RigLogic/Source/RigLogicModule/Public/RigLogic.h b/Engine/Plugins/Animation/RigLogic/Source/RigLogicModule/Public/RigLogic.h
--- a/Engine/Plugins/Animation/RigLogic/Source/RigLogicModule/Public/RigLogic.h
+++ b/Engine/Plugins/Animation/RigLogic/Source/RigLogicModule/Public/RigLogic.h
@@ -9,7 +9,6 @@
 #include "UObject/ObjectMacros.h"
 
 #include "FMemoryResource.h"
-#include "TransformArrayView.h"
 
 #include "RigLogic.generated.h"
 
@@ -28,13 +27,124 @@ enum class ERigLogicCalculationType: uint8
 {
 	Scalar,
 	SSE,
-	AVX
+	AVX,
+	NEON,
+	AnyVector
+};
+
+UENUM(BlueprintType)
+enum class ERigLogicTranslationType : uint8 {
+	None,
+	Vector = 3
+};
+
+UENUM(BlueprintType)
+enum class ERigLogicRotationType : uint8 {
+	None,
+	EulerAngles = 3,
+	Quaternions = 4
+};
+
+UENUM(BlueprintType)
+enum class ERigLogicRotationOrder : uint8 {
+	XYZ,
+	XZY,
+	YXZ,
+	YZX,
+	ZXY,
+	ZYX
+};
+
+UENUM(BlueprintType)
+enum class ERigLogicScaleType : uint8 {
+	None,
+	Vector = 3
+};
+
+USTRUCT(BlueprintType)
+struct FRigLogicConfiguration
+{
+	GENERATED_BODY()
+
+	FRigLogicConfiguration() :
+		CalculationType(ERigLogicCalculationType::AnyVector),
+		LoadJoints(true),
+		LoadBlendShapes(true),
+		LoadAnimatedMaps(true),
+		LoadMachineLearnedBehavior(true),
+		LoadRBFBehavior(true),
+		LoadTwistSwingBehavior(true),
+		TranslationType(ERigLogicTranslationType::Vector),
+		RotationType(ERigLogicRotationType::Quaternions),
+		RotationOrder(ERigLogicRotationOrder::ZYX),
+		ScaleType(ERigLogicScaleType::Vector)
+	{
+	}
+
+	FRigLogicConfiguration(ERigLogicCalculationType CalculationType,
+							bool LoadJoints,
+							bool LoadBlendShapes,
+							bool LoadAnimatedMaps,
+							bool LoadMachineLearnedBehavior,
+							bool LoadRBFBehavior,
+							bool LoadTwistSwingBehavior,
+							ERigLogicTranslationType TranslationType,
+							ERigLogicRotationType RotationType,
+							ERigLogicRotationOrder RotationOrder,
+							ERigLogicScaleType ScaleType) :
+		CalculationType(CalculationType),
+		LoadJoints(LoadJoints),
+		LoadBlendShapes(LoadBlendShapes),
+		LoadAnimatedMaps(LoadAnimatedMaps),
+		LoadMachineLearnedBehavior(LoadMachineLearnedBehavior),
+		LoadRBFBehavior(LoadRBFBehavior),
+		LoadTwistSwingBehavior(LoadTwistSwingBehavior),
+		TranslationType(TranslationType),
+		RotationType(RotationType),
+		RotationOrder(RotationOrder),
+		ScaleType(ScaleType
+	)
+	{
+	}
+
+	UPROPERTY(BlueprintReadOnly, Category = "RigLogic")
+	ERigLogicCalculationType CalculationType;
+
+	UPROPERTY(BlueprintReadOnly, Category = "RigLogic")
+	bool LoadJoints;
+
+	UPROPERTY(BlueprintReadOnly, Category = "RigLogic")
+	bool LoadBlendShapes;
+
+	UPROPERTY(BlueprintReadOnly, Category = "RigLogic")
+	bool LoadAnimatedMaps;
+
+	UPROPERTY(BlueprintReadOnly, Category = "RigLogic")
+	bool LoadMachineLearnedBehavior;
+	
+	UPROPERTY(BlueprintReadOnly, Category = "RigLogic")
+	bool LoadRBFBehavior;
+	
+	UPROPERTY(BlueprintReadOnly, Category = "RigLogic")
+	bool LoadTwistSwingBehavior;
+	
+	UPROPERTY(BlueprintReadOnly, Category = "RigLogic")
+	ERigLogicTranslationType TranslationType;
+	
+	UPROPERTY(BlueprintReadOnly, Category = "RigLogic")
+	ERigLogicRotationType RotationType;
+	
+	UPROPERTY(BlueprintReadOnly, Category = "RigLogic")
+	ERigLogicRotationOrder RotationOrder;
+
+	UPROPERTY(BlueprintReadOnly, Category = "RigLogic")
+	ERigLogicScaleType ScaleType;
 };
 
 class RIGLOGICMODULE_API FRigLogic
 {
 public:
-	FRigLogic(const IDNAReader* Reader, ERigLogicCalculationType CalculationType = ERigLogicCalculationType::SSE);
+	explicit FRigLogic(const IDNAReader* Reader, FRigLogicConfiguration Config = FRigLogicConfiguration());
 	~FRigLogic();
 
 	FRigLogic(const FRigLogic&) = delete;
@@ -44,11 +154,11 @@ public:
 	FRigLogic& operator=(FRigLogic&&) = default;
 
 	uint16 GetLODCount() const;
-	TArrayView<const float> GetRawNeutralJointValues() const;
-	FTransformArrayView GetNeutralJointValues() const;
+	TArrayView<const float> GetNeutralJointValues() const;
 	TArrayView<const uint16> GetJointVariableAttributeIndices(uint16 LOD) const;
 	uint16 GetJointGroupCount() const;
 	uint16 GetNeuralNetworkCount() const;
+	uint16 GetRBFSolverCount() const;
 	uint16 GetMeshCount() const;
 	uint16 GetMeshRegionCount(uint16 MeshIndex) const;
 	TArrayView<const uint16> GetNeuralNetworkIndices(uint16 MeshIndex, uint16 RegionIndex) const;
@@ -58,6 +168,8 @@ public:
 	void CalculateControls(FRigInstance* Instance) const;
 	void CalculateMachineLearnedBehaviorControls(FRigInstance* Instance) const;
 	void CalculateMachineLearnedBehaviorControls(FRigInstance* Instance, uint16 NeuralNetIndex) const;
+	void CalculateRBFControls(FRigInstance* Instance) const;
+	void CalculateRBFControls(FRigInstance* Instance, uint16 SolverIndex) const;
 	void CalculateJoints(FRigInstance* Instance) const;
 	void CalculateJoints(FRigInstance* Instance, uint16 JointGroupIndex) const;
 	void CalculateBlendShapes(FRigInstance* Instance) const;
diff --git a/Engine/Plugins/Animation/RigLogic/Source/RigLogicModule/Public/RigUnit_RigLogic.h b/Engine/Plugins/Animation/RigLogic/Source/RigLogicModule/Public/RigUnit_RigLogic.h
--- a/Engine/Plugins/Animation/RigLogic/Source/RigLogicModule/Public/RigUnit_RigLogic.h
+++ b/Engine/Plugins/Animation/RigLogic/Source/RigLogicModule/Public/RigUnit_RigLogic.h
@@ -12,20 +12,33 @@
 #include "RigUnit_RigLogic.generated.h"
 
 class IBehaviorReader;
-class FTransformArrayView;
 struct FSharedRigRuntimeContext;
 
 DECLARE_LOG_CATEGORY_EXTERN(LogRigLogicUnit, Log, All);
 
+/* A helper struct used inside of the RigUnit_RigLogic to map bone indices to DNA driver joint indices. */
+USTRUCT(meta = (DocumentationPolicy = "Strict"))
+struct FBoneIndexControlAttributeMapping
+{
+	GENERATED_BODY()
+
+	int32 BoneIndex;
+	int32 DNAJointIndex;
+	int32 RotationX;
+	int32 RotationY;
+	int32 RotationZ;
+	int32 RotationW;
+};
+
 /* A helper struct used inside of the RigUnit_RigLogic to store arrays of arrays of integers. */
 USTRUCT(meta = (DocumentationPolicy = "Strict"))
 struct FRigUnit_RigLogic_IntArray
 {
 	GENERATED_BODY()
 
-		// The values stored within this array.
-		UPROPERTY(transient)
-		TArray<int32> Values;
+	// The values stored within this array.
+	UPROPERTY(transient)
+	TArray<int32> Values;
 };
 
 /* The work data used by the FRigUnit_RigLogic */
@@ -68,6 +81,10 @@ struct FRigUnit_RigLogic_Data
 	UPROPERTY(transient)
 	TArray<int32> HierarchyBoneIndices;
 
+	/** RL driver joint index to ControlRig's hierarchy bone index and RigLogic control attribute mapping **/
+	UPROPERTY(transient)
+	TArray<FBoneIndexControlAttributeMapping> DriverJointsToControlAttributesMap;
+
 	/** RL mesh blend shape index to ControlRig's output blendshape curve index for each LOD **/
 	UPROPERTY(transient)
 	TArray<FRigUnit_RigLogic_IntArray> MorphTargetCurveIndices;
@@ -107,7 +124,7 @@ struct FRigUnit_RigLogic_Data
 	void MapMaskMultipliers(const URigHierarchy* InHierarchy);
 
 	/** Calculates joint positions, orientation and scale based on inputs curves of the control rig **/
-	void CalculateRigLogic(const URigHierarchy* InHierarchy);
+	void CalculateRigLogic(const URigHierarchy* InHierarchy, TArrayView<const float> NeutralJointValues);
 	/** Updates joint positions in the hierarchy based on inputs curves of the control rig **/
 	void UpdateJoints(URigHierarchy* Hierarchy, TArrayView<const float> NeutralJointValues, TArrayView<const float> DeltaJointValues);
 	/** Updates morph target curve values based on values of input curves of the control rig **/
