Update to new RigLogic implementing body rig support

diff --git a/Engine/Plugins/Animation/RigLogic/Source/RigLogicModule/Public/RigUnit_RigLogic_Test.h b/Engine/Plugins/Animation/RigLogic/Source/RigLogicModule/Public/RigUnit_RigLogic_Test.h
--- a/Engine/Plugins/Animation/RigLogic/Source/RigLogicModule/Public/RigUnit_RigLogic_Test.h
+++ b/Engine/Plugins/Animation/RigLogic/Source/RigLogicModule/Public/RigUnit_RigLogic_Test.h
@@ -620,6 +620,184 @@ public:
 		return {};
 	}
 
+	// IJointBehaviorMetadataReader methods
+	ETranslationRepresentation GetJointTranslationRepresentation(uint16 JointIndex) const override
+	{
+		return {};
+	}
+
+	ERotationRepresentation GetJointRotationRepresentation(uint16 JointIndex) const override
+	{
+		return {};
+	}
+
+	EScaleRepresentation GetJointScaleRepresentation(uint16 JointIndex) const override
+	{
+		return {};
+	}
+
+	// IRBFBehaviorReader methods
+	uint16 GetRBFPoseCount() const override
+	{
+		return {};
+	}
+
+	FString GetRBFPoseName(uint16 PoseIndex) const override
+	{
+		return {};
+	}
+
+	TArrayView<const uint16> GetRBFPoseJointOutputIndices(uint16 PoseIndex) const override
+	{
+		return {};
+	}
+
+	TArrayView<const uint16> GetRBFPoseBlendShapeChannelOutputIndices(uint16 PoseIndex) const override
+	{
+		return {};
+	}
+
+	TArrayView<const uint16> GetRBFPoseAnimatedMapOutputIndices(uint16 PoseIndex) const override
+	{
+		return {};
+	}
+
+	TArrayView<const float> GetRBFPoseJointOutputValues(uint16 PoseIndex) const override
+	{
+		return {};
+	}
+
+	float GetRBFPoseScale(uint16 PoseIndex) const override
+	{
+		return {};
+	}
+
+	uint16 GetRBFSolverCount() const override
+	{
+		return {};
+	}
+
+	uint16 GetRBFSolverIndexListCount() const override
+	{
+		return {};
+	}
+
+	TArrayView<const uint16> GetRBFSolverIndicesForLOD(uint16 LOD) const override
+	{
+		return {};
+	}
+
+	FString GetRBFSolverName(uint16 SolverIndex) const override
+	{
+		return {};
+	}
+
+	TArrayView<const uint16> GetRBFSolverRawControlIndices(uint16 SolverIndex) const override
+	{
+		return {};
+	}
+
+	TArrayView<const uint16> GetRBFSolverPoseIndices(uint16 SolverIndex) const override
+	{
+		return {};
+	}
+
+	TArrayView<const float> GetRBFSolverRawControlValues(uint16 SolverIndex) const override
+	{
+		return {};
+	}
+
+	ERBFSolverType GetRBFSolverType(uint16 SolverIndex) const override
+	{
+		return {};
+	}
+
+	float GetRBFSolverRadius(uint16 SolverIndex) const override
+	{
+		return {};
+	}
+
+	EAutomaticRadius GetRBFSolverAutomaticRadius(uint16 SolverIndex) const override
+	{
+		return {};
+	}
+
+	float GetRBFSolverWeightThreshold(uint16 SolverIndex) const override
+	{
+		return {};
+	}
+
+	ERBFDistanceMethod GetRBFSolverDistanceMethod(uint16 SolverIndex) const override
+	{
+		return {};
+	}
+
+	ERBFNormalizeMethod GetRBFSolverNormalizeMethod(uint16 SolverIndex) const override
+	{
+		return {};
+	}
+
+	ERBFFunctionType GetRBFSolverFunctionType(uint16 SolverIndex) const override
+	{
+		return {};
+	}
+
+	ETwistAxis GetRBFSolverTwistAxis(uint16 SolverIndex) const override
+	{
+		return {};
+	}
+
+	// ITwistSwingBehaviorReader methods
+	uint16 GetTwistCount() const override
+	{
+		return {};
+	}
+
+	ETwistAxis GetTwistSetupTwistAxis(uint16 TwistIndex) const override
+	{
+		return {};
+	}
+
+	TArrayView<const uint16> GetTwistInputControlIndices(uint16 TwistIndex) const override
+	{
+		return {};
+	}
+
+	TArrayView<const uint16> GetTwistOutputJointIndices(uint16 TwistIndex) const override
+	{
+		return {};
+	}
+
+	TArrayView<const float> GetTwistBlendWeights(uint16 TwistIndex) const override
+	{
+		return {};
+	}
+
+	uint16 GetSwingCount() const override
+	{
+		return {};
+	}
+
+	ETwistAxis GetSwingSetupTwistAxis(uint16 SwingIndex) const override
+	{
+		return {};
+	}
+
+	TArrayView<const uint16> GetSwingInputControlIndices(uint16 SwingIndex) const override
+	{
+		return {};
+	}
+
+	TArrayView<const uint16> GetSwingOutputJointIndices(uint16 SwingIndex) const override
+	{
+		return {};
+	}
+
+	TArrayView<const float> GetSwingBlendWeights(uint16 SwingIndex) const override
+	{
+		return {};
+	}
+
 	uint16 LODCount;
 
 	TArray<FString> rawControls;
@@ -742,7 +920,7 @@ struct FRigUnit_RigLogic::TestAccessor
 
 	/** UpdateJoints Tests **/
 	void AddToTransformArray(float* InArray, FTransform& Transform);
-	FTransformArrayView CreateTwoJointNeutralTransforms(float *InValueArray);
+	TArrayView<const float> CreateTwoJointNeutralTransforms(float *InValueArray);
 	TArrayView<const uint16> CreateTwoJointVariableAttributes(uint16* InVariableAttributeIndices, uint8 LOD);
 
 	void Exec_UpdateJoints(URigHierarchy* TestHierarchy, TArrayView<const float> NeutralJointValues, TArrayView<const float> DeltaJointValues);
diff --git a/Engine/Plugins/Animation/RigLogic/Source/RigLogicModule/Public/SkelMeshDNAReader.h b/Engine/Plugins/Animation/RigLogic/Source/RigLogicModule/Public/SkelMeshDNAReader.h
--- a/Engine/Plugins/Animation/RigLogic/Source/RigLogicModule/Public/SkelMeshDNAReader.h
+++ b/Engine/Plugins/Animation/RigLogic/Source/RigLogicModule/Public/SkelMeshDNAReader.h
@@ -137,6 +137,44 @@ public:
 	TArrayView<const float> GetNeuralNetworkLayerActivationFunctionParameters(uint16 NetIndex, uint16 LayerIndex) const override;
 	TArrayView<const float> GetNeuralNetworkLayerBiases(uint16 NetIndex, uint16 LayerIndex) const override;
 	TArrayView<const float> GetNeuralNetworkLayerWeights(uint16 NetIndex, uint16 LayerIndex) const override;
+	// JointBehaviorMetadataReader
+	ETranslationRepresentation GetJointTranslationRepresentation(uint16 JointIndex) const override;
+	ERotationRepresentation GetJointRotationRepresentation(uint16 JointIndex) const override;
+	EScaleRepresentation GetJointScaleRepresentation(uint16 JointIndex) const override;
+	// RBFBehavior
+	uint16 GetRBFPoseCount() const override;
+	FString GetRBFPoseName(uint16 PoseIndex) const override;
+	TArrayView<const uint16> GetRBFPoseJointOutputIndices(uint16 PoseIndex) const override;
+	TArrayView<const uint16> GetRBFPoseBlendShapeChannelOutputIndices(uint16 PoseIndex) const override;
+	TArrayView<const uint16> GetRBFPoseAnimatedMapOutputIndices(uint16 PoseIndex) const override;
+	TArrayView<const float> GetRBFPoseJointOutputValues(uint16 PoseIndex) const override;
+	float GetRBFPoseScale(uint16 PoseIndex) const override;
+	uint16 GetRBFSolverCount() const override;
+	uint16 GetRBFSolverIndexListCount() const override;
+	TArrayView<const uint16> GetRBFSolverIndicesForLOD(uint16 LOD) const override;
+	FString GetRBFSolverName(uint16 SolverIndex) const override;
+	TArrayView<const uint16> GetRBFSolverRawControlIndices(uint16 SolverIndex) const override;
+	TArrayView<const uint16> GetRBFSolverPoseIndices(uint16 SolverIndex) const override;
+	TArrayView<const float> GetRBFSolverRawControlValues(uint16 SolverIndex) const override;
+	ERBFSolverType GetRBFSolverType(uint16 SolverIndex) const override;
+	float GetRBFSolverRadius(uint16 SolverIndex) const override;
+	EAutomaticRadius GetRBFSolverAutomaticRadius(uint16 SolverIndex) const override;
+	float GetRBFSolverWeightThreshold(uint16 SolverIndex) const override;
+	ERBFDistanceMethod GetRBFSolverDistanceMethod(uint16 SolverIndex) const override;
+	ERBFNormalizeMethod GetRBFSolverNormalizeMethod(uint16 SolverIndex) const override;
+	ERBFFunctionType GetRBFSolverFunctionType(uint16 SolverIndex) const override;
+	ETwistAxis GetRBFSolverTwistAxis(uint16 SolverIndex) const override;
+	// TwistSwingBehavior
+	uint16 GetTwistCount() const override;
+	ETwistAxis GetTwistSetupTwistAxis(uint16 TwistIndex) const override;
+	TArrayView<const uint16> GetTwistInputControlIndices(uint16 TwistIndex) const override;
+	TArrayView<const uint16> GetTwistOutputJointIndices(uint16 TwistIndex) const override;
+	TArrayView<const float> GetTwistBlendWeights(uint16 TwistIndex) const override;
+	uint16 GetSwingCount() const override;
+	ETwistAxis GetSwingSetupTwistAxis(uint16 SwingIndex) const override;
+	TArrayView<const uint16> GetSwingInputControlIndices(uint16 SwingIndex) const override;
+	TArrayView<const uint16> GetSwingOutputJointIndices(uint16 SwingIndex) const override;
+	TArrayView<const float> GetSwingBlendWeights(uint16 SwingIndex) const override;
 
 	void Unload(EDNADataLayer /*unused*/) override;
 
diff --git a/Engine/Plugins/Animation/RigLogic/Source/RigLogicModule/RigLogicModule.Build.cs b/Engine/Plugins/Animation/RigLogic/Source/RigLogicModule/RigLogicModule.Build.cs
--- a/Engine/Plugins/Animation/RigLogic/Source/RigLogicModule/RigLogicModule.Build.cs
+++ b/Engine/Plugins/Animation/RigLogic/Source/RigLogicModule/RigLogicModule.Build.cs
@@ -33,7 +33,8 @@ namespace UnrealBuildTool.Rules
 			PrivateDependencyModuleNames.AddRange(
 				new string[]
 				{
-					"AnimationCore"
+					"AnimationCore",
+					"AnimGraphRuntime"
 				}
 			);
 		}
