Fix build error caused by merge conflict

diff --git a/Engine/Plugins/Animation/RigLogic/Source/RigLogicModule/Public/DNAAsset.h b/Engine/Plugins/Animation/RigLogic/Source/RigLogicModule/Public/DNAAsset.h
--- a/Engine/Plugins/Animation/RigLogic/Source/RigLogicModule/Public/DNAAsset.h
+++ b/Engine/Plugins/Animation/RigLogic/Source/RigLogicModule/Public/DNAAsset.h
@@ -94,10 +94,6 @@ public:
 	TSharedPtr<FSharedRigRuntimeContext> GetRigRuntimeContext();
 	TSharedPtr<FDNAIndexMapping> GetDNAIndexMapping(const USkeleton* Skeleton, const USkeletalMesh* SkeletalMesh);
 
-	TSharedPtr<FSharedRigRuntimeContext> GetRigRuntimeContext();
-	TSharedPtr<FDNAIndexMapping> GetDNAIndexMapping(const USkeleton* Skeleton,
-		const USkeletalMesh* SkeletalMesh);
-
 private:
 	friend struct FAnimNode_RigLogic;
 	friend struct FRigUnit_RigLogic;
