Fix missing initializer for joint output buffer

diff --git a/Engine/Plugins/Animation/RigLogic/RigLogic.uplugin b/Engine/Plugins/Animation/RigLogic/RigLogic.uplugin
--- a/Engine/Plugins/Animation/RigLogic/RigLogic.uplugin
+++ b/Engine/Plugins/Animation/RigLogic/RigLogic.uplugin
@@ -1,9 +1,9 @@
 {
 	"FileVersion": 3,
 	"Version": 10.1,
-	"VersionName": "10.1.2",
-	"FriendlyName": "RigLogic Plugin v10.1.2",
-	"Description": "3Lateral RigLogic Plugin for Facial Animation v10.1.2",
+	"VersionName": "10.1.4",
+	"FriendlyName": "RigLogic Plugin v10.1.4",
+	"Description": "3Lateral RigLogic Plugin for Facial Animation v10.1.4",
 	"Category": "Animation",
 	"CreatedBy": "3Lateral",
 	"CreatedByURL": "https://www.3lateral.com/",
diff --git a/Engine/Plugins/Animation/RigLogic/Source/RigLogicLib/Private/riglogic/joints/cpu/CPUJointsOutputInstance.cpp b/Engine/Plugins/Animation/RigLogic/Source/RigLogicLib/Private/riglogic/joints/cpu/CPUJointsOutputInstance.cpp
--- a/Engine/Plugins/Animation/RigLogic/Source/RigLogicLib/Private/riglogic/joints/cpu/CPUJointsOutputInstance.cpp
+++ b/Engine/Plugins/Animation/RigLogic/Source/RigLogicLib/Private/riglogic/joints/cpu/CPUJointsOutputInstance.cpp
@@ -41,6 +41,8 @@ CPUJointsOutputInstance::CPUJointsOutputInstance(std::uint16_t jointAttributeCou
                 std::fill(buffer.begin(), buffer.end(), 0.0f);
             };
     }
+
+    resetBuffer(outputBuffer);
 }
 
 ArrayView<float> CPUJointsOutputInstance::getOutputBuffer() {
diff --git a/Engine/Plugins/Animation/RigLogic/Source/RigLogicLib/Public/dna/version/Version.h b/Engine/Plugins/Animation/RigLogic/Source/RigLogicLib/Public/dna/version/Version.h
--- a/Engine/Plugins/Animation/RigLogic/Source/RigLogicLib/Public/dna/version/Version.h
+++ b/Engine/Plugins/Animation/RigLogic/Source/RigLogicLib/Public/dna/version/Version.h
@@ -4,5 +4,5 @@
 
 #define DNA_MAJOR_VERSION 9
 #define DNA_MINOR_VERSION 3
-#define DNA_PATCH_VERSION 5
-#define DNA_VERSION_STRING "9.3.5"
+#define DNA_PATCH_VERSION 6
+#define DNA_VERSION_STRING "9.3.6"
diff --git a/Engine/Plugins/Animation/RigLogic/Source/RigLogicLib/Public/riglogic/version/Version.h b/Engine/Plugins/Animation/RigLogic/Source/RigLogicLib/Public/riglogic/version/Version.h
--- a/Engine/Plugins/Animation/RigLogic/Source/RigLogicLib/Public/riglogic/version/Version.h
+++ b/Engine/Plugins/Animation/RigLogic/Source/RigLogicLib/Public/riglogic/version/Version.h
@@ -4,5 +4,5 @@
 
 #define RL_MAJOR_VERSION 10
 #define RL_MINOR_VERSION 1
-#define RL_PATCH_VERSION 2
-#define RL_VERSION_STRING "10.1.2"
+#define RL_PATCH_VERSION 4
+#define RL_VERSION_STRING "10.1.4"
diff --git a/Engine/Plugins/Animation/RigLogic/Source/RigLogicLib/Public/trio/version/Version.h b/Engine/Plugins/Animation/RigLogic/Source/RigLogicLib/Public/trio/version/Version.h
--- a/Engine/Plugins/Animation/RigLogic/Source/RigLogicLib/Public/trio/version/Version.h
+++ b/Engine/Plugins/Animation/RigLogic/Source/RigLogicLib/Public/trio/version/Version.h
@@ -4,5 +4,5 @@
 
 #define TRIO_MAJOR_VERSION 4
 #define TRIO_MINOR_VERSION 0
-#define TRIO_PATCH_VERSION 14
-#define TRIO_VERSION_STRING "4.0.14"
+#define TRIO_PATCH_VERSION 15
+#define TRIO_VERSION_STRING "4.0.15"
diff --git a/Engine/Plugins/Animation/RigLogic/Source/RigLogicLibTest/Private/rltests/riglogic/TestRigLogic.cpp b/Engine/Plugins/Animation/RigLogic/Source/RigLogicLibTest/Private/rltests/riglogic/TestRigLogic.cpp
--- a/Engine/Plugins/Animation/RigLogic/Source/RigLogicLibTest/Private/rltests/riglogic/TestRigLogic.cpp
+++ b/Engine/Plugins/Animation/RigLogic/Source/RigLogicLibTest/Private/rltests/riglogic/TestRigLogic.cpp
@@ -99,3 +99,20 @@ TEST_F(RigLogicTest, DumpStateThenRestore) {
     rl4::RigInstance::destroy(cloneRigInstance);
     rl4::RigLogic::destroy(cloneRigLogic);
 }
+
+TEST_F(RigLogicTest, JointOutputBufferInitialized) {
+    rl4::Configuration config{};
+    config.rotationType = rl4::RotationType::Quaternions;
+    auto qRigLogic = pma::makeScoped<rl4::RigLogic>(reader.get(), config);
+    auto qRigInstance = pma::makeScoped<rl4::RigInstance>(qRigLogic.get());
+    auto jointOutputs = qRigInstance->getJointOutputs();
+    static constexpr std::size_t qwOffset = 6;
+    static constexpr std::size_t jointAttrCount = 10;
+    for (std::size_t i = {}; i < jointOutputs.size(); ++i) {
+        if (i % jointAttrCount == qwOffset) {
+            ASSERT_EQ(jointOutputs[i], 1.0f);
+        } else {
+            ASSERT_EQ(jointOutputs[i], 0.0f);
+        }
+    }
+}
