Fix calculation of joint group size for each LOD in RigLogic

diff --git a/Engine/Plugins/Animation/RigLogic/Source/RigLogicLib/Private/riglogic/joints/cpu/bpcm/BPCMJointsBuilder.h b/Engine/Plugins/Animation/RigLogic/Source/RigLogicLib/Private/riglogic/joints/cpu/bpcm/BPCMJointsBuilder.h
--- a/Engine/Plugins/Animation/RigLogic/Source/RigLogicLib/Private/riglogic/joints/cpu/bpcm/BPCMJointsBuilder.h
+++ b/Engine/Plugins/Animation/RigLogic/Source/RigLogicLib/Private/riglogic/joints/cpu/bpcm/BPCMJointsBuilder.h
@@ -205,9 +205,9 @@ void BPCMJointsBuilder<TValue, TFVec>::setOutputRotationIndices(const JointBehav
         #endif
         auto deduplicate = [this](Vector<std::uint16_t>& v) {
                 UnorderedSet<std::uint16_t> deduplicator{memRes};
-                v.erase(std::remove_if(v.begin(), v.end(), [&deduplicator](const std::uint16_t value) {
+                v.erase(v.rend().base(), std::remove_if(v.rbegin(), v.rend(), [&deduplicator](const std::uint16_t value) {
                         return !deduplicator.insert(value).second;
-                    }), v.end());
+                    }).base());
             };
         #if !defined(__clang__) && defined(__GNUC__)
             #pragma GCC diagnostic pop
diff --git a/Engine/Plugins/Animation/RigLogic/Source/RigLogicLibTest/Private/rltests/joints/bpcm/BPCMFixturesBlock4.cpp b/Engine/Plugins/Animation/RigLogic/Source/RigLogicLibTest/Private/rltests/joints/bpcm/BPCMFixturesBlock4.cpp
--- a/Engine/Plugins/Animation/RigLogic/Source/RigLogicLibTest/Private/rltests/joints/bpcm/BPCMFixturesBlock4.cpp
+++ b/Engine/Plugins/Animation/RigLogic/Source/RigLogicLibTest/Private/rltests/joints/bpcm/BPCMFixturesBlock4.cpp
@@ -133,6 +133,12 @@ const Matrix<float> values = {
     {  // Joint group 10 - Pad to block-4 => block-4
         1.0f, 1.0f, 1.0f, 1.0f, 1.0f,
         2.0f, 2.0f, 2.0f, 2.0f, 2.0f
+    },
+    {  // Joint group 11 - No padding => block-4
+        1.0f, 1.0f, 1.0f, 1.0f, 1.0f,
+        2.0f, 2.0f, 2.0f, 2.0f, 2.0f,
+        3.0f, 3.0f, 3.0f, 3.0f, 3.0f,
+        4.0f, 4.0f, 4.0f, 4.0f, 4.0f
     }
 };
 
@@ -169,7 +175,10 @@ const Matrix<std::uint16_t> inputIndices = {
     },
     {  // Joint group 10
         7, 8, 9, 10, 11
-    }
+    },
+    {  // Joint group 11
+        1, 2, 3, 4, 5
+    },
 };
 const Matrix<std::uint16_t> outputIndices = {
     {  // Joint group 0
@@ -204,6 +213,9 @@ const Matrix<std::uint16_t> outputIndices = {
     },
     {  // Joint group 10
         135, 136
+    },
+    {  // Joint group 11
+        147, 149, 156, 148
     }
 };
 const Matrix<std::uint16_t> lods = {
@@ -239,6 +251,9 @@ const Matrix<std::uint16_t> lods = {
     },
     {  // Joint group 10
         2, 2, 2, 2
+    },
+    {  // Joint group 11
+        4, 4, 2, 2
     }
 };
 
@@ -469,6 +484,13 @@ const AlignedMatrix<float> floatValues = {
         1.0f, 2.0f, 0.0f, 0.0f,
         1.0f, 2.0f, 0.0f, 0.0f,
         1.0f, 2.0f, 0.0f, 0.0f
+    },
+    {  // Joint group 11
+        1.0f, 2.0f, 3.0f, 4.0f,
+        1.0f, 2.0f, 3.0f, 4.0f,
+        1.0f, 2.0f, 3.0f, 4.0f,
+        1.0f, 2.0f, 3.0f, 4.0f,
+        1.0f, 2.0f, 3.0f, 4.0f
     }
 };
 
@@ -686,7 +708,14 @@ const AlignedMatrix<std::uint16_t> halfFloatValues = {
         15360, 16384, 0, 0,
         15360, 16384, 0, 0,
         15360, 16384, 0, 0,
-    }
+    },
+    {  // Joint group 11
+        15360, 16384, 16896, 17408,
+        15360, 16384, 16896, 17408,
+        15360, 16384, 16896, 17408,
+        15360, 16384, 16896, 17408,
+        15360, 16384, 16896, 17408,
+    },
 };
 
 const AlignedMatrix<std::uint16_t> inputIndices = {
@@ -722,6 +751,9 @@ const AlignedMatrix<std::uint16_t> inputIndices = {
     },
     {  // Joint group 10
         7, 8, 9, 10, 11
+    },
+    {  // Joint group 11
+        1, 2, 3, 4, 5
     }
 };
 
@@ -759,6 +791,9 @@ const Vector<Vector<AlignedVector<std::uint16_t> > > outputIndices = {
         },
         {  // Joint group 10
             150, 151, 0, 0
+        },
+        {  // Joint group 11
+            163, 165, 173, 164
         }
     },
     {  // Euler-angle outputs
@@ -794,6 +829,9 @@ const Vector<Vector<AlignedVector<std::uint16_t> > > outputIndices = {
         },
         {  // Joint group 10
             135, 136, 0, 0
+        },
+        {  // Joint group 11
+            147, 149, 156, 148
         }
     }
 };
@@ -829,6 +867,9 @@ const Vector<Vector<AlignedVector<std::uint16_t> > > outputRotationIndices = {
         {  // Joint group 9
         },
         {  // Joint group 10
+        },
+        {  // Joint group 11
+            173, 163
         }
     },
     {  // Euler-angle outputs
@@ -853,6 +894,8 @@ const Vector<Vector<AlignedVector<std::uint16_t> > > outputRotationIndices = {
         {  // Joint group 9
         },
         {  // Joint group 10
+        },
+        {  // Joint group 11
         }
     }
 };
@@ -891,6 +934,9 @@ const Vector<Vector<AlignedVector<std::uint16_t> > > outputRotationLODs = {
         },
         {  // Joint group 10
             0, 0, 0, 0
+        },
+        {  // Joint group 11
+            2, 2, 2, 2
         }
     },
     {  // Euler-angle outputs
@@ -915,6 +961,8 @@ const Vector<Vector<AlignedVector<std::uint16_t> > > outputRotationLODs = {
         {  // Joint group 9
         },
         {  // Joint group 10
+        },
+        {  // Joint group 11
         }
     }
 };
@@ -955,6 +1003,9 @@ const Vector<Vector<bpcm::JointGroup> > jointGroups = {
         },
         {  // Joint group 10
             1268, 122, 100, 40, 11, 40, 20, 5, 4, 0
+        },
+        {  // Joint group 11
+            1288, 127, 104, 44, 11, 44, 20, 5, 4, 0
         }
     },
     {  // Euler-angle outputs
@@ -990,6 +1041,9 @@ const Vector<Vector<bpcm::JointGroup> > jointGroups = {
         },
         {  // Joint group 10
             1268, 122, 100, 40, 0, 0, 20, 5, 4, 0
+        },
+        {  // Joint group 11
+            1288, 127, 104, 44, 0, 0, 20, 5, 4, 0
         }
     }
 };
@@ -1061,6 +1115,12 @@ const Matrix<bpcm::LODRegion> lodRegions = {
         {2, 0, 0},  // LOD-1
         {2, 0, 0},  // LOD-2
         {2, 0, 0}  // LOD-3
+    },
+    {  // Joint group 11
+        {4, 0, 0},  // LOD-0
+        {4, 0, 0},  // LOD-1
+        {2, 0, 0},  // LOD-2
+        {2, 0, 0}  // LOD-3
     }
 };
 
@@ -1095,7 +1155,9 @@ const Vector<Matrix<float> > valuesPerLOD = {
             4256.0f, 4332.0f, 4408.0f, 0.5042307f, 0.0737294f, -0.821218f, -0.2567421f, 4712.0f, 4788.0f, 4864.0f,
             4940.0f, 5016.0f, 5092.0f, 0.4187725f, -0.2083559f, -0.7320557f, -0.4952896f, 5396.0f, 0.0f, 0.0f,
             15.0f, 30.0f, 0.0f, 0.0f, 0.0f, 0.0f, 1.0f, 0.0f, 0.0f, 0.0f,
-            40.0f, 80.0f
+            40.0f, 80.0f, 0.0f, 0.0f, 0.0f, 0.0f, 1.0f, 0.0f, 0.0f, 0.0f,
+            0.0f, 0.0f, 0.0f, 0.237489700f, 0.449574471f, 0.285265595f, 0.812468469f, 0.0f, 0.0f, 0.0f,
+            0.0f, 0.0f, 0.0f, 0.382683456f, 0.0f, 0.0f, 0.923879504f, 0.0f, 0.0f, 0.0f
         }, {
             // LOD-1
             76.0f, 0.0f, 0.0f, 0.0f, 0.0f, 0.0f, 1.0f, 0.0f, 0.0f, 0.0f,
@@ -1113,7 +1175,10 @@ const Vector<Matrix<float> > valuesPerLOD = {
             4256.0f, 4332.0f, 4408.0f, -0.495134f, 0.1205274f, 0.8575973f, -0.0695866f, 0.0f, 0.0f, 0.0f,
             4940.0f, 5016.0f, 5092.0f, 0.4187725f, -0.2083559f, -0.7320557f, -0.4952896f, 5396.0f, 0.0f, 0.0f,
             15.0f, 30.0f, 0.0f, 0.0f, 0.0f, 0.0f, 1.0f, 0.0f, 0.0f, 0.0f,
-            40.0f, 80.0f
+            40.0f, 80.0f, 0.0f, 0.0f, 0.0f, 0.0f, 1.0f, 0.0f, 0.0f, 0.0f,
+            0.0f, 0.0f, 0.0f, 0.237489700f, 0.449574471f, 0.285265595f, 0.812468469f, 0.0f, 0.0f, 0.0f,
+            0.0f, 0.0f, 0.0f, 0.382683456f, 0.0f, 0.0f, 0.923879504f, 0.0f, 0.0f, 0.0f
+
         }, {
             // LOD-2
             76.0f, 0.0f, 0.0f, 0.0f, 0.0f, 0.0f, 1.0f, 0.0f, 0.0f, 0.0f,
@@ -1131,7 +1196,9 @@ const Vector<Matrix<float> > valuesPerLOD = {
             4256.0f, 0.0f, 0.0f, 0.0f, 0.0f, 0.0f, 1.0f, 0.0f, 0.0f, 0.0f,
             4940.0f, 5016.0f, 5092.0f, 0.4187725f, -0.2083559f, -0.7320557f, -0.4952896f, 5396.0f, 0.0f, 0.0f,
             15.0f, 30.0f, 0.0f, 0.0f, 0.0f, 0.0f, 1.0f, 0.0f, 0.0f, 0.0f,
-            40.0f, 80.0f
+            40.0f, 80.0f, 0.0f, 0.0f, 0.0f, 0.0f, 1.0f, 0.0f, 0.0f, 0.0f,
+            0.0f, 0.0f, 0.0f, 0.126078621f, -0.0337826647f, 0.256604820f, 0.957662225f, 0.0f, 0.0f, 0.0f,
+            0.0f, 0.0f, 0.0f, 0.0f, 0.0f, 0.0f, 1.0f, 0.0f, 0.0f, 0.0f
         }, {
             // LOD-3
             76.0f, 0.0f, 0.0f, 0.0f, 0.0f, 0.0f, 1.0f, 0.0f, 0.0f, 0.0f,
@@ -1149,7 +1216,9 @@ const Vector<Matrix<float> > valuesPerLOD = {
             0.0f, 0.0f, 0.0f, 0.0f, 0.0f, 0.0f, 1.0f, 0.0f, 0.0f, 0.0f,
             4940.0f, 5016.0f, 5092.0f, 0.898794f, 0.0f, 0.0f, 0.4383711f, 0.0f, 0.0f, 0.0f,
             15.0f, 30.0f, 0.0f, 0.0f, 0.0f, 0.0f, 1.0f, 0.0f, 0.0f, 0.0f,
-            40.0f, 80.0f
+            40.0f, 80.0f, 0.0f, 0.0f, 0.0f, 0.0f, 1.0f, 0.0f, 0.0f, 0.0f,
+            0.0f, 0.0f, 0.0f, 0.126078621f, -0.0337826647f, 0.256604820f, 0.957662225f, 0.0f, 0.0f, 0.0f,
+            0.0f, 0.0f, 0.0f, 0.0f, 0.0f, 0.0f, 1.0f, 0.0f, 0.0f, 0.0f
         }
     },
     {  // Euler-angle outputs
@@ -1178,7 +1247,10 @@ const Vector<Matrix<float> > valuesPerLOD = {
             0.0f, 0.0f,
             15.0f, 30.0f,  // Joint group 9
             0.0f, 0.0f, 0.0f, 0.0f, 0.0f, 0.0f, 0.0f,
-            40.0f, 80.0f  // Joint group 10
+            40.0f, 80.0f,  // Joint group 10
+            0.0f, 0.0f, 0.0f, 0.0f, 0.0f, 0.0f, 0.0f,
+            0.0f, 0.0f, 0.0f, 15.0f, 60.0f, 30.0f, 0.0f, 0.0f, 0.0f,  // Joint group 11
+            0.0f, 0.0f, 0.0f, 45.0f, 0.0f, 0.0f, 0.0f, 0.0f, 0.0f  // Joint group 11
         }, {
             // LOD-1
             76.0f,  // Joint group 0
@@ -1204,7 +1276,10 @@ const Vector<Matrix<float> > valuesPerLOD = {
             0.0f, 0.0f,
             15.0f, 30.0f,  // Joint group 9
             0.0f, 0.0f, 0.0f, 0.0f, 0.0f, 0.0f, 0.0f,
-            40.0f, 80.0f  // Joint group 10
+            40.0f, 80.0f,  // Joint group 10
+            0.0f, 0.0f, 0.0f, 0.0f, 0.0f, 0.0f, 0.0f,
+            0.0f, 0.0f, 0.0f, 15.0f, 60.0f, 30.0f, 0.0f, 0.0f, 0.0f,  // Joint group 11
+            0.0f, 0.0f, 0.0f, 45.0f, 0.0f, 0.0f, 0.0f, 0.0f, 0.0f  // Joint group 11
         }, {
             // LOD-2
             76.0f,  // Joint group 0
@@ -1230,7 +1305,10 @@ const Vector<Matrix<float> > valuesPerLOD = {
             0.0f, 0.0f,
             15.0f, 30.0f,  // Joint group 9
             0.0f, 0.0f, 0.0f, 0.0f, 0.0f, 0.0f, 0.0f,
-            40.0f, 80.0f  // Joint group 10
+            40.0f, 80.0f,  // Joint group 10
+            0.0f, 0.0f, 0.0f, 0.0f, 0.0f, 0.0f, 0.0f,
+            0.0f, 0.0f, 0.0f, 15.0f, 0.0f, 30.0f, 0.0f, 0.0f, 0.0f,  // Joint group 11
+            0.0f, 0.0f, 0.0f, 0.0f, 0.0f, 0.0f, 0.0f, 0.0f, 0.0f  // Joint group 11
         }, {
             // LOD-3
             76.0f,  // Joint group 0
@@ -1256,7 +1334,10 @@ const Vector<Matrix<float> > valuesPerLOD = {
             0.0f, 0.0f,
             15.0f, 30.0f,  // Joint group 9
             0.0f, 0.0f, 0.0f, 0.0f, 0.0f, 0.0f, 0.0f,
-            40.0f, 80.0f  // Joint group 10
+            40.0f, 80.0f,  // Joint group 10
+            0.0f, 0.0f, 0.0f, 0.0f, 0.0f, 0.0f, 0.0f,
+            0.0f, 0.0f, 0.0f, 15.0f, 0.0f, 30.0f, 0.0f, 0.0f, 0.0f,  // Joint group 11
+            0.0f, 0.0f, 0.0f, 0.0f, 0.0f, 0.0f, 0.0f, 0.0f, 0.0f  // Joint group 11
         }
     }
 };
