Extend DNA RBF layer so RigLogic can drive the neck setup as well

diff --git a/Engine/Plugins/Animation/RigLogic/RigLogic.uplugin b/Engine/Plugins/Animation/RigLogic/RigLogic.uplugin
--- a/Engine/Plugins/Animation/RigLogic/RigLogic.uplugin
+++ b/Engine/Plugins/Animation/RigLogic/RigLogic.uplugin
@@ -1,9 +1,9 @@
 {
 	"FileVersion": 3,
-	"Version": 10.1,
-	"VersionName": "10.1.4",
-	"FriendlyName": "RigLogic Plugin v10.1.4",
-	"Description": "3Lateral RigLogic Plugin for Facial Animation v10.1.4",
+	"Version": 10.2,
+	"VersionName": "10.2.0",
+	"FriendlyName": "RigLogic Plugin v10.2.0",
+	"Description": "3Lateral RigLogic Plugin for Facial Animation v10.2.0",
 	"Category": "Animation",
 	"CreatedBy": "3Lateral",
 	"CreatedByURL": "https://www.3lateral.com/",
diff --git a/Engine/Plugins/Animation/RigLogic/Source/RigLogicLib/Private/dna/DNA.h b/Engine/Plugins/Animation/RigLogic/Source/RigLogicLib/Private/dna/DNA.h
--- a/Engine/Plugins/Animation/RigLogic/Source/RigLogicLib/Private/dna/DNA.h
+++ b/Engine/Plugins/Animation/RigLogic/Source/RigLogicLib/Private/dna/DNA.h
@@ -99,7 +99,8 @@ enum FileVersion : std::uint64_t {
     v22 = rev(2, 2),
     v23 = rev(2, 3),
     v24 = rev(2, 4),
-    latest = v24
+    v25 = rev(2, 5),
+    latest = v25
 };
 
 using Dispatch = Dispatcher<FileVersion, FileVersion::v21, FileVersion::v22, FileVersion::v23, FileVersion::v24>;
@@ -1317,6 +1318,29 @@ struct RawRBFPose {
 
 };
 
+struct RawRBFPoseExt {
+    DynArray<std::uint16_t> inputControlIndices;
+    DynArray<std::uint16_t> outputControlIndices;
+    DynArray<float> outputControlWeights;
+
+    explicit RawRBFPoseExt(MemoryResource* memRes) :
+        inputControlIndices{memRes},
+        outputControlIndices{memRes},
+        outputControlWeights{memRes} {
+    }
+
+    template<class Archive>
+    void serialize(Archive& archive) {
+        archive.label("inputControlIndices");
+        archive(inputControlIndices);
+        archive.label("outputControlIndices");
+        archive(outputControlIndices);
+        archive.label("outputControlWeights");
+        archive(outputControlWeights);
+    }
+
+};
+
 struct RawRBFSolver {
     terse::ArchiveSize<std::uint32_t, std::uint32_t> size;
     terse::Anchor<std::uint32_t> baseMarker;
@@ -1409,6 +1433,25 @@ struct RawRBFBehavior {
 
 };
 
+struct RawRBFBehaviorExt {
+    Vector<String<char> > poseControlNames;
+    Vector<RawRBFPoseExt> poses;
+
+    explicit RawRBFBehaviorExt(MemoryResource* memRes) :
+        poseControlNames{memRes},
+        poses{memRes} {
+    }
+
+    template<class Archive>
+    void serialize(Archive& archive) {
+        archive.label("poseControlNames");
+        archive(poseControlNames);
+        archive.label("poses");
+        archive(poses);
+    }
+
+};
+
 struct RawJointRepresentation {
     std::uint16_t translation;
     std::uint16_t rotation;
@@ -1525,6 +1568,7 @@ struct DNA {
     Layer<sid4("geom"), rev(1, 1), RawGeometry, FileVersion::v22> geometry;
     Layer<sid4("mlbh"), rev(1, 0), RawMachineLearnedBehavior, FileVersion::v23> machineLearnedBehavior;
     Layer<sid4("rbfb"), rev(1, 0), RawRBFBehavior, FileVersion::v24> rbfBehavior;
+    Layer<sid4("rbfe"), rev(1, 0), RawRBFBehaviorExt, FileVersion::v25> rbfBehaviorExt;
     Layer<sid4("jbmd"), rev(1, 0), RawJointBehaviorMetadata, FileVersion::v24> jointBehaviorMetadata;
     Layer<sid4("twsw"), rev(1, 0), RawTwistSwingBehavior, FileVersion::v24> twistSwingBehavior;
 
@@ -1536,6 +1580,7 @@ struct DNA {
                        decltype(geometry),
                        decltype(machineLearnedBehavior),
                        decltype(rbfBehavior),
+                       decltype(rbfBehaviorExt),
                        decltype(jointBehaviorMetadata),
                        decltype(twistSwingBehavior)>;
     Layers layers;
@@ -1549,10 +1594,11 @@ struct DNA {
         geometry{memRes},
         machineLearnedBehavior{memRes},
         rbfBehavior{memRes},
+        rbfBehaviorExt{memRes},
         jointBehaviorMetadata{memRes},
         twistSwingBehavior{memRes},
         layers{unknownPolicy, upgradePolicy, memRes, &descriptor, &definition, &behavior, &geometry,
-               &machineLearnedBehavior, &rbfBehavior, &jointBehaviorMetadata, &twistSwingBehavior} {
+               &machineLearnedBehavior, &rbfBehavior, &rbfBehaviorExt, &jointBehaviorMetadata, &twistSwingBehavior} {
     }
 
     explicit DNA(MemoryResource* memRes_) : DNA{UnknownLayerPolicy::Preserve, UpgradeFormatPolicy::Allowed,
@@ -1629,6 +1675,8 @@ struct DNA {
         }
 
         context->data = nullptr;
+
+        postLoadActions();
     }
 
     template<class Archive, typename AnyVersion>
@@ -1647,6 +1695,23 @@ struct DNA {
         }
     }
 
+    void postLoadActions() {
+        if (rbfBehaviorExt.poses.empty() && !rbfBehavior.poses.empty()) {
+            rbfBehaviorExt.poses.resize(rbfBehavior.poses.size(), RawRBFPoseExt{memRes});
+            rbfBehaviorExt.poseControlNames.resize(rbfBehavior.poses.size());
+            const auto controlOffset = (definition.rawControlNames.size() +
+                                        behavior.controls.psdCount +
+                                        machineLearnedBehavior.mlControlNames.size());
+            for (std::size_t pi = {}; pi < rbfBehaviorExt.poses.size(); ++pi) {
+                rbfBehaviorExt.poseControlNames[pi] = rbfBehavior.poses[pi].name;
+                rbfBehaviorExt.poses[pi].outputControlIndices.resize(1ul);
+                rbfBehaviorExt.poses[pi].outputControlWeights.resize(1ul);
+                rbfBehaviorExt.poses[pi].outputControlIndices[0] = static_cast<std::uint16_t>(controlOffset + pi);
+                rbfBehaviorExt.poses[pi].outputControlWeights[0] = 1.0f;
+            }
+        }
+    }
+
     void unloadDefinition() {
         static_cast<RawDefinition&>(definition) = RawDefinition{memRes};
     }
@@ -1665,6 +1730,7 @@ struct DNA {
 
     void unloadRBFBehavior() {
         static_cast<RawRBFBehavior&>(rbfBehavior) = RawRBFBehavior{memRes};
+        static_cast<RawRBFBehaviorExt&>(rbfBehaviorExt) = RawRBFBehaviorExt{memRes};
     }
 
     void unloadJointBehaviorMetadata() {
diff --git a/Engine/Plugins/Animation/RigLogic/Source/RigLogicLib/Private/dna/DenormalizedData.h b/Engine/Plugins/Animation/RigLogic/Source/RigLogicLib/Private/dna/DenormalizedData.h
--- a/Engine/Plugins/Animation/RigLogic/Source/RigLogicLib/Private/dna/DenormalizedData.h
+++ b/Engine/Plugins/Animation/RigLogic/Source/RigLogicLib/Private/dna/DenormalizedData.h
@@ -18,7 +18,6 @@
     #pragma warning(pop)
 #endif
 
-
 namespace dna {
 
 template<class Reader>
@@ -119,35 +118,28 @@ struct DenormalizedData {
         void populateRBFPoseJointOutputIndices(const Reader* source) {
             const std::uint16_t poseCount = source->getRBFPoseCount();
             rbfPoseJointOutputIndices.resize(poseCount);
-
-            const auto indexOf = [](ConstArrayView<std::uint16_t> container, std::uint16_t val) {
-                    const auto valIt = std::find(container.begin(), container.end(), val);
-                    return static_cast<size_t>(std::distance(container.begin(), valIt));
-                };
-            const std::uint16_t rawControlCount = source->getRawControlCount();
-            const std::uint16_t psdControlCount = source->getPSDCount();
-            const std::uint16_t mlControlCount = source->getMLControlCount();
-            auto rbfControlOffset = rawControlCount + psdControlCount + mlControlCount;
-
-            for (std::uint16_t gi = {}; gi < source->getJointGroupCount(); ++gi) {
-                const auto inputIndices = source->getJointGroupInputIndices(gi);
-                const auto columnCount = inputIndices.size();
-                for (const auto& inputIndex : inputIndices) {
-                    if (inputIndex >= rbfControlOffset) {
-                        const auto poseIndex = static_cast<std::uint16_t>(inputIndex - rbfControlOffset);
-                        const auto outputIndices = source->getJointGroupOutputIndices(gi);
-                        rbfPoseJointOutputIndices[poseIndex].reserve(outputIndices.size());
-                        const auto columnIndex = indexOf(inputIndices, inputIndex);
-                        const auto values = source->getJointGroupValues(gi);
-                        for (std::uint16_t oi = 0; oi < outputIndices.size(); ++oi) {
-                            if (std::abs(values[oi * columnCount + columnIndex]) > 0.0f) {
-                                rbfPoseJointOutputIndices[poseIndex].push_back(outputIndices[oi]);
+            for (std::uint16_t pi = {}; pi < poseCount; ++pi) {
+                const auto poseControlIndices = source->getRBFPoseOutputControlIndices(pi);
+                for (std::uint16_t gi = {}; gi < source->getJointGroupCount(); ++gi) {
+                    const auto inputIndices = source->getJointGroupInputIndices(gi);
+                    const auto outputIndices = source->getJointGroupOutputIndices(gi);
+                    const auto values = source->getJointGroupValues(gi);
+                    const auto columnCount = inputIndices.size();
+                    const auto rowCount = outputIndices.size();
+                    for (std::size_t columnIndex = {}; columnIndex < columnCount; ++columnIndex) {
+                        const auto inputIndex = inputIndices[columnIndex];
+                        if (std::find(poseControlIndices.begin(), poseControlIndices.end(),
+                                      inputIndex) != poseControlIndices.end()) {
+                            rbfPoseJointOutputIndices[pi].reserve(rbfPoseJointOutputIndices[pi].size() + rowCount);
+                            for (std::uint16_t rowIndex = {}; rowIndex < rowCount; ++rowIndex) {
+                                if (std::abs(values[rowIndex * columnCount + columnIndex]) > 0.0f) {
+                                    rbfPoseJointOutputIndices[pi].push_back(outputIndices[rowIndex]);
+                                }
                             }
                         }
                     }
                 }
             }
-
         }
 
         void populateRBFPoseJointOutputValues(const Reader* source) {
@@ -169,31 +161,30 @@ struct DenormalizedData {
                     const auto valIt = std::find(container.begin(), container.end(), val);
                     return static_cast<size_t>(std::distance(container.begin(), valIt));
                 };
-            const std::uint16_t rawControlCount = source->getRawControlCount();
-            const std::uint16_t psdControlCount = source->getPSDCount();
-            const std::uint16_t mlControlCount = source->getMLControlCount();
-            auto rbfControlOffset = rawControlCount + psdControlCount + mlControlCount;
+
             for (std::uint16_t pi = {}; pi < poseCount; ++pi) {
-                const auto inputIndex = static_cast<std::uint16_t>(rbfControlOffset + pi);
-                const auto jointOutputIndices = source->getRBFPoseJointOutputIndices(pi);
-                const auto jointOutputIndicesCount = jointOutputIndices.size();
-                rbfPoseJointOutputValues[pi].resize(jointOutputIndices.size(), 0.0f);
-                for (std::uint16_t oi = {}; oi < jointOutputIndicesCount; ++oi) {
-                    const auto outputIndex = jointOutputIndices[oi];
-                    const auto jointGroupIndex = findJointGroupIndex(outputIndex);
-
-                    const auto jointGroupOutputIndices = source->getJointGroupOutputIndices(jointGroupIndex);
-                    const auto rowCount = jointGroupOutputIndices.size();
-
-                    const auto jointGroupInputIndices = source->getJointGroupInputIndices(jointGroupIndex);
-                    const auto columnCount = jointGroupInputIndices.size();
-
-                    const auto columnIndex = indexOf(jointGroupInputIndices, inputIndex);
-                    const auto rowIndex = indexOf(jointGroupOutputIndices, outputIndex);
-
-                    if ((columnIndex < columnCount) && (rowIndex < rowCount)) {
-                        const auto values = source->getJointGroupValues(jointGroupIndex);
-                        rbfPoseJointOutputValues[pi][oi] = values[rowIndex * columnCount + columnIndex];
+                const auto poseControlIndices = source->getRBFPoseOutputControlIndices(pi);
+                for (const auto inputIndex : poseControlIndices) {
+                    const auto jointOutputIndices = source->getRBFPoseJointOutputIndices(pi);
+                    const auto jointOutputIndicesCount = jointOutputIndices.size();
+                    rbfPoseJointOutputValues[pi].resize(jointOutputIndices.size(), 0.0f);
+                    for (std::uint16_t oi = {}; oi < jointOutputIndicesCount; ++oi) {
+                        const auto outputIndex = jointOutputIndices[oi];
+                        const auto jointGroupIndex = findJointGroupIndex(outputIndex);
+
+                        const auto jointGroupOutputIndices = source->getJointGroupOutputIndices(jointGroupIndex);
+                        const auto rowCount = jointGroupOutputIndices.size();
+
+                        const auto jointGroupInputIndices = source->getJointGroupInputIndices(jointGroupIndex);
+                        const auto columnCount = jointGroupInputIndices.size();
+
+                        const auto columnIndex = indexOf(jointGroupInputIndices, inputIndex);
+                        const auto rowIndex = indexOf(jointGroupOutputIndices, outputIndex);
+
+                        if ((columnIndex < columnCount) && (rowIndex < rowCount)) {
+                            const auto values = source->getJointGroupValues(jointGroupIndex);
+                            rbfPoseJointOutputValues[pi][oi] = values[rowIndex * columnCount + columnIndex];
+                        }
                     }
                 }
             }
@@ -205,22 +196,18 @@ struct DenormalizedData {
             const auto bscInputIndices = source->getBlendShapeChannelInputIndices();
             const auto bscOutputIndices = source->getBlendShapeChannelOutputIndices();
             const auto bscCount = bscInputIndices.size();
-            const std::uint16_t rawControlCount = source->getRawControlCount();
-            const std::uint16_t psdControlCount = source->getPSDCount();
-            const std::uint16_t mlControlCount = source->getMLControlCount();
-            auto rbfControlOffset = rawControlCount + psdControlCount + mlControlCount;
-
             const auto poseCount = source->getRBFPoseCount();
             rbfBlendShapeChannelOutputIndices.resize(poseCount);
             for (std::uint16_t pi = {}; pi < poseCount; ++pi) {
-                const auto inputIndex = rbfControlOffset + pi;
-                for (std::uint16_t bi = {}; bi < bscCount; ++bi) {
-                    if (bscInputIndices[bi] == inputIndex) {
-                        rbfBlendShapeChannelOutputIndices[pi].push_back(bscOutputIndices[bi]);
+                const auto poseControlIndices = source->getRBFPoseOutputControlIndices(pi);
+                for (const auto inputIndex : poseControlIndices) {
+                    for (std::uint16_t bi = {}; bi < bscCount; ++bi) {
+                        if (bscInputIndices[bi] == inputIndex) {
+                            rbfBlendShapeChannelOutputIndices[pi].push_back(bscOutputIndices[bi]);
+                        }
                     }
                 }
             }
-
         }
 
         void populateRBFAnimatedMapOutputIndices(const Reader* source) {
@@ -229,21 +216,17 @@ struct DenormalizedData {
             const auto amInputIndices = source->getAnimatedMapInputIndices();
             const auto amOutputIndices = source->getAnimatedMapOutputIndices();
             const auto amCount = source->getAnimatedMapCount();
-            const std::uint16_t rawControlCount = source->getRawControlCount();
-            const std::uint16_t psdControlCount = source->getPSDCount();
-            const std::uint16_t mlControlCount = source->getMLControlCount();
-            auto rbfControlOffset = rawControlCount + psdControlCount + mlControlCount;
-
             const auto poseCount = source->getRBFPoseCount();
             for (std::uint16_t pi = {}; pi < poseCount; ++pi) {
-                const auto inputIndex = rbfControlOffset + pi;
-                for (std::uint16_t ai = {}; ai < amCount; ++ai) {
-                    if (amInputIndices[ai] == inputIndex) {
-                        rbfAnimatedMapOutputIndices[pi].push_back(amOutputIndices[ai]);
+                const auto poseControlIndices = source->getRBFPoseOutputControlIndices(pi);
+                for (const auto inputIndex : poseControlIndices) {
+                    for (std::uint16_t ai = {}; ai < amCount; ++ai) {
+                        if (amInputIndices[ai] == inputIndex) {
+                            rbfAnimatedMapOutputIndices[pi].push_back(amOutputIndices[ai]);
+                        }
                     }
                 }
             }
-
         }
 
 };
diff --git a/Engine/Plugins/Animation/RigLogic/Source/RigLogicLib/Private/dna/ReaderImpl.h b/Engine/Plugins/Animation/RigLogic/Source/RigLogicLib/Private/dna/ReaderImpl.h
--- a/Engine/Plugins/Animation/RigLogic/Source/RigLogicLib/Private/dna/ReaderImpl.h
+++ b/Engine/Plugins/Animation/RigLogic/Source/RigLogicLib/Private/dna/ReaderImpl.h
@@ -179,6 +179,11 @@ class ReaderImpl : public TReaderBase, public virtual BaseImpl {
         ConstArrayView<std::uint16_t> getRBFPoseAnimatedMapOutputIndices(std::uint16_t poseIndex) const override;
         ConstArrayView<float> getRBFPoseJointOutputValues(std::uint16_t poseIndex) const override;
         float getRBFPoseScale(std::uint16_t poseIndex) const override;
+        std::uint16_t getRBFPoseControlCount() const override;
+        StringView getRBFPoseControlName(std::uint16_t poseControlIndex) const override;
+        ConstArrayView<std::uint16_t> getRBFPoseInputControlIndices(std::uint16_t poseIndex) const override;
+        ConstArrayView<std::uint16_t> getRBFPoseOutputControlIndices(std::uint16_t poseIndex) const override;
+        ConstArrayView<float> getRBFPoseOutputControlWeights(std::uint16_t poseIndex) const override;
         std::uint16_t getRBFSolverCount() const override;
         std::uint16_t getRBFSolverIndexListCount() const override;
         ConstArrayView<std::uint16_t> getRBFSolverIndicesForLOD(std::uint16_t lod) const override;
@@ -1237,6 +1242,47 @@ inline float ReaderImpl<TReaderBase>::getRBFPoseScale(std::uint16_t poseIndex) c
     return 0.0f;
 }
 
+template<class TReaderBase>
+inline std::uint16_t ReaderImpl<TReaderBase>::getRBFPoseControlCount() const {
+    return static_cast<std::uint16_t>(dna.rbfBehaviorExt.poseControlNames.size());
+}
+
+template<class TReaderBase>
+StringView ReaderImpl<TReaderBase>::getRBFPoseControlName(std::uint16_t poseControlIndex) const {
+    if (poseControlIndex < dna.rbfBehaviorExt.poseControlNames.size()) {
+        const auto& poseControlName = dna.rbfBehaviorExt.poseControlNames[poseControlIndex];
+        return {poseControlName.data(), poseControlName.size()};
+    }
+    return {};
+}
+
+template<class TReaderBase>
+inline ConstArrayView<std::uint16_t> ReaderImpl<TReaderBase>::getRBFPoseInputControlIndices(std::uint16_t poseIndex) const {
+    const auto& poses = dna.rbfBehaviorExt.poses;
+    if (poseIndex < poses.size()) {
+        return poses[poseIndex].inputControlIndices;
+    }
+    return {};
+}
+
+template<class TReaderBase>
+inline ConstArrayView<std::uint16_t> ReaderImpl<TReaderBase>::getRBFPoseOutputControlIndices(std::uint16_t poseIndex) const {
+    const auto& poses = dna.rbfBehaviorExt.poses;
+    if (poseIndex < poses.size()) {
+        return poses[poseIndex].outputControlIndices;
+    }
+    return {};
+}
+
+template<class TReaderBase>
+inline ConstArrayView<float> ReaderImpl<TReaderBase>::getRBFPoseOutputControlWeights(std::uint16_t poseIndex) const {
+    const auto& poses = dna.rbfBehaviorExt.poses;
+    if (poseIndex < poses.size()) {
+        return poses[poseIndex].outputControlWeights;
+    }
+    return {};
+}
+
 template<class TReaderBase>
 inline std::uint16_t ReaderImpl<TReaderBase>::getRBFSolverCount() const {
     return static_cast<std::uint16_t>(dna.rbfBehavior.solvers.size());
diff --git a/Engine/Plugins/Animation/RigLogic/Source/RigLogicLib/Private/dna/Writer.cpp b/Engine/Plugins/Animation/RigLogic/Source/RigLogicLib/Private/dna/Writer.cpp
--- a/Engine/Plugins/Animation/RigLogic/Source/RigLogicLib/Private/dna/Writer.cpp
+++ b/Engine/Plugins/Animation/RigLogic/Source/RigLogicLib/Private/dna/Writer.cpp
@@ -527,10 +527,24 @@ static void copyRBFBehavior(const RBFBehaviorReader* source, RBFBehaviorWriter*
         source->getLODCount(),
         memRes);
 
+    for (std::uint16_t pciPlusOne = source->getRBFPoseControlCount(); pciPlusOne > 0u; --pciPlusOne) {
+        const auto pci = static_cast<std::uint16_t>(pciPlusOne - 1u);
+        destination->setRBFPoseControlName(pci, source->getRBFPoseControlName(pci).data());
+    }
+
     for (std::uint16_t posePlusOne = source->getRBFPoseCount(); posePlusOne > 0u; --posePlusOne) {
         const auto pi = static_cast<std::uint16_t>(posePlusOne - 1u);
         destination->setRBFPoseName(pi, source->getRBFPoseName(pi).data());
         destination->setRBFPoseScale(pi, source->getRBFPoseScale(pi));
+        const auto inputControlIndices = source->getRBFPoseInputControlIndices(pi);
+        destination->setRBFPoseInputControlIndices(pi, inputControlIndices.data(),
+                                                   static_cast<std::uint16_t>(inputControlIndices.size()));
+        const auto outputControlIndices = source->getRBFPoseOutputControlIndices(pi);
+        destination->setRBFPoseOutputControlIndices(pi, outputControlIndices.data(),
+                                                    static_cast<std::uint16_t>(outputControlIndices.size()));
+        const auto outputControlWeights = source->getRBFPoseOutputControlWeights(pi);
+        destination->setRBFPoseOutputControlWeights(pi, outputControlWeights.data(),
+                                                    static_cast<std::uint16_t>(outputControlWeights.size()));
     }
 
     for (std::uint16_t solverPlusOne = source->getRBFSolverCount(); solverPlusOne > 0u; --solverPlusOne) {
diff --git a/Engine/Plugins/Animation/RigLogic/Source/RigLogicLib/Private/dna/WriterImpl.h b/Engine/Plugins/Animation/RigLogic/Source/RigLogicLib/Private/dna/WriterImpl.h
--- a/Engine/Plugins/Animation/RigLogic/Source/RigLogicLib/Private/dna/WriterImpl.h
+++ b/Engine/Plugins/Animation/RigLogic/Source/RigLogicLib/Private/dna/WriterImpl.h
@@ -225,6 +225,16 @@ class WriterImpl : public TWriterBase, public virtual BaseImpl {
         void clearRBFPoses() override;
         void setRBFPoseName(std::uint16_t poseIndex, const char* name) override;
         void setRBFPoseScale(std::uint16_t poseIndex, float scale) override;
+        void clearRBFPoseControlNames() override;
+        void setRBFPoseControlName(std::uint16_t poseControlIndex, const char* name) override;
+        void setRBFPoseInputControlIndices(std::uint16_t poseIndex,
+                                           const std::uint16_t* controlIndices,
+                                           std::uint16_t controlIndexCount) override;
+        void setRBFPoseOutputControlIndices(std::uint16_t poseIndex,
+                                            const std::uint16_t* controlIndices,
+                                            std::uint16_t controlIndexCount) override;
+        void setRBFPoseOutputControlWeights(std::uint16_t poseIndex, const float* controlWeights,
+                                            std::uint16_t controlWeightCount) override;
         void clearRBFSolvers() override;
         void clearRBFSolverIndices() override;
         void setRBFSolverIndices(std::uint16_t index, const std::uint16_t* solverIndices, std::uint16_t count) override;
@@ -1002,6 +1012,7 @@ inline void WriterImpl<TWriterBase>::setNeuralNetworkLayerWeights(std::uint16_t
 template<class TWriterBase>
 inline void WriterImpl<TWriterBase>::clearRBFPoses() {
     dna.rbfBehavior.poses.clear();
+    dna.rbfBehaviorExt.poses.clear();
 }
 
 template<class TWriterBase>
@@ -1016,6 +1027,40 @@ inline void WriterImpl<TWriterBase>::setRBFPoseScale(std::uint16_t poseIndex, fl
     pose.scale = scale;
 }
 
+template<class TWriterBase>
+inline void WriterImpl<TWriterBase>::clearRBFPoseControlNames() {
+    dna.rbfBehaviorExt.poseControlNames.clear();
+}
+
+template<class TWriterBase>
+inline void WriterImpl<TWriterBase>::setRBFPoseControlName(std::uint16_t poseControlIndex, const char* name) {
+    setAt(dna.rbfBehaviorExt.poseControlNames, poseControlIndex, name);
+}
+
+template<class TWriterBase>
+inline void WriterImpl<TWriterBase>::setRBFPoseInputControlIndices(std::uint16_t poseIndex,
+                                                                   const std::uint16_t* controlIndices,
+                                                                   std::uint16_t controlIndexCount) {
+    auto& pose = getAt(dna.rbfBehaviorExt.poses, poseIndex);
+    pose.inputControlIndices.assign(controlIndices, controlIndices + controlIndexCount);
+}
+
+template<class TWriterBase>
+inline void WriterImpl<TWriterBase>::setRBFPoseOutputControlIndices(std::uint16_t poseIndex,
+                                                                    const std::uint16_t* controlIndices,
+                                                                    std::uint16_t controlIndexCount) {
+    auto& pose = getAt(dna.rbfBehaviorExt.poses, poseIndex);
+    pose.outputControlIndices.assign(controlIndices, controlIndices + controlIndexCount);
+}
+
+template<class TWriterBase>
+inline void WriterImpl<TWriterBase>::setRBFPoseOutputControlWeights(std::uint16_t poseIndex,
+                                                                    const float* controlWeights,
+                                                                    std::uint16_t controlWeightCount) {
+    auto& pose = getAt(dna.rbfBehaviorExt.poses, poseIndex);
+    pose.outputControlWeights.assign(controlWeights, controlWeights + controlWeightCount);
+}
+
 template<class TWriterBase>
 inline void WriterImpl<TWriterBase>::clearRBFSolvers() {
     dna.rbfBehavior.solvers.clear();
diff --git a/Engine/Plugins/Animation/RigLogic/Source/RigLogicLib/Private/dna/stream/FilteredBinaryInputArchive.cpp b/Engine/Plugins/Animation/RigLogic/Source/RigLogicLib/Private/dna/stream/FilteredBinaryInputArchive.cpp
--- a/Engine/Plugins/Animation/RigLogic/Source/RigLogicLib/Private/dna/stream/FilteredBinaryInputArchive.cpp
+++ b/Engine/Plugins/Animation/RigLogic/Source/RigLogicLib/Private/dna/stream/FilteredBinaryInputArchive.cpp
@@ -516,6 +516,19 @@ void FilteredBinaryInputArchive::process(RawRBFBehavior& dest) {
     process(dest.poses);
 }
 
+void FilteredBinaryInputArchive::process(RawRBFBehaviorExt& dest) {
+    if (malformed) {
+        return;
+    }
+
+    if (!contains(layerBitmask, DataLayerBitmask::RBFBehavior)) {
+        return;
+    }
+
+    process(dest.poseControlNames);
+    process(dest.poses);
+}
+
 void FilteredBinaryInputArchive::process(RawJointBehaviorMetadata& dest) {
     if (malformed) {
         return;
diff --git a/Engine/Plugins/Animation/RigLogic/Source/RigLogicLib/Private/dna/stream/FilteredBinaryInputArchive.h b/Engine/Plugins/Animation/RigLogic/Source/RigLogicLib/Private/dna/stream/FilteredBinaryInputArchive.h
--- a/Engine/Plugins/Animation/RigLogic/Source/RigLogicLib/Private/dna/stream/FilteredBinaryInputArchive.h
+++ b/Engine/Plugins/Animation/RigLogic/Source/RigLogicLib/Private/dna/stream/FilteredBinaryInputArchive.h
@@ -82,6 +82,7 @@ class FilteredBinaryInputArchive final : public AnimatedMapFilter, public BlendS
         void process(RawVertexSkinWeights& dest);
         void process(RawMachineLearnedBehavior& dest);
         void process(RawRBFBehavior& dest);
+        void process(RawRBFBehaviorExt& dest);
         void process(RawJointBehaviorMetadata& dest);
         void process(RawTwistSwingBehavior& dest);
 
diff --git a/Engine/Plugins/Animation/RigLogic/Source/RigLogicLib/Private/riglogic/controls/ControlsFactory.cpp b/Engine/Plugins/Animation/RigLogic/Source/RigLogicLib/Private/riglogic/controls/ControlsFactory.cpp
--- a/Engine/Plugins/Animation/RigLogic/Source/RigLogicLib/Private/riglogic/controls/ControlsFactory.cpp
+++ b/Engine/Plugins/Animation/RigLogic/Source/RigLogicLib/Private/riglogic/controls/ControlsFactory.cpp
@@ -85,7 +85,7 @@ Controls::Pointer ControlsFactory::create(const Configuration& config, const dna
                                                  conditionals.getOutputCount(),
                                                  psds.getDistinctPSDCount(),
                                                  reader->getMLControlCount(),
-                                                 reader->getRBFPoseCount());
+                                                 reader->getRBFPoseControlCount());
 
     return UniqueInstance<Controls>::with(memRes).create(std::move(conditionals), std::move(psds), instanceFactory);
 }
diff --git a/Engine/Plugins/Animation/RigLogic/Source/RigLogicLib/Private/riglogic/joints/cpu/quaternions/QuaternionJointsBuilder.h b/Engine/Plugins/Animation/RigLogic/Source/RigLogicLib/Private/riglogic/joints/cpu/quaternions/QuaternionJointsBuilder.h
--- a/Engine/Plugins/Animation/RigLogic/Source/RigLogicLib/Private/riglogic/joints/cpu/quaternions/QuaternionJointsBuilder.h
+++ b/Engine/Plugins/Animation/RigLogic/Source/RigLogicLib/Private/riglogic/joints/cpu/quaternions/QuaternionJointsBuilder.h
@@ -97,9 +97,9 @@ void QuaternionJointsBuilder<TValue, TFVec256, TFVec128>::setOutputIndices(const
     #endif
     auto deduplicate = [this](Vector<std::uint16_t>& v) {
             UnorderedSet<std::uint16_t> deduplicator{memRes};
-            v.erase(std::remove_if(v.begin(), v.end(), [&](const std::uint16_t value) {
+            v.erase(v.rend().base(), std::remove_if(v.rbegin(), v.rend(), [&deduplicator](const std::uint16_t value) {
                 return !deduplicator.insert(value).second;
-            }), v.end());
+            }).base());
         };
     #if !defined(__clang__) && defined(__GNUC__)
         #pragma GCC diagnostic pop
diff --git a/Engine/Plugins/Animation/RigLogic/Source/RigLogicLib/Private/riglogic/rbf/RBFBehaviorFactory.cpp b/Engine/Plugins/Animation/RigLogic/Source/RigLogicLib/Private/riglogic/rbf/RBFBehaviorFactory.cpp
--- a/Engine/Plugins/Animation/RigLogic/Source/RigLogicLib/Private/riglogic/rbf/RBFBehaviorFactory.cpp
+++ b/Engine/Plugins/Animation/RigLogic/Source/RigLogicLib/Private/riglogic/rbf/RBFBehaviorFactory.cpp
@@ -53,7 +53,7 @@ RBFBehaviorEvaluator::Pointer createRBFEvaluator(const Configuration& config, co
 
 RBFBehavior::Pointer RBFBehaviorFactory::create(const Configuration& config, const dna::Reader* reader, MemoryResource* memRes) {
     auto moduleFactory = UniqueInstance<RBFBehavior>::with(memRes);
-    if (!config.loadRBFBehavior || (reader->getRBFSolverCount() == 0u) || (reader->getRBFPoseCount() == 0u)) {
+    if (!config.loadRBFBehavior || (reader->getRBFSolverCount() == 0u) || (reader->getRBFPoseControlCount() == 0u)) {
         auto evaluator =
             UniqueInstance<RBFBehaviorNullEvaluator, RBFBehaviorEvaluator>::with(memRes).create();
         return moduleFactory.create(std::move(evaluator));
diff --git a/Engine/Plugins/Animation/RigLogic/Source/RigLogicLib/Private/riglogic/rbf/cpu/CPURBFBehaviorEvaluator.h b/Engine/Plugins/Animation/RigLogic/Source/RigLogicLib/Private/riglogic/rbf/cpu/CPURBFBehaviorEvaluator.h
--- a/Engine/Plugins/Animation/RigLogic/Source/RigLogicLib/Private/riglogic/rbf/cpu/CPURBFBehaviorEvaluator.h
+++ b/Engine/Plugins/Animation/RigLogic/Source/RigLogicLib/Private/riglogic/rbf/cpu/CPURBFBehaviorEvaluator.h
@@ -32,15 +32,21 @@ class Evaluator : public RBFBehaviorEvaluator {
     public:
         Evaluator(LODSpec<std::uint16_t>&& lods_,
                   SolverVectorType&& solvers_,
-                  Matrix<std::uint16_t> solverRawControlInputIndices_,
-                  Matrix<std::uint16_t> solverRawControlOutputIndices_,
+                  Matrix<std::uint16_t>&& solverRawControlInputIndices_,
+                  Matrix<std::uint16_t>&& solverPoseIndices_,
+                  Matrix<std::uint16_t>&& poseInputControlIndices_,
+                  Matrix<std::uint16_t>&& poseOutputControlIndices_,
+                  Matrix<float>&& poseOutputControlWeights_,
                   std::uint16_t maximumInputCount_,
                   std::uint16_t maxTargetCount_,
                   OutputInstance::Factory&& instanceFactory_) :
             lods{std::move(lods_)},
             solvers{std::move(solvers_)},
             solverRawControlInputIndices{std::move(solverRawControlInputIndices_)},
-            solverRawControlOutputIndices{std::move(solverRawControlOutputIndices_)},
+            solverPoseIndices{std::move(solverPoseIndices_)},
+            poseInputControlIndices{std::move(poseInputControlIndices_)},
+            poseOutputControlIndices{std::move(poseOutputControlIndices_)},
+            poseOutputControlWeights{std::move(poseOutputControlWeights_)},
             instanceFactory{std::move(instanceFactory_)},
             maximumInputCount{maximumInputCount_},
             maxTargetCount{maxTargetCount_} {
@@ -61,10 +67,10 @@ class Evaluator : public RBFBehaviorEvaluator {
             const auto rawControlInputCount = rawControlInputIndices.size();
             inputBuffer = inputBuffer.subview(0u, rawControlInputCount);
 
-            const auto& rawControlOutputIndices = solverRawControlOutputIndices[solverIndex];
-            const auto rawControlOutputCount = rawControlOutputIndices.size();
-            intermediateWeightsBuffer = intermediateWeightsBuffer.subview(0u, rawControlOutputCount);
-            outputWeightsBuffer = outputWeightsBuffer.subview(0u, rawControlOutputCount);
+            const auto& poseIndices = solverPoseIndices[solverIndex];
+            const auto poseCount = poseIndices.size();
+            intermediateWeightsBuffer = intermediateWeightsBuffer.subview(0u, poseCount);
+            outputWeightsBuffer = outputWeightsBuffer.subview(0u, poseCount);
 
             for (std::uint16_t i = {}; i < rawControlInputCount; ++i) {
                 const std::uint16_t ri = rawControlInputIndices[i];
@@ -73,11 +79,22 @@ class Evaluator : public RBFBehaviorEvaluator {
 
             solvers[solverIndex]->solve(inputBuffer, intermediateWeightsBuffer, outputWeightsBuffer);
 
-            for (std::uint16_t i = {}; i < rawControlOutputCount; ++i) {
-                const std::uint16_t ri = rawControlOutputIndices[i];
-                rawControls[ri] = outputWeightsBuffer[i];
-            }
+            for (std::uint16_t i = {}; i < poseCount; ++i) {
+                const std::uint16_t pi = poseIndices[i];
+                const auto& inputControlIndices = poseInputControlIndices[pi];
+                const auto& outputControlIndices = poseOutputControlIndices[pi];
+                const auto& outputControlWeights = poseOutputControlWeights[pi];
+                assert(outputControlIndices.size() == outputControlWeights.size());
+
+                float inputWeight = 1.0f;
+                for (const auto inputControlIndex : inputControlIndices) {
+                    inputWeight *= rawControls[inputControlIndex];
+                }
 
+                for (std::uint16_t ci = 0u; ci < outputControlIndices.size(); ci++) {
+                    rawControls[outputControlIndices[ci]] = outputControlWeights[ci] * outputWeightsBuffer[i] * inputWeight;
+                }
+            }
         }
 
         void calculate(ControlsInputInstance* inputs, RBFBehaviorOutputInstance* intermediateOutputs,
@@ -110,7 +127,6 @@ class Evaluator : public RBFBehaviorEvaluator {
             auto outputWeightsBuffer = static_cast<OutputInstance*>(intermediateOutputs)->getOutputWeightsBuffer();
 
             calculate(solverIndex, rawControls, inputBuffer, intermediateWeightsBuffer, outputWeightsBuffer);
-
         }
 
         void load(terse::BinaryInputArchive<BoundedIOStream>& archive) override {
@@ -132,7 +148,10 @@ class Evaluator : public RBFBehaviorEvaluator {
                 solvers[i]->load(archive);
             }
             archive(solverRawControlInputIndices);
-            archive(solverRawControlOutputIndices);
+            archive(solverPoseIndices);
+            archive(poseInputControlIndices);
+            archive(poseOutputControlIndices);
+            archive(poseOutputControlWeights);
             archive(maximumInputCount);
             archive(maxTargetCount);
         }
@@ -147,7 +166,10 @@ class Evaluator : public RBFBehaviorEvaluator {
                 solver->save(archive);
             }
             archive(solverRawControlInputIndices);
-            archive(solverRawControlOutputIndices);
+            archive(solverPoseIndices);
+            archive(poseInputControlIndices);
+            archive(poseOutputControlIndices);
+            archive(poseOutputControlWeights);
             archive(maximumInputCount);
             archive(maxTargetCount);
         }
@@ -156,7 +178,10 @@ class Evaluator : public RBFBehaviorEvaluator {
         LODSpec<std::uint16_t> lods;
         SolverVectorType solvers;
         Matrix<std::uint16_t> solverRawControlInputIndices;
-        Matrix<std::uint16_t> solverRawControlOutputIndices;
+        Matrix<std::uint16_t> solverPoseIndices;
+        Matrix<std::uint16_t> poseInputControlIndices;
+        Matrix<std::uint16_t> poseOutputControlIndices;
+        Matrix<float> poseOutputControlWeights;
         OutputInstance::Factory instanceFactory;
         std::uint16_t maximumInputCount;
         std::uint16_t maxTargetCount;
diff --git a/Engine/Plugins/Animation/RigLogic/Source/RigLogicLib/Private/riglogic/rbf/cpu/CPURBFBehaviorFactory.h b/Engine/Plugins/Animation/RigLogic/Source/RigLogicLib/Private/riglogic/rbf/cpu/CPURBFBehaviorFactory.h
--- a/Engine/Plugins/Animation/RigLogic/Source/RigLogicLib/Private/riglogic/rbf/cpu/CPURBFBehaviorFactory.h
+++ b/Engine/Plugins/Animation/RigLogic/Source/RigLogicLib/Private/riglogic/rbf/cpu/CPURBFBehaviorFactory.h
@@ -28,7 +28,10 @@ class Factory {
             std::uint16_t maximumInputCount{};
             std::uint16_t maximumTargetCount{};
             Matrix<std::uint16_t> solverRawControlInputIndices{memRes};
-            Matrix<std::uint16_t> solverRawControlOutputIndices{memRes};
+            Matrix<std::uint16_t> solverPoseIndices{memRes};
+            Matrix<std::uint16_t> poseInputControlIndices{memRes};
+            Matrix<std::uint16_t> poseOutputControlIndices{memRes};
+            Matrix<float> poseOutputControlWeights{memRes};
 
             auto instanceFactory = [](std::uint16_t maxInputCount, std::uint16_t maxTargetCount, MemoryResource* instanceMemRes) {
                     using OutputInstancePointer = UniqueInstance<OutputInstance, RBFBehaviorOutputInstance>;
@@ -40,7 +43,10 @@ class Factory {
                 return factory.create(LODSpec<std::uint16_t>{memRes},
                                       std::move(solvers),
                                       std::move(solverRawControlInputIndices),
-                                      std::move(solverRawControlOutputIndices),
+                                      std::move(solverPoseIndices),
+                                      std::move(poseInputControlIndices),
+                                      std::move(poseOutputControlIndices),
+                                      std::move(poseOutputControlWeights),
                                       maximumInputCount,
                                       maximumTargetCount,
                                       std::move(instanceFactory));
@@ -49,12 +55,7 @@ class Factory {
             auto lods = computeLODs(reader, memRes);
             solvers.reserve(lods.count);
             solverRawControlInputIndices.resize(lods.count);
-            solverRawControlOutputIndices.resize(lods.count);
-
-            const std::uint16_t rbfControlsOffset = static_cast<std::uint16_t>(
-                reader->getRawControlCount() +
-                reader->getPSDCount() +
-                reader->getMLControlCount());
+            solverPoseIndices.resize(lods.count);
 
             RBFSolverRecipe recipe{};
             Vector<float> targetScales{memRes};
@@ -88,10 +89,7 @@ class Factory {
                 recipe.rawControlCount = static_cast<std::uint16_t>(rawControlIndices.size());
 
                 auto poseIndices = reader->getRBFSolverPoseIndices(solverIndex);
-                solverRawControlOutputIndices[solverIndex].reserve(poseIndices.size());
-                for (const std::uint16_t pi : poseIndices) {
-                    solverRawControlOutputIndices[solverIndex].push_back(static_cast<std::uint16_t>(rbfControlsOffset + pi));
-                }
+                solverPoseIndices[solverIndex].assign(poseIndices.begin(), poseIndices.end());
 
                 const auto targetCount = static_cast<std::uint16_t>(poseIndices.size());
                 if (targetCount > maximumTargetCount) {
@@ -108,11 +106,27 @@ class Factory {
                 auto solver = RBFSolver::create(recipe, memRes);
                 solvers.emplace_back(std::move(solver));
             }
+            const auto poseCount = reader->getRBFPoseCount();
+            poseInputControlIndices.resize(poseCount);
+            poseOutputControlIndices.resize(poseCount);
+            poseOutputControlWeights.resize(poseCount);
+
+            for (std::uint16_t poseIndex = {}; poseIndex < poseCount; ++poseIndex) {
+                const auto inputControlIndices = reader->getRBFPoseInputControlIndices(poseIndex);
+                const auto outputControlIndices = reader->getRBFPoseOutputControlIndices(poseIndex);
+                const auto outputControlWeights = reader->getRBFPoseOutputControlWeights(poseIndex);
+                poseInputControlIndices[poseIndex].assign(inputControlIndices.begin(), inputControlIndices.end());
+                poseOutputControlIndices[poseIndex].assign(outputControlIndices.begin(), outputControlIndices.end());
+                poseOutputControlWeights[poseIndex].assign(outputControlWeights.begin(), outputControlWeights.end());
+            }
 
             return factory.create(std::move(lods),
                                   std::move(solvers),
                                   std::move(solverRawControlInputIndices),
-                                  std::move(solverRawControlOutputIndices),
+                                  std::move(solverPoseIndices),
+                                  std::move(poseInputControlIndices),
+                                  std::move(poseOutputControlIndices),
+                                  std::move(poseOutputControlWeights),
                                   maximumInputCount,
                                   maximumTargetCount,
                                   std::move(instanceFactory));
diff --git a/Engine/Plugins/Animation/RigLogic/Source/RigLogicLib/Private/riglogic/riglogic/RigLogicImpl.cpp b/Engine/Plugins/Animation/RigLogic/Source/RigLogicLib/Private/riglogic/riglogic/RigLogicImpl.cpp
--- a/Engine/Plugins/Animation/RigLogic/Source/RigLogicLib/Private/riglogic/riglogic/RigLogicImpl.cpp
+++ b/Engine/Plugins/Animation/RigLogic/Source/RigLogicLib/Private/riglogic/riglogic/RigLogicImpl.cpp
@@ -49,7 +49,7 @@ static RigMetrics::Pointer computeRigMetrics(const dna::Reader* reader, const Co
     metrics->mlControlCount = reader->getMLControlCount();
     metrics->neuralNetworkCount = reader->getNeuralNetworkCount();
     metrics->rbfSolverCount = reader->getRBFSolverCount();
-    metrics->rbfControlCount = reader->getRBFPoseCount();
+    metrics->rbfControlCount = reader->getRBFPoseControlCount();
     return metrics;
 }
 
diff --git a/Engine/Plugins/Animation/RigLogic/Source/RigLogicLib/Public/dna/layers/RBFBehaviorReader.h b/Engine/Plugins/Animation/RigLogic/Source/RigLogicLib/Public/dna/layers/RBFBehaviorReader.h
--- a/Engine/Plugins/Animation/RigLogic/Source/RigLogicLib/Public/dna/layers/RBFBehaviorReader.h
+++ b/Engine/Plugins/Animation/RigLogic/Source/RigLogicLib/Public/dna/layers/RBFBehaviorReader.h
@@ -77,6 +77,62 @@ class DNAAPI RBFBehaviorReader : public virtual BehaviorReader {
                 poseIndex must be less than the value returned by getRBFPoseCount
         */
         virtual float getRBFPoseScale(std::uint16_t poseIndex) const = 0;
+        /**
+            @brief Number of distinct RBF controls.
+            @note
+                This is the number of unique output control indices for all RBF poses.
+            @see getRBFPoseOutputControlIndices
+        */
+        virtual std::uint16_t getRBFPoseControlCount() const = 0;
+        /**
+            @brief Name of the requested RBF pose control.
+            @param poseControlIndex
+                A name's position in the zero-indexed array of RBF pose control names.
+            @warning
+                The index must be less than the value returned by getRBFPoseControlCount.
+            @return View over the control name string.
+        */
+        virtual StringView getRBFPoseControlName(std::uint16_t poseControlIndex) const = 0;
+        /**
+            @brief List of control indices that additionally multiply the pose output (optional).
+            @note
+                When the final output weight of a pose is calculated (this is after it's already
+                multiplied by the weights from @getRBFPoseOutputControlWeights), it can optionally be
+                additionally multiplied by a list of input controls.
+            @param poseIndex
+                A poses's position in the zero-indexed array of poses.
+            @warning
+               poseIndex must be less than the value returned by getRBFPoseCount
+            @return View over the array of control indices.
+        */
+        virtual ConstArrayView<std::uint16_t> getRBFPoseInputControlIndices(std::uint16_t poseIndex) const = 0;
+        /**
+            @brief List of control indices that the specified pose drives.
+            @note
+                This is an additional level of indirection, through which a single RBF pose is able to
+                drive even multiple expressions (by providing multiple control indices for a pose).
+                The solver calculates the output weight for a pose, and that output weight is distributed
+                (and multiplied by the multiplier, @see getRBFPoseOutputControlWeights) to all the controls
+                that are specified in this array.
+            @param poseIndex
+                A poses's position in the zero-indexed array of poses.
+            @warning
+               poseIndex must be less than the value returned by getRBFPoseCount
+            @return View over the array of control indices.
+        */
+        virtual ConstArrayView<std::uint16_t> getRBFPoseOutputControlIndices(std::uint16_t poseIndex) const = 0;
+        /**
+            @brief List of weights that multiply the RBF outputs.
+            @note
+                When an RBF solver calculates the weights for each of its poses, those output weights are
+                additionally multiplied by this fixed list of multipliers (one weight per control is given).
+            @param poseIndex
+                A poses's position in the zero-indexed array of poses.
+            @warning
+               poseIndex must be less than the value returned by getRBFPoseCount
+            @return View over the array of output control multipliers.
+        */
+        virtual ConstArrayView<float> getRBFPoseOutputControlWeights(std::uint16_t poseIndex) const = 0;
         /**
             @brief Number of RBF solvers.
         */
diff --git a/Engine/Plugins/Animation/RigLogic/Source/RigLogicLib/Public/dna/layers/RBFBehaviorWriter.h b/Engine/Plugins/Animation/RigLogic/Source/RigLogicLib/Public/dna/layers/RBFBehaviorWriter.h
--- a/Engine/Plugins/Animation/RigLogic/Source/RigLogicLib/Public/dna/layers/RBFBehaviorWriter.h
+++ b/Engine/Plugins/Animation/RigLogic/Source/RigLogicLib/Public/dna/layers/RBFBehaviorWriter.h
@@ -51,6 +51,79 @@ class DNAAPI RBFBehaviorWriter : public virtual BehaviorWriter {
                 The solver's pose scale.
         */
         virtual void setRBFPoseScale(std::uint16_t poseIndex, float scale) = 0;
+        /**
+            @brief Delete all RBF pose control names.
+        */
+        virtual void clearRBFPoseControlNames() = 0;
+        /**
+            @brief Name of the specified RBF pose control.
+            @param index
+                A name's position in the zero-indexed array of RBF pose control names.
+            @note
+                The control name storage will be implicitly resized (if needed) to provide
+                storage for the number of names that is inferred from the specified index.
+            @param name
+                A null-terminated string.
+            @note
+                The passed in name is copied, which will involve an additional allocation.
+        */
+        virtual void setRBFPoseControlName(std::uint16_t poseControlIndex, const char* name) = 0;
+        /**
+            @brief Set the list of control indices that additionally multiply the pose output (optional).
+            @param poseIndex
+                A poses's position in the zero-indexed array of poses.
+            @warning
+               poseIndex must be less than the value returned by getRBFPoseCount
+            @param controlIndices
+                The source address from which the control indices are to be copied.
+            @note
+                When the final output weight of a pose is calculated (this is after it's already
+                multiplied by the weights from @getRBFPoseOutputControlWeights), it can optionally be
+                additionally multiplied by a list of input controls.
+            @param controlCount
+                The number of control indices to copy.
+        */
+        virtual void setRBFPoseInputControlIndices(std::uint16_t poseIndex,
+                                                   const std::uint16_t* controlIndices,
+                                                   std::uint16_t controlIndexCount) = 0;
+        /**
+            @brief Set the list of control indices that the specified pose drives.
+            @param poseIndex
+                A poses's position in the zero-indexed array of poses.
+            @warning
+               poseIndex must be less than the value returned by getRBFPoseCount
+            @param controlIndices
+                The source address from which the control indices are to be copied.
+            @note
+                These control indices
+                This is an additional level of indirection, through which a single RBF pose is able to
+                drive even multiple expressions (by providing multiple control indices for a pose).
+                The solver calculates the output weight for a pose, and that output weight is distributed
+                (and multiplied by the multiplier, @see setRBFPoseOutputControlWeights) to all the controls
+                that are specified in this array.
+            @param controlCount
+                The number of control indices to copy.
+        */
+        virtual void setRBFPoseOutputControlIndices(std::uint16_t poseIndex,
+                                                    const std::uint16_t* controlIndices,
+                                                    std::uint16_t controlIndexCount) = 0;
+        /**
+            @brief Set the list of weights that multiply the RBF outputs.
+            @param poseIndex
+                A poses's position in the zero-indexed array of poses.
+            @warning
+               poseIndex must be less than the value returned by getRBFPoseCount
+            @param controlWeights
+                The source address from which the control weights are to be copied.
+            @note
+                When an RBF solver calculates the weights for each of its poses, those output weights are
+                additionally multiplied by this fixed list of multipliers (one weight per control is given).
+            @param controlWeightCount
+                The number of control weights to copy.
+        */
+        virtual void setRBFPoseOutputControlWeights(std::uint16_t poseIndex,
+                                                    const float* controlWeights,
+                                                    std::uint16_t controlWeightCount) = 0;
         /**
             @brief Delete all RBF solvers.
         */
diff --git a/Engine/Plugins/Animation/RigLogic/Source/RigLogicLib/Public/dna/version/Version.h b/Engine/Plugins/Animation/RigLogic/Source/RigLogicLib/Public/dna/version/Version.h
--- a/Engine/Plugins/Animation/RigLogic/Source/RigLogicLib/Public/dna/version/Version.h
+++ b/Engine/Plugins/Animation/RigLogic/Source/RigLogicLib/Public/dna/version/Version.h
@@ -3,6 +3,6 @@
 #pragma once
 
 #define DNA_MAJOR_VERSION 9
-#define DNA_MINOR_VERSION 3
-#define DNA_PATCH_VERSION 6
-#define DNA_VERSION_STRING "9.3.6"
+#define DNA_MINOR_VERSION 4
+#define DNA_PATCH_VERSION 0
+#define DNA_VERSION_STRING "9.4.0"
diff --git a/Engine/Plugins/Animation/RigLogic/Source/RigLogicLib/Public/riglogic/version/Version.h b/Engine/Plugins/Animation/RigLogic/Source/RigLogicLib/Public/riglogic/version/Version.h
--- a/Engine/Plugins/Animation/RigLogic/Source/RigLogicLib/Public/riglogic/version/Version.h
+++ b/Engine/Plugins/Animation/RigLogic/Source/RigLogicLib/Public/riglogic/version/Version.h
@@ -3,6 +3,6 @@
 #pragma once
 
 #define RL_MAJOR_VERSION 10
-#define RL_MINOR_VERSION 1
-#define RL_PATCH_VERSION 4
-#define RL_VERSION_STRING "10.1.4"
+#define RL_MINOR_VERSION 2
+#define RL_PATCH_VERSION 0
+#define RL_VERSION_STRING "10.2.0"
diff --git a/Engine/Plugins/Animation/RigLogic/Source/RigLogicLibTest/Private/dnatests/FakeDNAReader.h b/Engine/Plugins/Animation/RigLogic/Source/RigLogicLibTest/Private/dnatests/FakeDNAReader.h
--- a/Engine/Plugins/Animation/RigLogic/Source/RigLogicLibTest/Private/dnatests/FakeDNAReader.h
+++ b/Engine/Plugins/Animation/RigLogic/Source/RigLogicLibTest/Private/dnatests/FakeDNAReader.h
@@ -569,6 +569,26 @@ class FakeDNAReader : public Reader {
             return {};
         }
 
+        std::uint16_t getRBFPoseControlCount() const override {
+            return {};
+        }
+
+        StringView getRBFPoseControlName(std::uint16_t poseControlIndex) const override {
+            return {};
+        }
+
+        ConstArrayView<std::uint16_t> getRBFPoseInputControlIndices(std::uint16_t poseIndex) const override {
+            return {};
+        }
+
+        ConstArrayView<std::uint16_t> getRBFPoseOutputControlIndices(std::uint16_t poseIndex) const override {
+            return {};
+        }
+
+        ConstArrayView<float> getRBFPoseOutputControlWeights(std::uint16_t poseIndex) const override {
+            return {};
+        }
+
         std::uint16_t getRBFSolverCount() const override {
             return {};
         }
diff --git a/Engine/Plugins/Animation/RigLogic/Source/RigLogicLibTest/Private/dnatests/FixturesJSON.cpp b/Engine/Plugins/Animation/RigLogic/Source/RigLogicLibTest/Private/dnatests/FixturesJSON.cpp
--- a/Engine/Plugins/Animation/RigLogic/Source/RigLogicLibTest/Private/dnatests/FixturesJSON.cpp
+++ b/Engine/Plugins/Animation/RigLogic/Source/RigLogicLibTest/Private/dnatests/FixturesJSON.cpp
@@ -15,7 +15,7 @@ const char* jsonDNA =
     },\n\
     \"version\": {\n\
         \"generation\": 2,\n\
-        \"version\": 4\n\
+        \"version\": 5\n\
     },\n\
     \"index\": {\n\
         \"entries\": [{\n\
@@ -48,6 +48,11 @@ const char* jsonDNA =
             \"version\": 65536,\n\
             \"offset\": 0,\n\
             \"size\": 0\n\
+        }, {\n\
+            \"id\": 1919051365,\n\
+            \"version\": 65536,\n\
+            \"offset\": 0,\n\
+            \"size\": 0\n\
         }, {\n\
             \"id\": 1784835428,\n\
             \"version\": 65536,\n\
@@ -204,6 +209,10 @@ const char* jsonDNA =
         \"solvers\": [],\n\
         \"poses\": []\n\
     },\n\
+    \"rbfe1.0\": {\n\
+        \"poseControlNames\": [],\n\
+        \"poses\": []\n\
+    },\n\
     \"jbmd1.0\": {\n\
         \"jointRepresentations\": []\n\
     },\n\
diff --git a/Engine/Plugins/Animation/RigLogic/Source/RigLogicLibTest/Private/dnatests/Fixturesv22.cpp b/Engine/Plugins/Animation/RigLogic/Source/RigLogicLibTest/Private/dnatests/Fixturesv22.cpp
--- a/Engine/Plugins/Animation/RigLogic/Source/RigLogicLibTest/Private/dnatests/Fixturesv22.cpp
+++ b/Engine/Plugins/Animation/RigLogic/Source/RigLogicLibTest/Private/dnatests/Fixturesv22.cpp
@@ -1181,62 +1181,66 @@ const unsigned char RawV2xNewerWithUnknownDataPreservedAndDNARewritten::header[]
     0x00, 0x02,  // Generation
     0xff, 0xff,  // Version
     // Index Table
-    0x00, 0x00, 0x00, 0x0e,  // Index table entry count
+    0x00, 0x00, 0x00, 0x0f,  // Index table entry count
     0x64, 0x65, 0x73, 0x63,  // Descriptor id
     0x00, 0x01, 0x00, 0x02,  // Descriptor version
-    0x00, 0x00, 0x00, 0xeb,  // Descriptor offset
+    0x00, 0x00, 0x00, 0xfb,  // Descriptor offset
     0x00, 0x00, 0x00, 0x10,  // Descriptor size
     0x64, 0x65, 0x66, 0x6e,  // Definition id
     0x00, 0x01, 0x00, 0x02,  // Definition version
-    0x00, 0x00, 0x00, 0xfb,  // Definition offset
+    0x00, 0x00, 0x01, 0x0b,  // Definition offset
     0x00, 0x00, 0x00, 0x10,  // Definition size
     0x62, 0x68, 0x76, 0x72,  // Behavior id
     0x00, 0x01, 0x00, 0x02,  // Behavior version
-    0x00, 0x00, 0x01, 0x0b,  // Behavior offset
+    0x00, 0x00, 0x01, 0x1b,  // Behavior offset
     0x00, 0x00, 0x00, 0x20,  // Behavior size
     0x67, 0x65, 0x6f, 0x6d,  // Geometry id
     0x00, 0x01, 0x00, 0x02,  // Geometry version
-    0x00, 0x00, 0x01, 0x2b,  // Geometry offset
+    0x00, 0x00, 0x01, 0x3b,  // Geometry offset
     0x00, 0x00, 0x00, 0x20,  // Geometry size
     0x75, 0x6e, 0x6b, 0x31,  // Unknown 1 id
     0x00, 0x01, 0x00, 0x03,  // Unknown 1 version
-    0x00, 0x00, 0x01, 0x4b,  // Unknown 1 offset
+    0x00, 0x00, 0x01, 0x5b,  // Unknown 1 offset
     0x00, 0x00, 0x00, 0x10,  // Unknown 1 size
     0x75, 0x6e, 0x6b, 0x32,  // Unknown 2 id
     0x00, 0x01, 0x00, 0x02,  // Unknown 2 version
-    0x00, 0x00, 0x01, 0x5b,  // Unknown 2 offset
+    0x00, 0x00, 0x01, 0x6b,  // Unknown 2 offset
     0x00, 0x00, 0x00, 0x20,  // Unknown 2 size
     0x64, 0x65, 0x73, 0x63,  // Descriptor id
     0x00, 0x01, 0x00, 0x01,  // Descriptor version
-    0x00, 0x00, 0x01, 0x7b,  // Descriptor offset
+    0x00, 0x00, 0x01, 0x8b,  // Descriptor offset
     0x00, 0x00, 0x00, 0x24,  // Descriptor size
     0x64, 0x65, 0x66, 0x6e,  // Definition id
     0x00, 0x01, 0x00, 0x01,  // Definition version
-    0x00, 0x00, 0x01, 0x9f,  // Definition offset
+    0x00, 0x00, 0x01, 0xaf,  // Definition offset
     0x00, 0x00, 0x00, 0x5c,  // Definition size
     0x62, 0x68, 0x76, 0x72,  // Behavior id
     0x00, 0x01, 0x00, 0x01,  // Behavior version
-    0x00, 0x00, 0x01, 0xfb,  // Behavior offset
+    0x00, 0x00, 0x02, 0x0b,  // Behavior offset
     0x00, 0x00, 0x00, 0x56,  // Behavior size
     0x67, 0x65, 0x6f, 0x6d,  // Geometry id
     0x00, 0x01, 0x00, 0x01,  // Geometry version
-    0x00, 0x00, 0x02, 0x51,  // Geometry offset
+    0x00, 0x00, 0x02, 0x61,  // Geometry offset
     0x00, 0x00, 0x00, 0x04,  // Geometry size
     0x6d, 0x6c, 0x62, 0x68,  // Machine learned behavior id
     0x00, 0x01, 0x00, 0x00,  // Machine learned behavior version
-    0x00, 0x00, 0x02, 0x55,  // Machine learned behavior offset
+    0x00, 0x00, 0x02, 0x65,  // Machine learned behavior offset
     0x00, 0x00, 0x00, 0x18,  // Machine learned behavior size
     0x72, 0x62, 0x66, 0x62,  // RBF behavior id
     0x00, 0x01, 0x00, 0x00,  // RBF behavior version
-    0x00, 0x00, 0x02, 0x6d,  // RBF behavior offset
+    0x00, 0x00, 0x02, 0x7d,  // RBF behavior offset
     0x00, 0x00, 0x00, 0x10,  // RBF behavior size
+    0x72, 0x62, 0x66, 0x65,  // RBF behavior ext id
+    0x00, 0x01, 0x00, 0x00,  // RBF behavior ext version
+    0x00, 0x00, 0x02, 0x8d,  // RBF behavior ext offset
+    0x00, 0x00, 0x00, 0x08,  // RBF behavior ext size
     0x6a, 0x62, 0x6d, 0x64,  // Joint behavior metadata id
     0x00, 0x01, 0x00, 0x00,  // Joint behavior metadata version
-    0x00, 0x00, 0x02, 0x7d,  // Joint behavior metadata offset
+    0x00, 0x00, 0x02, 0x95,  // Joint behavior metadata offset
     0x00, 0x00, 0x00, 0x04,  // Joint behavior metadata size
     0x74, 0x77, 0x73, 0x77,  // Twist swing setups id
     0x00, 0x01, 0x00, 0x00,  // Twist swing setups version
-    0x00, 0x00, 0x02, 0x81,  // Twist swing setups offset
+    0x00, 0x00, 0x02, 0x99,  // Twist swing setups offset
     0x00, 0x00, 0x00, 0x08  // Twist swing setups size
 };
 
@@ -1293,6 +1297,11 @@ const unsigned char RawV2xNewerWithUnknownDataPreservedAndDNARewritten::rbfBehav
     0x00, 0x00, 0x00, 0x00
 };
 
+const unsigned char RawV2xNewerWithUnknownDataPreservedAndDNARewritten::rbfBehaviorExt[] = {
+    0x00, 0x00, 0x00, 0x00,
+    0x00, 0x00, 0x00, 0x00
+};
+
 const unsigned char RawV2xNewerWithUnknownDataPreservedAndDNARewritten::jointBehaviorMetaData[] = {
     0x00, 0x00, 0x00, 0x00
 };
@@ -1332,6 +1341,8 @@ std::vector<char> RawV2xNewerWithUnknownDataPreservedAndDNARewritten::getBytes()
     bytes.insert(bytes.end(), machineLearnedBehavior, machineLearnedBehavior + sizeof(machineLearnedBehavior));
     // RBF behavior
     bytes.insert(bytes.end(), rbfBehavior, rbfBehavior + sizeof(rbfBehavior));
+    // RBF behavior ext
+    bytes.insert(bytes.end(), rbfBehaviorExt, rbfBehaviorExt + sizeof(rbfBehaviorExt));
     // Joint behavior metadata
     bytes.insert(bytes.end(), jointBehaviorMetaData, jointBehaviorMetaData + sizeof(jointBehaviorMetaData));
     // Twist swing behavior
@@ -1345,38 +1356,42 @@ const unsigned char RawV2xNewerWithUnknownDataIgnoredAndDNARewritten::header[] =
     0x00, 0x02,  // Generation
     0xff, 0xff,  // Version
     // Index Table
-    0x00, 0x00, 0x00, 0x08,  // Index table entry count
+    0x00, 0x00, 0x00, 0x09,  // Index table entry count
     0x64, 0x65, 0x73, 0x63,  // Descriptor id
     0x00, 0x01, 0x00, 0x01,  // Descriptor version
-    0x00, 0x00, 0x00, 0x8b,  // Descriptor offset
+    0x00, 0x00, 0x00, 0x9b,  // Descriptor offset
     0x00, 0x00, 0x00, 0x24,  // Descriptor size
     0x64, 0x65, 0x66, 0x6e,  // Definition id
     0x00, 0x01, 0x00, 0x01,  // Definition version
-    0x00, 0x00, 0x00, 0xaf,  // Definition offset
+    0x00, 0x00, 0x00, 0xbf,  // Definition offset
     0x00, 0x00, 0x00, 0x5c,  // Definition size
     0x62, 0x68, 0x76, 0x72,  // Behavior id
     0x00, 0x01, 0x00, 0x01,  // Behavior version
-    0x00, 0x00, 0x01, 0x0b,  // Behavior offset
+    0x00, 0x00, 0x01, 0x1b,  // Behavior offset
     0x00, 0x00, 0x00, 0x56,  // Behavior size
     0x67, 0x65, 0x6f, 0x6d,  // Geometry id
     0x00, 0x01, 0x00, 0x01,  // Geometry version
-    0x00, 0x00, 0x01, 0x61,  // Geometry offset
+    0x00, 0x00, 0x01, 0x71,  // Geometry offset
     0x00, 0x00, 0x00, 0x04,  // Geometry size
     0x6d, 0x6c, 0x62, 0x68,  // Machine learned behavior id
     0x00, 0x01, 0x00, 0x00,  // Machine learned behavior version
-    0x00, 0x00, 0x01, 0x65,  // Machine learned behavior offset
+    0x00, 0x00, 0x01, 0x75,  // Machine learned behavior offset
     0x00, 0x00, 0x00, 0x18,  // Machine learned behavior size
     0x72, 0x62, 0x66, 0x62,  // RBF behavior id
     0x00, 0x01, 0x00, 0x00,  // RBF behavior version
-    0x00, 0x00, 0x01, 0x7d,  // RBF behavior offset
+    0x00, 0x00, 0x01, 0x8d,  // RBF behavior offset
     0x00, 0x00, 0x00, 0x10,  // RBF behavior size
+    0x72, 0x62, 0x66, 0x65,  // RBF behavior ext id
+    0x00, 0x01, 0x00, 0x00,  // RBF behavior ext version
+    0x00, 0x00, 0x01, 0x9d,  // RBF behavior ext offset
+    0x00, 0x00, 0x00, 0x08,  // RBF behavior ext size
     0x6a, 0x62, 0x6d, 0x64,  // Joint behavior metadata id
     0x00, 0x01, 0x00, 0x00,  // Joint behavior metadata version
-    0x00, 0x00, 0x01, 0x8d,  // Joint behavior metadata offset
+    0x00, 0x00, 0x01, 0xa5,  // Joint behavior metadata offset
     0x00, 0x00, 0x00, 0x04,  // Joint behavior metadata size
     0x74, 0x77, 0x73, 0x77,  // Twist swing setups id
     0x00, 0x01, 0x00, 0x00,  // Twist swing setups version
-    0x00, 0x00, 0x01, 0x91,  // Twist swing setups offset
+    0x00, 0x00, 0x01, 0xa9,  // Twist swing setups offset
     0x00, 0x00, 0x00, 0x08  // Twist swing setups size
 };
 
@@ -1413,6 +1428,11 @@ std::vector<char> RawV2xNewerWithUnknownDataIgnoredAndDNARewritten::getBytes() {
                  RawV2xNewerWithUnknownDataPreservedAndDNARewritten::rbfBehavior,
                  RawV2xNewerWithUnknownDataPreservedAndDNARewritten::rbfBehavior +
                  sizeof(RawV2xNewerWithUnknownDataPreservedAndDNARewritten::rbfBehavior));
+    // RBF behavior ext
+    bytes.insert(bytes.end(),
+                 RawV2xNewerWithUnknownDataPreservedAndDNARewritten::rbfBehaviorExt,
+                 RawV2xNewerWithUnknownDataPreservedAndDNARewritten::rbfBehaviorExt +
+                 sizeof(RawV2xNewerWithUnknownDataPreservedAndDNARewritten::rbfBehaviorExt));
     // Joint behavior metadata
     bytes.insert(bytes.end(),
                  RawV2xNewerWithUnknownDataPreservedAndDNARewritten::jointBehaviorMetaData,
diff --git a/Engine/Plugins/Animation/RigLogic/Source/RigLogicLibTest/Private/dnatests/Fixturesv22.h b/Engine/Plugins/Animation/RigLogic/Source/RigLogicLibTest/Private/dnatests/Fixturesv22.h
--- a/Engine/Plugins/Animation/RigLogic/Source/RigLogicLibTest/Private/dnatests/Fixturesv22.h
+++ b/Engine/Plugins/Animation/RigLogic/Source/RigLogicLibTest/Private/dnatests/Fixturesv22.h
@@ -47,13 +47,14 @@ struct RawV2xNewer {
 };
 
 struct RawV2xNewerWithUnknownDataPreservedAndDNARewritten {
-    static const unsigned char header[235];
+    static const unsigned char header[251];
     static const unsigned char descriptor[36];
     static const unsigned char definition[92];
     static const unsigned char behavior[86];
     static const unsigned char geometry[4];
     static const unsigned char machineLearnedBehavior[24];
     static const unsigned char rbfBehavior[16];
+    static const unsigned char rbfBehaviorExt[8];
     static const unsigned char jointBehaviorMetaData[4];
     static const unsigned char twistSwingBehavior[8];
 
@@ -61,7 +62,7 @@ struct RawV2xNewerWithUnknownDataPreservedAndDNARewritten {
 };
 
 struct RawV2xNewerWithUnknownDataIgnoredAndDNARewritten {
-    static const unsigned char header[139];
+    static const unsigned char header[155];
 
     static std::vector<char> getBytes();
 };
diff --git a/Engine/Plugins/Animation/RigLogic/Source/RigLogicLibTest/Private/dnatests/Fixturesv25.cpp b/Engine/Plugins/Animation/RigLogic/Source/RigLogicLibTest/Private/dnatests/Fixturesv25.cpp
new file mode 100644
--- /dev/null
+++ b/Engine/Plugins/Animation/RigLogic/Source/RigLogicLibTest/Private/dnatests/Fixturesv25.cpp
@@ -0,0 +1,3525 @@
+// Copyright Epic Games, Inc. All Rights Reserved.
+
+#include "dnatests/Fixturesv25.h"
+
+#include "dna/Reader.h"
+#include "pma/TypeDefs.h"
+
+#pragma warning(disable : 4503)
+
+namespace dna {
+
+const unsigned char RawV25::header[] = {
+    0x44, 0x4e, 0x41,  // DNA signature
+    0x00, 0x02,  // Generation
+    0x00, 0x05,  // Version
+    // Index Table
+    0x00, 0x00, 0x00, 0x09,  // Index table entry count
+    0x64, 0x65, 0x73, 0x63,  // Descriptor id
+    0x00, 0x01, 0x00, 0x01,  // Descriptor version
+    0x00, 0x00, 0x00, 0x9b,  // Descriptor offset
+    0x00, 0x00, 0x00, 0x57,  // Descriptor size
+    0x64, 0x65, 0x66, 0x6e,  // Definition id
+    0x00, 0x01, 0x00, 0x01,  // Definition version
+    0x00, 0x00, 0x00, 0xf2,  // Definition offset
+    0x00, 0x00, 0x03, 0x1a,  // Definition size
+    0x62, 0x68, 0x76, 0x72,  // Behavior id
+    0x00, 0x01, 0x00, 0x01,  // Behavior version
+    0x00, 0x00, 0x04, 0x0c,  // Behavior offset
+    0x00, 0x00, 0x05, 0x46,  // Behavior size
+    0x67, 0x65, 0x6f, 0x6d,  // Geometry id
+    0x00, 0x01, 0x00, 0x01,  // Geometry version
+    0x00, 0x00, 0x09, 0x52,  // Geometry offset
+    0x00, 0x00, 0x04, 0x38,  // Geometry size
+    0x6d, 0x6c, 0x62, 0x68,  // Machine learned behavior id
+    0x00, 0x01, 0x00, 0x00,  // Machine learned behavior version
+    0x00, 0x00, 0x0d, 0x8a,  // Machine learned behavior offset
+    0x00, 0x00, 0x02, 0xfa,  // Machine learned behavior size
+    0x72, 0x62, 0x66, 0x62,  // RBF behavior id
+    0x00, 0x01, 0x00, 0x00,  // RBF behavior version
+    0x00, 0x00, 0x10, 0x84,  // RBF behavior offset
+    0x00, 0x00, 0x01, 0x47,  // RBF behavior size
+    0x72, 0x62, 0x66, 0x65,  // RBF behavior ext id
+    0x00, 0x01, 0x00, 0x00,  // RBF behavior ext version
+    0x00, 0x00, 0x11, 0xcb,  // RBF behavior ext offset
+    0x00, 0x00, 0x00, 0xe4,  // RBF behavior ext size
+    0x6a, 0x62, 0x6d, 0x64,  // Joint behavior metadata id
+    0x00, 0x01, 0x00, 0x00,  // Joint behavior metadata version
+    0x00, 0x00, 0x12, 0xaf,  // Joint behavior metadata offset
+    0x00, 0x00, 0x00, 0x3a,  // Joint behavior metadata size
+    0x74, 0x77, 0x73, 0x77,  // Twist swing setups id
+    0x00, 0x01, 0x00, 0x00,  // Twist swing setups version
+    0x00, 0x00, 0x12, 0xe9,  // Twist swing setups offset
+    0x00, 0x00, 0x00, 0xc8  // Twist swing setups size
+};
+
+const unsigned char RawV25::descriptor[] = {
+    0x00, 0x00, 0x00, 0x04,  // Name length
+    0x74, 0x65, 0x73, 0x74,  // Name
+    0x00, 0x05,  // Archetype
+    0x00, 0x02,  // Gender
+    0x00, 0x2a,  // Age
+    0x00, 0x00, 0x00, 0x02,  // Metadata count
+    0x00, 0x00, 0x00, 0x05,  // Metadata key length
+    0x6b, 0x65, 0x79, 0x2d, 0x41,  // Metadata key: "key-A"
+    0x00, 0x00, 0x00, 0x07,  // Metadata value length
+    0x76, 0x61, 0x6c, 0x75, 0x65, 0x2d, 0x41,  // Metadata value: "value-A"
+    0x00, 0x00, 0x00, 0x05,  // Metadata key length
+    0x6b, 0x65, 0x79, 0x2d, 0x42,  // Metadata key: "key-B"
+    0x00, 0x00, 0x00, 0x07,  // Metadata value length
+    0x76, 0x61, 0x6c, 0x75, 0x65, 0x2d, 0x42,  // Metadata value: "value-B"
+    0x00, 0x01,  // Unit translation
+    0x00, 0x01,  // Unit rotation
+    0x00, 0x01,  // Coordinate system x-axis
+    0x00, 0x02,  // Coordinate system y-axis
+    0x00, 0x04,  // Coordinate system z-axis
+    0x00, 0x02,  // LOD Count
+    0x00, 0x00,  // MaxLOD: 0
+    0x00, 0x00, 0x00, 0x01,  // Complexity name length
+    0x41,  // 'A' - Complexity name
+    0x00, 0x00, 0x00, 0x06,  // DB name length
+    0x74, 0x65, 0x73, 0x74, 0x44, 0x42  // Name
+};
+
+const unsigned char RawV25::definition[] = {
+    0x00, 0x00, 0x00, 0x02,  // Joint name indices lod to row mapping length
+    0x00, 0x00,  // Map from LOD-0 to row 0 in below defined matrix
+    0x00, 0x01,  // Map from LOD-1 to row 1 in below defined matrix
+    0x00, 0x00, 0x00, 0x02,  // Joint name indices per LOD row count
+    0x00, 0x00, 0x00, 0x09,  // Indices matrix row-0
+    0x00, 0x00,  // Joint name index: 0
+    0x00, 0x01,  // Joint name index: 1
+    0x00, 0x02,  // Joint name index: 2
+    0x00, 0x03,  // Joint name index: 3
+    0x00, 0x04,  // Joint name index: 4
+    0x00, 0x05,  // Joint name index: 5
+    0x00, 0x06,  // Joint name index: 6
+    0x00, 0x07,  // Joint name index: 7
+    0x00, 0x08,  // Joint name index: 8
+    0x00, 0x00, 0x00, 0x06,  // Indices matrix row-1
+    0x00, 0x00,  // Joint name index: 0
+    0x00, 0x01,  // Joint name index: 1
+    0x00, 0x02,  // Joint name index: 2
+    0x00, 0x03,  // Joint name index: 3
+    0x00, 0x06,  // Joint name index: 6
+    0x00, 0x08,  // Joint name index: 8
+    0x00, 0x00, 0x00, 0x02,  // Blend shape name indices lod to row mapping length
+    0x00, 0x00,  // Map from LOD-0 to row 0 in below defined matrix
+    0x00, 0x01,  // Map from LOD-1 to row 1 in below defined matrix
+    0x00, 0x00, 0x00, 0x02,  // Blend shape name indices per LOD row count
+    0x00, 0x00, 0x00, 0x09,  // Indices matrix row-0
+    0x00, 0x00,  // Blend shape name index: 0
+    0x00, 0x01,  // Blend shape name index: 1
+    0x00, 0x02,  // Blend shape name index: 2
+    0x00, 0x03,  // Blend shape name index: 3
+    0x00, 0x04,  // Blend shape name index: 4
+    0x00, 0x05,  // Blend shape name index: 5
+    0x00, 0x06,  // Blend shape name index: 6
+    0x00, 0x07,  // Blend shape name index: 7
+    0x00, 0x08,  // Blend shape name index: 8
+    0x00, 0x00, 0x00, 0x04,  // Indices matrix row-1
+    0x00, 0x02,  // Blend shape name index: 2
+    0x00, 0x05,  // Blend shape name index: 5
+    0x00, 0x07,  // Blend shape name index: 7
+    0x00, 0x08,  // Blend shape name index: 8
+    0x00, 0x00, 0x00, 0x02,  // Animated map name indices lod to row mapping length
+    0x00, 0x00,  // Map from LOD-0 to row 0 in below defined matrix
+    0x00, 0x01,  // Map from LOD-1 to row 1 in below defined matrix
+    0x00, 0x00, 0x00, 0x02,  // Animated map name indices per LOD row count
+    0x00, 0x00, 0x00, 0x0a,  // Indices matrix row-0
+    0x00, 0x00,  // Animated map name index: 0
+    0x00, 0x01,  // Animated map name index: 1
+    0x00, 0x02,  // Animated map name index: 2
+    0x00, 0x03,  // Animated map name index: 3
+    0x00, 0x04,  // Animated map name index: 4
+    0x00, 0x05,  // Animated map name index: 5
+    0x00, 0x06,  // Animated map name index: 6
+    0x00, 0x07,  // Animated map name index: 7
+    0x00, 0x08,  // Animated map name index: 8
+    0x00, 0x09,  // Animated map name index: 9
+    0x00, 0x00, 0x00, 0x04,  // Indices matrix row-1
+    0x00, 0x02,  // Animated map name index: 2
+    0x00, 0x05,  // Animated map name index: 5
+    0x00, 0x07,  // Animated map name index: 7
+    0x00, 0x08,  // Animated map name index: 8
+    0x00, 0x00, 0x00, 0x02,  // Mesh name indices lod to row mapping length
+    0x00, 0x00,  // Map from LOD-0 to row 0 in below defined matrix
+    0x00, 0x01,  // Map from LOD-1 to row 1 in below defined matrix
+    0x00, 0x00, 0x00, 0x02,  // Mesh name indices per LOD row count
+    0x00, 0x00, 0x00, 0x02,  // Indices matrix row-0
+    0x00, 0x00,  // Mesh name index: 0
+    0x00, 0x01,  // Mesh name index: 1
+    0x00, 0x00, 0x00, 0x01,  // Indices matrix row-1
+    0x00, 0x02,  // Mesh name index: 2
+    0x00, 0x00, 0x00, 0x09,  // Gui control names length
+    0x00, 0x00, 0x00, 0x02,  // Gui control name 0 length
+    0x47, 0x41,  // Gui control name 0 : GA
+    0x00, 0x00, 0x00, 0x02,  // Gui control name 1 length
+    0x47, 0x42,  // Gui control name 1 : GB
+    0x00, 0x00, 0x00, 0x02,  // Gui control name 2 length
+    0x47, 0x43,  // Gui control name 2 : GC
+    0x00, 0x00, 0x00, 0x02,  // Gui control name 3 length
+    0x47, 0x44,  // Gui control name 3 : GD
+    0x00, 0x00, 0x00, 0x02,  // Gui control name 4 length
+    0x47, 0x45,  // Gui control name 4 : GE
+    0x00, 0x00, 0x00, 0x02,  // Gui control name 5 length
+    0x47, 0x46,  // Gui control name 5 : GF
+    0x00, 0x00, 0x00, 0x02,  // Gui control name 6 length
+    0x47, 0x47,  // Gui control name 6 : GG
+    0x00, 0x00, 0x00, 0x02,  // Gui control name 7 length
+    0x47, 0x48,  // Gui control name 7 : GH
+    0x00, 0x00, 0x00, 0x02,  // Gui control name 8 length
+    0x47, 0x49,  // Gui control name 8 : GI
+    0x00, 0x00, 0x00, 0x09,  // Raw control names length
+    0x00, 0x00, 0x00, 0x02,  // Raw control name 0 length
+    0x52, 0x41,  // Raw control name 0 : RA
+    0x00, 0x00, 0x00, 0x02,  // Raw control name 1 length
+    0x52, 0x42,  // Raw control name 1 : RB
+    0x00, 0x00, 0x00, 0x02,  // Raw control name 2 length
+    0x52, 0x43,  // Raw control name 2 : RC
+    0x00, 0x00, 0x00, 0x02,  // Raw control name 3 length
+    0x52, 0x44,  // Raw control name 3 : RD
+    0x00, 0x00, 0x00, 0x02,  // Raw control name 4 length
+    0x52, 0x45,  // Raw control name 4 : RE
+    0x00, 0x00, 0x00, 0x02,  // Raw control name 5 length
+    0x52, 0x46,  // Raw control name 5 : RF
+    0x00, 0x00, 0x00, 0x02,  // Raw control name 6 length
+    0x52, 0x47,  // Raw control name 6 : RG
+    0x00, 0x00, 0x00, 0x02,  // Raw control name 7 length
+    0x52, 0x48,  // Raw control name 7 : RH
+    0x00, 0x00, 0x00, 0x02,  // Raw control name 8 length
+    0x52, 0x49,  // Raw control name 8 : RI
+    0x00, 0x00, 0x00, 0x09,  // Joint names length
+    0x00, 0x00, 0x00, 0x02,  // Joint name 0 length
+    0x4a, 0x41,  // Joint name 0 : JA
+    0x00, 0x00, 0x00, 0x02,  // Joint name 1 length
+    0x4a, 0x42,  // Joint name 1 : JB
+    0x00, 0x00, 0x00, 0x02,  // Joint name 2 length
+    0x4a, 0x43,  // Joint name 2 : JC
+    0x00, 0x00, 0x00, 0x02,  // Joint name 3 length
+    0x4a, 0x44,  // Joint name 3 : JD
+    0x00, 0x00, 0x00, 0x02,  // Joint name 4 length
+    0x4a, 0x45,  // Joint name 4 : JE
+    0x00, 0x00, 0x00, 0x02,  // Joint name 5 length
+    0x4a, 0x46,  // Joint name 5 : JF
+    0x00, 0x00, 0x00, 0x02,  // Joint name 6 length
+    0x4a, 0x47,  // Joint name 6 : JG
+    0x00, 0x00, 0x00, 0x02,  // Joint name 7 length
+    0x4a, 0x48,  // Joint name 7 : JH
+    0x00, 0x00, 0x00, 0x02,  // Joint name 8 length
+    0x4a, 0x49,  // Joint name 8 : JI
+    0x00, 0x00, 0x00, 0x09,  // BlendShape names length
+    0x00, 0x00, 0x00, 0x02,  // Blendshape name 0 length
+    0x42, 0x41,  // Blendshape name 0 : BA
+    0x00, 0x00, 0x00, 0x02,  // Blendshape name 1 length
+    0x42, 0x42,  // Blendshape name 1 : BB
+    0x00, 0x00, 0x00, 0x02,  // Blendshape name 2 length
+    0x42, 0x43,  // Blendshape name 2 : BC
+    0x00, 0x00, 0x00, 0x02,  // Blendshape name 3 length
+    0x42, 0x44,  // Blendshape name 3 : BD
+    0x00, 0x00, 0x00, 0x02,  // Blendshape name 4 length
+    0x42, 0x45,  // Blendshape name 4 : BE
+    0x00, 0x00, 0x00, 0x02,  // Blendshape name 5 length
+    0x42, 0x46,  // Blendshape name 5 : BF
+    0x00, 0x00, 0x00, 0x02,  // Blendshape name 6 length
+    0x42, 0x47,  // Blendshape name 6 : BG
+    0x00, 0x00, 0x00, 0x02,  // Blendshape name 7 length
+    0x42, 0x48,  // Blendshape name 7 : BH
+    0x00, 0x00, 0x00, 0x02,  // Blendshape name 8 length
+    0x42, 0x49,  // Blendshape name 8 : BI
+    0x00, 0x00, 0x00, 0x0a,  // Animated Map names length
+    0x00, 0x00, 0x00, 0x02,  // Animated Map name 0 length
+    0x41, 0x41,  // Animated Map name 0 : AA
+    0x00, 0x00, 0x00, 0x02,  // Animated Map name 1 length
+    0x41, 0x42,  // Animated Map name 1 : AB
+    0x00, 0x00, 0x00, 0x02,  // Animated Map name 2 length
+    0x41, 0x43,  // Animated Map name 2 : AC
+    0x00, 0x00, 0x00, 0x02,  // Animated Map name 3 length
+    0x41, 0x44,  // Animated Map name 3 : AD
+    0x00, 0x00, 0x00, 0x02,  // Animated Map name 4 length
+    0x41, 0x45,  // Animated Map name 4 : AE
+    0x00, 0x00, 0x00, 0x02,  // Animated Map name 5 length
+    0x41, 0x46,  // Animated Map name 5 : AF
+    0x00, 0x00, 0x00, 0x02,  // Animated Map name 6 length
+    0x41, 0x47,  // Animated Map name 6 : AG
+    0x00, 0x00, 0x00, 0x02,  // Animated Map name 7 length
+    0x41, 0x48,  // Animated Map name 7 : AH
+    0x00, 0x00, 0x00, 0x02,  // Animated Map name 8 length
+    0x41, 0x49,  // Animated Map name 8 : AI
+    0x00, 0x00, 0x00, 0x02,  // Animated Map name 9 length
+    0x41, 0x4a,  // Animated Map name 8 : AJ
+    0x00, 0x00, 0x00, 0x03,  // Mesh names length
+    0x00, 0x00, 0x00, 0x02,  // Mesh name 0 length
+    0x4d, 0x41,  // Mesh name 0 : MA
+    0x00, 0x00, 0x00, 0x02,  // Mesh name 1 length
+    0x4d, 0x42,  // Mesh name 1 : MB
+    0x00, 0x00, 0x00, 0x02,  // Mesh name 2 length
+    0x4d, 0x43,  // Mesh name 2 : MC
+    0x00, 0x00, 0x00, 0x09,  // Mesh indices length for mesh -> blendShape mapping
+    0x00, 0x00,  // Mesh index 0
+    0x00, 0x00,  // Mesh index 0
+    0x00, 0x00,  // Mesh index 0
+    0x00, 0x01,  // Mesh index 1
+    0x00, 0x01,  // Mesh index 1
+    0x00, 0x01,  // Mesh index 1
+    0x00, 0x01,  // Mesh index 1
+    0x00, 0x02,  // Mesh index 2
+    0x00, 0x02,  // Mesh index 2
+    0x00, 0x00, 0x00, 0x09,  // BlendShape indices length for mesh -> blendShape mapping
+    0x00, 0x00,  // BlendShape 0
+    0x00, 0x01,  // BlendShape 1
+    0x00, 0x02,  // BlendShape 2
+    0x00, 0x03,  // BlendShape 3
+    0x00, 0x04,  // BlendShape 4
+    0x00, 0x05,  // BlendShape 5
+    0x00, 0x06,  // BlendShape 6
+    0x00, 0x07,  // BlendShape 7
+    0x00, 0x08,  // BlendShape 8
+    0x00, 0x00, 0x00, 0x09,  // Joint hierarchy length
+    0x00, 0x00,  // JA - root
+    0x00, 0x00,  // JB
+    0x00, 0x00,  // JC
+    0x00, 0x01,  // JD
+    0x00, 0x01,  // JE
+    0x00, 0x04,  // JF
+    0x00, 0x02,  // JG
+    0x00, 0x04,  // JH
+    0x00, 0x02,  // JI
+    0x00, 0x00, 0x00, 0x09,  // Neutral joint translation X values length
+    0x3f, 0x80, 0x00, 0x00,  // 1.0f
+    0x40, 0x00, 0x00, 0x00,  // 2.0f
+    0x40, 0x40, 0x00, 0x00,  // 3.0f
+    0x40, 0x80, 0x00, 0x00,  // 4.0f
+    0x40, 0xa0, 0x00, 0x00,  // 5.0f
+    0x40, 0xc0, 0x00, 0x00,  // 6.0f
+    0x40, 0xe0, 0x00, 0x00,  // 7.0f
+    0x41, 0x00, 0x00, 0x00,  // 8.0f
+    0x41, 0x10, 0x00, 0x00,  // 9.0f
+    0x00, 0x00, 0x00, 0x09,  // Neutral joint translation Y values length
+    0x3f, 0x80, 0x00, 0x00,  // 1.0f
+    0x40, 0x00, 0x00, 0x00,  // 2.0f
+    0x40, 0x40, 0x00, 0x00,  // 3.0f
+    0x40, 0x80, 0x00, 0x00,  // 4.0f
+    0x40, 0xa0, 0x00, 0x00,  // 5.0f
+    0x40, 0xc0, 0x00, 0x00,  // 6.0f
+    0x40, 0xe0, 0x00, 0x00,  // 7.0f
+    0x41, 0x00, 0x00, 0x00,  // 8.0f
+    0x41, 0x10, 0x00, 0x00,  // 9.0f
+    0x00, 0x00, 0x00, 0x09,  // Neutral joint translation Z values length
+    0x3f, 0x80, 0x00, 0x00,  // 1.0f
+    0x40, 0x00, 0x00, 0x00,  // 2.0f
+    0x40, 0x40, 0x00, 0x00,  // 3.0f
+    0x40, 0x80, 0x00, 0x00,  // 4.0f
+    0x40, 0xa0, 0x00, 0x00,  // 5.0f
+    0x40, 0xc0, 0x00, 0x00,  // 6.0f
+    0x40, 0xe0, 0x00, 0x00,  // 7.0f
+    0x41, 0x00, 0x00, 0x00,  // 8.0f
+    0x41, 0x10, 0x00, 0x00,  // 9.0f
+    0x00, 0x00, 0x00, 0x09,  // Neutral joint rotation X values length
+    0x3f, 0x80, 0x00, 0x00,  // 1.0f
+    0x40, 0x00, 0x00, 0x00,  // 2.0f
+    0x40, 0x40, 0x00, 0x00,  // 3.0f
+    0x40, 0x80, 0x00, 0x00,  // 4.0f
+    0x40, 0xa0, 0x00, 0x00,  // 5.0f
+    0x40, 0xc0, 0x00, 0x00,  // 6.0f
+    0x40, 0xe0, 0x00, 0x00,  // 7.0f
+    0x41, 0x00, 0x00, 0x00,  // 8.0f
+    0x41, 0x10, 0x00, 0x00,  // 9.0f
+    0x00, 0x00, 0x00, 0x09,  // Neutral joint rotation Y values length
+    0x3f, 0x80, 0x00, 0x00,  // 1.0f
+    0x40, 0x00, 0x00, 0x00,  // 2.0f
+    0x40, 0x40, 0x00, 0x00,  // 3.0f
+    0x40, 0x80, 0x00, 0x00,  // 4.0f
+    0x40, 0xa0, 0x00, 0x00,  // 5.0f
+    0x40, 0xc0, 0x00, 0x00,  // 6.0f
+    0x40, 0xe0, 0x00, 0x00,  // 7.0f
+    0x41, 0x00, 0x00, 0x00,  // 8.0f
+    0x41, 0x10, 0x00, 0x00,  // 9.0f
+    0x00, 0x00, 0x00, 0x09,  // Neutral joint rotation Z values length
+    0x3f, 0x80, 0x00, 0x00,  // 1.0f
+    0x40, 0x00, 0x00, 0x00,  // 2.0f
+    0x40, 0x40, 0x00, 0x00,  // 3.0f
+    0x40, 0x80, 0x00, 0x00,  // 4.0f
+    0x40, 0xa0, 0x00, 0x00,  // 5.0f
+    0x40, 0xc0, 0x00, 0x00,  // 6.0f
+    0x40, 0xe0, 0x00, 0x00,  // 7.0f
+    0x41, 0x00, 0x00, 0x00,  // 8.0f
+    0x41, 0x10, 0x00, 0x00  // 9.0f
+};
+
+const unsigned char RawV25::conditionals[] {
+    // Input indices
+    0x00, 0x00, 0x00, 0x0f,  // Input indices count
+    0x00, 0x00,  // Index: 0      C1  L0  L1
+    0x00, 0x01,  // Index: 1  C0  C1  L0  L1
+    0x00, 0x01,  // Index: 1  C0  C1  L0  L1
+    0x00, 0x02,  // Index: 2  C0      L0  L1
+    0x00, 0x03,  // Index: 3      C1  L0  L1
+    0x00, 0x03,  // Index: 3      C1  L0  L1
+    0x00, 0x04,  // Index: 4  C0      L0
+    0x00, 0x04,  // Index: 4  C0      L0
+    0x00, 0x04,  // Index: 4  C0      L0
+    0x00, 0x05,  // Index: 5      C1  L0
+    0x00, 0x06,  // Index: 6      C1  L0
+    0x00, 0x07,  // Index: 7  C0      L0
+    0x00, 0x07,  // Index: 7  C0      L0
+    0x00, 0x08,  // Index: 8  C0  C1  L0
+    0x00, 0x08,  // Index: 8      C1  L0
+    // Output indices
+    0x00, 0x00, 0x00, 0x0f,  // Output indices count
+    0x00, 0x00,  // Index: 0      C1  L0  L1
+    0x00, 0x01,  // Index: 1  C0  C1  L0  L1
+    0x00, 0x01,  // Index: 1  C0  C1  L0  L1
+    0x00, 0x02,  // Index: 2  C0      L0  L1
+    0x00, 0x03,  // Index: 3      C1  L0  L1
+    0x00, 0x03,  // Index: 3      C1  L0  L1
+    0x00, 0x04,  // Index: 4  C0      L0
+    0x00, 0x04,  // Index: 4  C0      L0
+    0x00, 0x04,  // Index: 4  C0      L0
+    0x00, 0x05,  // Index: 5      C1  L0
+    0x00, 0x06,  // Index: 6      C1  L0
+    0x00, 0x07,  // Index: 7  C0      L0
+    0x00, 0x07,  // Index: 7  C0      L0
+    0x00, 0x08,  // Index: 8  C0  C1  L0
+    0x00, 0x08,  // Index: 8      C1  L0
+    // From values
+    0x00, 0x00, 0x00, 0x0f,  // From values count
+    0x00, 0x00, 0x00, 0x00,  // 0.0f      C1  L0  L1
+    0x00, 0x00, 0x00, 0x00,  // 0.0f  C0  C1  L0  L1
+    0x3f, 0x19, 0x99, 0x9a,  // 0.6f  C0  C1  L0  L1
+    0x3e, 0xcc, 0xcc, 0xcd,  // 0.4f  C0      L0  L1
+    0x3d, 0xcc, 0xcc, 0xcd,  // 0.1f      C1  L0  L1
+    0x3f, 0x33, 0x33, 0x33,  // 0.7f      C1  L0  L1
+    0x00, 0x00, 0x00, 0x00,  // 0.0f  C0      L0
+    0x3e, 0xcc, 0xcc, 0xcd,  // 0.4f  C0      L0
+    0x3f, 0x33, 0x33, 0x33,  // 0.7f  C0      L0
+    0x3f, 0x00, 0x00, 0x00,  // 0.5f      C1  L0
+    0x00, 0x00, 0x00, 0x00,  // 0.0f      C1  L0
+    0x3d, 0xcc, 0xcc, 0xcd,  // 0.1f  C0      L0
+    0x3f, 0x19, 0x99, 0x9a,  // 0.6f  C0      L0
+    0x3e, 0x4c, 0xcc, 0xcd,  // 0.2f  C0  C1  L0
+    0x00, 0x00, 0x00, 0x00,  // 0.0f      C1  L0
+    // To values
+    0x00, 0x00, 0x00, 0x0f,  // To values count
+    0x3f, 0x80, 0x00, 0x00,  // 1.0f      C1  L0  L1
+    0x3f, 0x19, 0x99, 0x9a,  // 0.6f  C0  C1  L0  L1
+    0x3f, 0x80, 0x00, 0x00,  // 1.0f  C0  C1  L0  L1
+    0x3f, 0x66, 0x66, 0x66,  // 0.9f  C0      L0  L1
+    0x3f, 0x33, 0x33, 0x33,  // 0.7f      C1  L0  L1
+    0x3f, 0x80, 0x00, 0x00,  // 1.0f      C1  L0  L1
+    0x3e, 0xcc, 0xcc, 0xcd,  // 0.4f  C0      L0
+    0x3f, 0x33, 0x33, 0x33,  // 0.7f  C0      L0
+    0x3f, 0x80, 0x00, 0x00,  // 1.0f  C0      L0
+    0x3f, 0x80, 0x00, 0x00,  // 1.0f      C1  L0
+    0x3f, 0x80, 0x00, 0x00,  // 1.0f      C1  L0
+    0x3f, 0x19, 0x99, 0x9a,  // 0.6f  C0      L0
+    0x3f, 0x80, 0x00, 0x00,  // 1.0f  C0      L0
+    0x3f, 0x4c, 0xcc, 0xcd,  // 0.8f  C0  C1  L0
+    0x3f, 0x80, 0x00, 0x00,  // 1.0f      C1  L0
+    // Slope values
+    0x00, 0x00, 0x00, 0x0f,  // Slope values count
+    0x3f, 0x80, 0x00, 0x00,  // 1.0f      C1  L0  L1
+    0x3f, 0x66, 0x66, 0x66,  // 0.9f  C0  C1  L0  L1
+    0x3f, 0x66, 0x66, 0x66,  // 0.9f  C0  C1  L0  L1
+    0x3f, 0x4c, 0xcc, 0xcd,  // 0.8f  C0      L0  L1
+    0x3f, 0x33, 0x33, 0x33,  // 0.7f      C1  L0  L1
+    0x3f, 0x33, 0x33, 0x33,  // 0.7f      C1  L0  L1
+    0x3f, 0x19, 0x99, 0x9a,  // 0.6f  C0      L0
+    0x3f, 0x19, 0x99, 0x9a,  // 0.6f  C0      L0
+    0x3f, 0x19, 0x99, 0x9a,  // 0.6f  C0      L0
+    0x3f, 0x00, 0x00, 0x00,  // 0.5f      C1  L0
+    0x3f, 0x19, 0x99, 0x9a,  // 0.6f      C1  L0
+    0x3f, 0x33, 0x33, 0x33,  // 0.7f  C0      L0
+    0x3f, 0x33, 0x33, 0x33,  // 0.7f  C0      L0
+    0x3f, 0x4c, 0xcc, 0xcd,  // 0.8f  C0  C1  L0
+    0x3f, 0x66, 0x66, 0x66,  // 0.9f      C1  L0
+    // Cut values
+    0x00, 0x00, 0x00, 0x0f,  // Cut values count
+    0x00, 0x00, 0x00, 0x00,  // 0.0f      C1  L0  L1
+    0x3f, 0x00, 0x00, 0x00,  // 0.5f  C0  C1  L0  L1
+    0x3f, 0x00, 0x00, 0x00,  // 0.5f  C0  C1  L0  L1
+    0x3e, 0xcc, 0xcc, 0xcd,  // 0.4f  C0      L0  L1
+    0x3e, 0x99, 0x99, 0x9a,  // 0.3f      C1  L0  L1
+    0x3e, 0x99, 0x99, 0x9a,  // 0.3f      C1  L0  L1
+    0x3f, 0x80, 0x00, 0x00,  // 1.0f  C0      L0
+    0x3f, 0x80, 0x00, 0x00,  // 1.0f  C0      L0
+    0x3f, 0x80, 0x00, 0x00,  // 1.0f  C0      L0
+    0x3e, 0x4c, 0xcc, 0xcd,  // 0.2f      C1  L0
+    0x3e, 0xcc, 0xcc, 0xcd,  // 0.4f      C1  L0
+    0x3f, 0x4c, 0xcc, 0xcd,  // 0.8f  C0      L0
+    0x3f, 0x4c, 0xcc, 0xcd,  // 0.8f  C0      L0
+    0x3f, 0x80, 0x00, 0x00,  // 1.0f  C0  C1  L0
+    0x3e, 0x4c, 0xcc, 0xcd  // 0.2f       C1  L0
+};
+
+const unsigned char RawV25::psds[] {
+    // Rows
+    0x00, 0x00, 0x00, 0x18,  // Row index count
+    0x00, 0x08,  // Index:  8  C1
+    0x00, 0x08,  // Index:  8  C1
+    0x00, 0x08,  // Index:  8  C1
+    0x00, 0x09,  // Index:  9      C2
+    0x00, 0x09,  // Index:  9      C2
+    0x00, 0x0a,  // Index: 10  C1
+    0x00, 0x0a,  // Index: 10  C1
+    0x00, 0x0a,  // Index: 10  C1
+    0x00, 0x0b,  // Index: 11      C2
+    0x00, 0x0c,  // Index: 12      C2
+    0x00, 0x0d,  // Index: 13  C1
+    0x00, 0x0d,  // Index: 13  C1
+    0x00, 0x0d,  // Index: 13  C1
+    0x00, 0x0e,  // Index: 14  C1
+    0x00, 0x0e,  // Index: 14  C1
+    0x00, 0x0f,  // Index: 15  C1
+    0x00, 0x10,  // Index: 16      C2
+    0x00, 0x12,  // Index: 18      C2
+    0x00, 0x12,  // Index: 18      C2
+    0x00, 0x12,  // Index: 18      C2
+    0x00, 0x12,  // Index: 18      C2
+    0x00, 0x13,  // Index: 19  C1
+    0x00, 0x13,  // Index: 19  C1
+    0x00, 0x14,  // Index: 20  C1
+    // Columns
+    0x00, 0x00, 0x00, 0x18,  // Column index count
+    0x00, 0x00,  // Index: 0      C2
+    0x00, 0x03,  // Index: 3      C2
+    0x00, 0x06,  // Index: 6      C2
+    0x00, 0x02,  // Index: 2  C1
+    0x00, 0x05,  // Index: 5      C2
+    0x00, 0x02,  // Index: 2  C1
+    0x00, 0x03,  // Index: 3      C2
+    0x00, 0x07,  // Index: 7  C1
+    0x00, 0x03,  // Index: 3      C2
+    0x00, 0x02,  // Index: 2  C1
+    0x00, 0x00,  // Index: 0      C2
+    0x00, 0x01,  // Index: 1  C1  C2
+    0x00, 0x02,  // Index: 2  C1
+    0x00, 0x03,  // Index: 3      C2
+    0x00, 0x06,  // Index: 6      C2
+    0x00, 0x00,  // Index: 0      C2
+    0x00, 0x04,  // Index: 4  C1
+    0x00, 0x00,  // Index: 0      C2
+    0x00, 0x03,  // Index: 3      C2
+    0x00, 0x04,  // Index: 4  C1
+    0x00, 0x05,  // Index: 5      C2
+    0x00, 0x06,  // Index: 6      C2
+    0x00, 0x07,  // Index: 7  C1
+    0x00, 0x02,  // Index: 2  C1
+    // Values
+    0x00, 0x00, 0x00, 0x18,  // Value count
+    0x3f, 0x80, 0x00, 0x00,  // 1.0f
+    0x3f, 0x66, 0x66, 0x66,  // 0.9f
+    0x3f, 0x66, 0x66, 0x66,  // 0.9f
+    0x3f, 0x19, 0x99, 0x9a,  // 0.6f
+    0x3f, 0x80, 0x00, 0x00,  // 1.0f      C2
+    0x3f, 0x4c, 0xcc, 0xcd,  // 0.8f  C1
+    0x3f, 0x66, 0x66, 0x66,  // 0.9f
+    0x3f, 0x4c, 0xcc, 0xcd,  // 0.8f  C1
+    0x3f, 0x80, 0x00, 0x00,  // 1.0f      C2
+    0x3e, 0x99, 0x99, 0x9a,  // 0.3f
+    0x3f, 0x80, 0x00, 0x00,  // 1.0f
+    0x3f, 0x66, 0x66, 0x66,  // 0.9f  C1
+    0x3f, 0x80, 0x00, 0x00,  // 1.0f  C1
+    0x3f, 0x66, 0x66, 0x66,  // 0.9f
+    0x3f, 0x00, 0x00, 0x00,  // 0.5f
+    0x3f, 0x00, 0x00, 0x00,  // 0.5f
+    0x3f, 0x66, 0x66, 0x66,  // 0.9f
+    0x3f, 0x33, 0x33, 0x33,  // 0.7f      C2
+    0x3f, 0x19, 0x99, 0x9a,  // 0.6f      C2
+    0x3f, 0x80, 0x00, 0x00,  // 1.0f
+    0x3f, 0x80, 0x00, 0x00,  // 1.0f      C2
+    0x3f, 0x80, 0x00, 0x00,  // 1.0f
+    0x3f, 0x19, 0x99, 0x9a,  // 0.6f  C1
+    0x3f, 0x80, 0x00, 0x00  // 1.0f  C1
+};
+
+const unsigned char RawV25::controls[] = {
+    0x00, 0x0c  // PSD count
+};
+
+const unsigned char RawV25::joints[] = {
+    0x00, 0x51,  // Rows = 81
+    0x00, 0x0a,  // Columns = 10
+    // Joint groups
+    0x00, 0x00, 0x00, 0x04,  // Joint group count
+    // Joint group-0
+    0x00, 0x00, 0x00, 0x02,  // LOD count
+    0x00, 0x03,  // LOD-0 row-count
+    0x00, 0x03,  // LOD-1 row-count
+    0x00, 0x00, 0x00, 0x07,  // Input indices count
+    0x00, 0x00,  // Index: 0      C1
+    0x00, 0x01,  // Index: 1  C0  C1
+    0x00, 0x02,  // Index: 2  C0
+    0x00, 0x03,  // Index: 3      C1
+    0x00, 0x06,  // Index: 6      C1
+    0x00, 0x07,  // Index: 7  C0
+    0x00, 0x08,  // Index: 8  C0  C1
+    0x00, 0x00, 0x00, 0x03,  // Output indices count
+    0x00, 0x02,  // Index: 2
+    0x00, 0x03,  // Index: 3
+    0x00, 0x05,  // Index: 5
+    0x00, 0x00, 0x00, 0x15,  // Float value count: 21
+    // Row 0
+    0x00, 0x00, 0x00, 0x00,  // 0.00f      C1
+    0x3d, 0x4c, 0xcc, 0xcd,  // 0.05f  C0  C1
+    0x3d, 0xcc, 0xcc, 0xcd,  // 0.10f  C0
+    0x3e, 0x19, 0x99, 0x9a,  // 0.15f      C1
+    0x3e, 0x4c, 0xcc, 0xcd,  // 0.20f      C1
+    0x3e, 0x80, 0x00, 0x00,  // 0.25f  C0
+    0x3e, 0x99, 0x99, 0x9a,  // 0.30f  C0  C1
+    // Row 1
+    0x3e, 0xb3, 0x33, 0x33,  // 0.35f      C1
+    0x3e, 0xcc, 0xcc, 0xcd,  // 0.40f  C0  C1
+    0x3e, 0xe6, 0x66, 0x66,  // 0.45f  C0
+    0x3f, 0x00, 0x00, 0x00,  // 0.50f      C1
+    0x3f, 0x0c, 0xcc, 0xcd,  // 0.55f      C1
+    0x3f, 0x19, 0x99, 0x9a,  // 0.60f  C0
+    0x3f, 0x26, 0x66, 0x66,  // 0.65f  C0  C1
+    // Row 2
+    0x3f, 0x33, 0x33, 0x33,  // 0.70f      C1
+    0x3f, 0x40, 0x00, 0x00,  // 0.75f  C0  C1
+    0x3f, 0x4c, 0xcc, 0xcd,  // 0.80f  C0
+    0x3f, 0x59, 0x99, 0x9a,  // 0.85f      C1
+    0x3f, 0x66, 0x66, 0x66,  // 0.90f      C1
+    0x3f, 0x73, 0x33, 0x33,  // 0.95f  C0
+    0x3f, 0x80, 0x00, 0x00,  // 1.00f  C0  C1
+    // Joint indices
+    0x00, 0x00, 0x00, 0x01,  // Joint index count: 1
+    0x00, 0x00,  // Index: 0
+    // Joint group-1
+    0x00, 0x00, 0x00, 0x02,  // LOD count
+    0x00, 0x04,  // LOD-0 row-count
+    0x00, 0x02,  // LOD-1 row-count
+    0x00, 0x00, 0x00, 0x05,  // Input indices count
+    0x00, 0x03,  // Index: 3      C1
+    0x00, 0x04,  // Index: 4  C0
+    0x00, 0x07,  // Index: 7  C0
+    0x00, 0x08,  // Index: 8  C0  C1
+    0x00, 0x09,  // Index: 9      C1
+    0x00, 0x00, 0x00, 0x04,  // Output indices count
+    0x00, 0x12,  // Index: 18
+    0x00, 0x14,  // Index: 20
+    0x00, 0x24,  // Index: 36
+    0x00, 0x26,  // Index: 38
+    0x00, 0x00, 0x00, 0x14,  // Float value count: 20
+    // Row 0
+    0x3c, 0x23, 0xd7, 0x0a,  // 0.01f      C1
+    0x3c, 0xa3, 0xd7, 0x0a,  // 0.02f  C0
+    0x3c, 0xf5, 0xc2, 0x8f,  // 0.03f  C0
+    0x3d, 0x23, 0xd7, 0x0a,  // 0.04f  C0  C1
+    0x3d, 0x4c, 0xcc, 0xcd,  // 0.05f      C1
+    // Row 1
+    0x3d, 0x75, 0xc2, 0x8f,  // 0.06f      C1
+    0x3d, 0x8f, 0x5c, 0x29,  // 0.07f  C0
+    0x3d, 0xa3, 0xd7, 0x0a,  // 0.08f  C0
+    0x3d, 0xb8, 0x51, 0xec,  // 0.09f  C0  C1
+    0x3d, 0xcc, 0xcc, 0xcd,  // 0.10f      C1
+    // Row 2
+    0x3d, 0xe1, 0x47, 0xae,  // 0.11f      C1
+    0x3d, 0xf5, 0xc2, 0x8f,  // 0.12f  C0
+    0x3e, 0x05, 0x1e, 0xb8,  // 0.13f  C0
+    0x3e, 0x0f, 0x5c, 0x29,  // 0.14f  C0  C1
+    0x3e, 0x19, 0x99, 0x9a,  // 0.15f      C1
+    // Row 3
+    0x3e, 0x23, 0xd7, 0x0a,  // 0.16f      C1
+    0x3e, 0x2e, 0x14, 0x7b,  // 0.17f  C0
+    0x3e, 0x38, 0x51, 0xec,  // 0.18f  C0
+    0x3e, 0x42, 0x8f, 0x5c,  // 0.19f  C0  C1
+    0x3e, 0x4c, 0xcc, 0xcd,  // 0.20f      C1
+    // Joint indices
+    0x00, 0x00, 0x00, 0x02,  // Joint index count: 2
+    0x00, 0x02,  // Index: 2
+    0x00, 0x04,  // Index: 4
+    // Joint group-2
+    0x00, 0x00, 0x00, 0x02,  // LOD count
+    0x00, 0x03,  // LOD-0 row-count
+    0x00, 0x02,  // LOD-1 row-count
+    0x00, 0x00, 0x00, 0x04,  // Input indices count
+    0x00, 0x04,  // Index: 4  C0
+    0x00, 0x05,  // Index: 5      C1
+    0x00, 0x08,  // Index: 8  C0  C1
+    0x00, 0x09,  // Index: 9      C1
+    0x00, 0x00, 0x00, 0x03,  // Output indices count
+    0x00, 0x37,  // Index: 55
+    0x00, 0x38,  // Index: 56
+    0x00, 0x3f,  // Index: 63
+    0x00, 0x00, 0x00, 0x0c,  // Float value count: 12
+    // Row 0
+    0x3e, 0x9e, 0xb8, 0x52,  // 0.31f  C0
+    0x3e, 0xb8, 0x51, 0xec,  // 0.36f      C1
+    0x3e, 0xd7, 0x0a, 0x3d,  // 0.42f  C0  C1
+    0x3e, 0xf0, 0xa3, 0xd7,  // 0.47f      C1
+    // Row 1
+    0x3f, 0x07, 0xae, 0x14,  // 0.53f  C0
+    0x3f, 0x14, 0x7a, 0xe1,  // 0.58f      C1
+    0x3f, 0x23, 0xd7, 0x0a,  // 0.64f  C0  C1
+    0x3f, 0x30, 0xa3, 0xd7,  // 0.69f      C1
+    // Row 2
+    0x3f, 0x40, 0x00, 0x00,  // 0.75f  C0
+    0x3f, 0x4c, 0xcc, 0xcd,  // 0.80f      C1
+    0x3f, 0x5c, 0x28, 0xf6,  // 0.86f  C0  C1
+    0x3f, 0x68, 0xf5, 0xc3,  // 0.91f       C1
+    // Joint indices
+    0x00, 0x00, 0x00, 0x02,  // Joint index count: 2
+    0x00, 0x06,  // Index: 6
+    0x00, 0x07,  // Index: 7
+    // Joint group-3
+    0x00, 0x00, 0x00, 0x02,  // LOD count
+    0x00, 0x03,  // LOD-0 row-count
+    0x00, 0x00,  // LOD-1 row-count
+    0x00, 0x00, 0x00, 0x04,  // Input indices count
+    0x00, 0x02,  // Index: 2  C0
+    0x00, 0x05,  // Index: 5      C1
+    0x00, 0x06,  // Index: 6  C0  C1
+    0x00, 0x08,  // Index: 8      C1
+    0x00, 0x00, 0x00, 0x03,  // Output indices count
+    0x00, 0x2d,  // Index: 45
+    0x00, 0x2e,  // Index: 46
+    0x00, 0x47,  // Index: 71
+    0x00, 0x00, 0x00, 0x0c,  // Float value count: 12
+    // Row 0
+    0x3e, 0x9e, 0xb8, 0x52,  // 0.31f  C0
+    0x3e, 0xb8, 0x51, 0xec,  // 0.36f      C1
+    0x3e, 0xd7, 0x0a, 0x3d,  // 0.42f  C0  C1
+    0x3e, 0xf0, 0xa3, 0xd7,  // 0.47f      C1
+    // Row 1
+    0x3f, 0x07, 0xae, 0x14,  // 0.53f  C0
+    0x3f, 0x14, 0x7a, 0xe1,  // 0.58f      C1
+    0x3f, 0x23, 0xd7, 0x0a,  // 0.64f  C0  C1
+    0x3f, 0x30, 0xa3, 0xd7,  // 0.69f      C1
+    // Row 2
+    0x3f, 0x40, 0x00, 0x00,  // 0.75f  C0
+    0x3f, 0x4c, 0xcc, 0xcd,  // 0.80f      C1
+    0x3f, 0x5c, 0x28, 0xf6,  // 0.86f  C0  C1
+    0x3f, 0x68, 0xf5, 0xc3,  // 0.91f       C1
+    // Joint indices
+    0x00, 0x00, 0x00, 0x02,  // Joint index count: 2
+    0x00, 0x05,  // Index: 5
+    0x00, 0x07  // Index: 7
+};
+
+const unsigned char RawV25::blendshapes[] = {
+    0x00, 0x00, 0x00, 0x02,  // LOD count
+    0x00, 0x07,  // LOD-0 row-count
+    0x00, 0x04,  // LOD-1 row-count
+    0x00, 0x00, 0x00, 0x07,  // Input indices count
+    0x00, 0x00,  // Index: 0      C1  L0  L1
+    0x00, 0x01,  // Index: 1  C0  C1  L0  L1
+    0x00, 0x02,  // Index: 2  C0      L0  L1
+    0x00, 0x03,  // Index: 3      C1  L0  L1
+    0x00, 0x06,  // Index: 6      C1  L0
+    0x00, 0x07,  // Index: 7  C0      L0
+    0x00, 0x08,  // Index: 8  C0  C1  L0
+    0x00, 0x00, 0x00, 0x07,  // Output indices count
+    0x00, 0x00,  // Index: 0      C1  L0  L1
+    0x00, 0x01,  // Index: 1  C0  C1  L0  L1
+    0x00, 0x02,  // Index: 2  C0      L0  L1
+    0x00, 0x03,  // Index: 3      C1  L0  L1
+    0x00, 0x06,  // Index: 6      C1  L0
+    0x00, 0x07,  // Index: 7  C0      L0
+    0x00, 0x08  // Index: 8  C0  C1  L0
+};
+
+const unsigned char RawV25::animatedmaps[] = {
+    // LOD sizes
+    0x00, 0x00, 0x00, 0x02,  // Row count per LOD
+    0x00, 0x0f,  // LOD-0 row-count
+    0x00, 0x06  // LOD-1 row-count
+};
+
+const unsigned char RawV25::geometry[] = {
+    0x00, 0x00, 0x00, 0x03,  // Mesh count
+    // Mesh-0
+    0x00, 0x00, 0x01, 0x52,  // Mesh-0 size
+    0x00, 0x00, 0x00, 0x03,  // Vertex positions X values length
+    0x40, 0xe0, 0x00, 0x00,  // 7.0f
+    0x41, 0x00, 0x00, 0x00,  // 8.0f
+    0x41, 0x10, 0x00, 0x00,  // 9.0f
+    0x00, 0x00, 0x00, 0x03,  // Vertex positions Y values length
+    0x40, 0xe0, 0x00, 0x00,  // 7.0f
+    0x41, 0x00, 0x00, 0x00,  // 8.0f
+    0x41, 0x10, 0x00, 0x00,  // 9.0f
+    0x00, 0x00, 0x00, 0x03,  // Vertex positions Z values length
+    0x40, 0xe0, 0x00, 0x00,  // 7.0f
+    0x41, 0x00, 0x00, 0x00,  // 8.0f
+    0x41, 0x10, 0x00, 0x00,  // 9.0f
+    0x00, 0x00, 0x00, 0x03,  // Texture coordinates U values length
+    0x40, 0xe0, 0x00, 0x00,  // 7.0f
+    0x41, 0x00, 0x00, 0x00,  // 8.0f
+    0x41, 0x10, 0x00, 0x00,  // 9.0f
+    0x00, 0x00, 0x00, 0x03,  // Texture coordinates V values length
+    0x40, 0xe0, 0x00, 0x00,  // 7.0f
+    0x41, 0x00, 0x00, 0x00,  // 8.0f
+    0x41, 0x10, 0x00, 0x00,  // 9.0f
+    0x00, 0x00, 0x00, 0x03,  // Vertex normals X values length
+    0x40, 0xe0, 0x00, 0x00,  // 7.0f
+    0x41, 0x00, 0x00, 0x00,  // 8.0f
+    0x41, 0x10, 0x00, 0x00,  // 9.0f
+    0x00, 0x00, 0x00, 0x03,  // Vertex normals Y values length
+    0x40, 0xe0, 0x00, 0x00,  // 7.0f
+    0x41, 0x00, 0x00, 0x00,  // 8.0f
+    0x41, 0x10, 0x00, 0x00,  // 9.0f
+    0x00, 0x00, 0x00, 0x03,  // Vertex normals Z values length
+    0x40, 0xe0, 0x00, 0x00,  // 7.0f
+    0x41, 0x00, 0x00, 0x00,  // 8.0f
+    0x41, 0x10, 0x00, 0x00,  // 9.0f
+    0x00, 0x00, 0x00, 0x03,  // Vertex layouts - position indices length
+    0x00, 0x00, 0x00, 0x00,  // Vertex position: 0
+    0x00, 0x00, 0x00, 0x01,  // Vertex position: 1
+    0x00, 0x00, 0x00, 0x02,  // Vertex position: 2
+    0x00, 0x00, 0x00, 0x03,  // Vertex layouts - texture coordinate indices length
+    0x00, 0x00, 0x00, 0x00,  // Vertex texture coordinate: 0
+    0x00, 0x00, 0x00, 0x01,  // Vertex texture coordinate: 1
+    0x00, 0x00, 0x00, 0x02,  // Vertex texture coordinate: 2
+    0x00, 0x00, 0x00, 0x03,  // Vertex layouts - normal indices length
+    0x00, 0x00, 0x00, 0x00,  // Vertex normal: 0
+    0x00, 0x00, 0x00, 0x01,  // Vertex normal: 1
+    0x00, 0x00, 0x00, 0x02,  // Vertex normal: 2
+    0x00, 0x00, 0x00, 0x01,  // Face count: 1
+    0x00, 0x00, 0x00, 0x03,  // Face 1 layout indices length: 3
+    0x00, 0x00, 0x00, 0x00,  // Layout index: 0
+    0x00, 0x00, 0x00, 0x01,  // Layout index: 1
+    0x00, 0x00, 0x00, 0x02,  // Layout index: 2
+    0x00, 0x08,  // Maximum influence per vertex
+    0x00, 0x00, 0x00, 0x03,  // Skin weights structure count: 3 (for each vertex)
+    0x00, 0x00, 0x00, 0x03,  // Weights length: 3 (for each influencing joint)
+    0x3f, 0x33, 0x33, 0x33,  // 0.7f
+    0x3d, 0xcc, 0xcc, 0xcd,  // 0.1f
+    0x3e, 0x4c, 0xcc, 0xcd,  // 0.2f
+    0x00, 0x00, 0x00, 0x03,  // Influencing joint count: 3 (for each weight)
+    0x00, 0x00,  // Joint: 0
+    0x00, 0x01,  // Joint: 1
+    0x00, 0x02,  // Joint: 2
+    0x00, 0x00, 0x00, 0x02,  // Weights length: 2 (for each influencing joint)
+    0x3f, 0x00, 0x00, 0x00,  // 0.5f
+    0x3f, 0x00, 0x00, 0x00,  // 0.5f
+    0x00, 0x00, 0x00, 0x02,  // Influencing joint count: 2 (for each weight)
+    0x00, 0x03,  // Joint: 3
+    0x00, 0x04,  // Joint: 4
+    0x00, 0x00, 0x00, 0x02,  // Weights length: 2 (for each influencing joint)
+    0x3e, 0xcc, 0xcc, 0xcd,  // 0.4f
+    0x3f, 0x19, 0x99, 0x9a,  // 0.6f
+    0x00, 0x00, 0x00, 0x02,  // Influencing joint count: 2 (for each weight)
+    0x00, 0x05,  // Joint: 5
+    0x00, 0x06,  // Joint: 6
+    0x00, 0x00, 0x00, 0x01,  // Number of blendshapes
+    0x00, 0x00, 0x00, 0x03,  // Blend shape deltas X values length
+    0x40, 0xe0, 0x00, 0x00,  // 7.0f
+    0x41, 0x00, 0x00, 0x00,  // 8.0f
+    0x41, 0x10, 0x00, 0x00,  // 9.0f
+    0x00, 0x00, 0x00, 0x03,  // Blend shape deltas Y values length
+    0x40, 0xe0, 0x00, 0x00,  // 7.0f
+    0x41, 0x00, 0x00, 0x00,  // 8.0f
+    0x41, 0x10, 0x00, 0x00,  // 9.0f
+    0x00, 0x00, 0x00, 0x03,  // Blend shape deltas Z values length
+    0x40, 0xe0, 0x00, 0x00,  // 7.0f
+    0x41, 0x00, 0x00, 0x00,  // 8.0f
+    0x41, 0x10, 0x00, 0x00,  // 9.0f
+    0x00, 0x00, 0x00, 0x03,  // Vertex position indices length (for each delta)
+    0x00, 0x00, 0x00, 0x00,  // Vertex position: 0
+    0x00, 0x00, 0x00, 0x01,  // Vertex position: 1
+    0x00, 0x00, 0x00, 0x02,  // Vertex position: 2
+    0x00, 0x02,  // Blend shape index in Definition
+    // Mesh-1
+    0x00, 0x00, 0x01, 0x52,  // Mesh-1 size
+    0x00, 0x00, 0x00, 0x03,  // Vertex positions X values length
+    0x40, 0x80, 0x00, 0x00,  // 4.0f
+    0x40, 0xa0, 0x00, 0x00,  // 5.0f
+    0x40, 0xc0, 0x00, 0x00,  // 6.0f
+    0x00, 0x00, 0x00, 0x03,  // Vertex positions Y values length
+    0x40, 0x80, 0x00, 0x00,  // 4.0f
+    0x40, 0xa0, 0x00, 0x00,  // 5.0f
+    0x40, 0xc0, 0x00, 0x00,  // 6.0f
+    0x00, 0x00, 0x00, 0x03,  // Vertex positions Z values length
+    0x40, 0x80, 0x00, 0x00,  // 4.0f
+    0x40, 0xa0, 0x00, 0x00,  // 5.0f
+    0x40, 0xc0, 0x00, 0x00,  // 6.0f
+    0x00, 0x00, 0x00, 0x03,  // Texture coordinates U values length
+    0x40, 0x80, 0x00, 0x00,  // 4.0f
+    0x40, 0xa0, 0x00, 0x00,  // 5.0f
+    0x40, 0xc0, 0x00, 0x00,  // 6.0f
+    0x00, 0x00, 0x00, 0x03,  // Texture coordinates V values length
+    0x40, 0x80, 0x00, 0x00,  // 4.0f
+    0x40, 0xa0, 0x00, 0x00,  // 5.0f
+    0x40, 0xc0, 0x00, 0x00,  // 6.0f
+    0x00, 0x00, 0x00, 0x03,  // Vertex normals X values length
+    0x40, 0x80, 0x00, 0x00,  // 4.0f
+    0x40, 0xa0, 0x00, 0x00,  // 5.0f
+    0x40, 0xc0, 0x00, 0x00,  // 6.0f
+    0x00, 0x00, 0x00, 0x03,  // Vertex normals Y values length
+    0x40, 0x80, 0x00, 0x00,  // 4.0f
+    0x40, 0xa0, 0x00, 0x00,  // 5.0f
+    0x40, 0xc0, 0x00, 0x00,  // 6.0f
+    0x00, 0x00, 0x00, 0x03,  // Vertex normals Z values length
+    0x40, 0x80, 0x00, 0x00,  // 4.0f
+    0x40, 0xa0, 0x00, 0x00,  // 5.0f
+    0x40, 0xc0, 0x00, 0x00,  // 6.0f
+    0x00, 0x00, 0x00, 0x03,  // Vertex layouts - position indices length
+    0x00, 0x00, 0x00, 0x00,  // Vertex position: 0
+    0x00, 0x00, 0x00, 0x01,  // Vertex position: 1
+    0x00, 0x00, 0x00, 0x02,  // Vertex position: 2
+    0x00, 0x00, 0x00, 0x03,  // Vertex layouts - texture coordinate indices length
+    0x00, 0x00, 0x00, 0x00,  // Vertex texture coordinate: 0
+    0x00, 0x00, 0x00, 0x01,  // Vertex texture coordinate: 1
+    0x00, 0x00, 0x00, 0x02,  // Vertex texture coordinate: 2
+    0x00, 0x00, 0x00, 0x03,  // Vertex layouts - normal indices length
+    0x00, 0x00, 0x00, 0x00,  // Vertex normal: 0
+    0x00, 0x00, 0x00, 0x01,  // Vertex normal: 1
+    0x00, 0x00, 0x00, 0x02,  // Vertex normal: 2
+    0x00, 0x00, 0x00, 0x01,  // Face count: 1
+    0x00, 0x00, 0x00, 0x03,  // Face 1 layout indices length: 3
+    0x00, 0x00, 0x00, 0x00,  // Layout index: 0
+    0x00, 0x00, 0x00, 0x01,  // Layout index: 1
+    0x00, 0x00, 0x00, 0x02,  // Layout index: 2
+    0x00, 0x08,  // Maximum influence per vertex
+    0x00, 0x00, 0x00, 0x03,  // Skin weights structure count: 3 (for each vertex)
+    0x00, 0x00, 0x00, 0x03,  // Weights length: 3 (for each influencing joint)
+    0x3e, 0xcc, 0xcc, 0xcd,  // 0.4f
+    0x3e, 0x99, 0x99, 0x9a,  // 0.3f
+    0x3e, 0x99, 0x99, 0x9a,  // 0.3f
+    0x00, 0x00, 0x00, 0x03,  // Influencing joint count: 3 (for each weight)
+    0x00, 0x00,  // Joint: 0
+    0x00, 0x01,  // Joint: 1
+    0x00, 0x02,  // Joint: 2
+    0x00, 0x00, 0x00, 0x02,  // Weights length: 2 (for each influencing joint)
+    0x3f, 0x4c, 0xcc, 0xcd,  // 0.8f
+    0x3e, 0x4c, 0xcc, 0xcd,  // 0.2f
+    0x00, 0x00, 0x00, 0x02,  // Influencing joint count: 2 (for each weight)
+    0x00, 0x03,  // Joint: 3
+    0x00, 0x04,  // Joint: 4
+    0x00, 0x00, 0x00, 0x02,  // Weights length: 2 (for each influencing joint)
+    0x3d, 0xcc, 0xcc, 0xcd,  // 0.1f
+    0x3f, 0x66, 0x66, 0x66,  // 0.9f
+    0x00, 0x00, 0x00, 0x02,  // Influencing joint count: 2 (for each weight)
+    0x00, 0x05,  // Joint: 5
+    0x00, 0x06,  // Joint: 6
+    0x00, 0x00, 0x00, 0x01,  // Number of blendshapes
+    0x00, 0x00, 0x00, 0x03,  // Blend shape deltas X values length
+    0x40, 0x80, 0x00, 0x00,  // 4.0f
+    0x40, 0xa0, 0x00, 0x00,  // 5.0f
+    0x40, 0xc0, 0x00, 0x00,  // 6.0f
+    0x00, 0x00, 0x00, 0x03,  // Blend shape deltas Y values length
+    0x40, 0x80, 0x00, 0x00,  // 4.0f
+    0x40, 0xa0, 0x00, 0x00,  // 5.0f
+    0x40, 0xc0, 0x00, 0x00,  // 6.0f
+    0x00, 0x00, 0x00, 0x03,  // Blend shape deltas Z values length
+    0x40, 0x80, 0x00, 0x00,  // 4.0f
+    0x40, 0xa0, 0x00, 0x00,  // 5.0f
+    0x40, 0xc0, 0x00, 0x00,  // 6.0f
+    0x00, 0x00, 0x00, 0x03,  // Vertex position indices length (for each delta)
+    0x00, 0x00, 0x00, 0x00,  // Vertex position: 0
+    0x00, 0x00, 0x00, 0x01,  // Vertex position: 1
+    0x00, 0x00, 0x00, 0x02,  // Vertex position: 2
+    0x00, 0x02,  // Blend shape index in Definition
+    // Mesh-2
+    0x00, 0x00, 0x01, 0x84,  // Mesh-2 size
+    0x00, 0x00, 0x00, 0x03,  // Vertex positions X values length
+    0x3f, 0x80, 0x00, 0x00,  // 1.0f
+    0x40, 0x00, 0x00, 0x00,  // 2.0f
+    0x40, 0x40, 0x00, 0x00,  // 3.0f
+    0x00, 0x00, 0x00, 0x03,  // Vertex positions Y values length
+    0x3f, 0x80, 0x00, 0x00,  // 1.0f
+    0x40, 0x00, 0x00, 0x00,  // 2.0f
+    0x40, 0x40, 0x00, 0x00,  // 3.0f
+    0x00, 0x00, 0x00, 0x03,  // Vertex positions Z values length
+    0x3f, 0x80, 0x00, 0x00,  // 1.0f
+    0x40, 0x00, 0x00, 0x00,  // 2.0f
+    0x40, 0x40, 0x00, 0x00,  // 3.0f
+    0x00, 0x00, 0x00, 0x03,  // Texture coordinates U values length
+    0x3f, 0x80, 0x00, 0x00,  // 1.0f
+    0x40, 0x00, 0x00, 0x00,  // 2.0f
+    0x40, 0x40, 0x00, 0x00,  // 3.0f
+    0x00, 0x00, 0x00, 0x03,  // Texture coordinates V values length
+    0x3f, 0x80, 0x00, 0x00,  // 1.0f
+    0x40, 0x00, 0x00, 0x00,  // 2.0f
+    0x40, 0x40, 0x00, 0x00,  // 3.0f
+    0x00, 0x00, 0x00, 0x03,  // Vertex normals X values length
+    0x3f, 0x80, 0x00, 0x00,  // 1.0f
+    0x40, 0x00, 0x00, 0x00,  // 2.0f
+    0x40, 0x40, 0x00, 0x00,  // 3.0f
+    0x00, 0x00, 0x00, 0x03,  // Vertex normals Y values length
+    0x3f, 0x80, 0x00, 0x00,  // 1.0f
+    0x40, 0x00, 0x00, 0x00,  // 2.0f
+    0x40, 0x40, 0x00, 0x00,  // 3.0f
+    0x00, 0x00, 0x00, 0x03,  // Vertex normals Z values length
+    0x3f, 0x80, 0x00, 0x00,  // 1.0f
+    0x40, 0x00, 0x00, 0x00,  // 2.0f
+    0x40, 0x40, 0x00, 0x00,  // 3.0f
+    0x00, 0x00, 0x00, 0x03,  // Vertex layouts - position indices length
+    0x00, 0x00, 0x00, 0x00,  // Vertex position: 0
+    0x00, 0x00, 0x00, 0x01,  // Vertex position: 1
+    0x00, 0x00, 0x00, 0x02,  // Vertex position: 2
+    0x00, 0x00, 0x00, 0x03,  // Vertex layouts - texture coordinate indices length
+    0x00, 0x00, 0x00, 0x00,  // Vertex texture coordinate: 0
+    0x00, 0x00, 0x00, 0x01,  // Vertex texture coordinate: 1
+    0x00, 0x00, 0x00, 0x02,  // Vertex texture coordinate: 2
+    0x00, 0x00, 0x00, 0x03,  // Vertex layouts - normal indices length
+    0x00, 0x00, 0x00, 0x00,  // Vertex normal: 0
+    0x00, 0x00, 0x00, 0x01,  // Vertex normal: 1
+    0x00, 0x00, 0x00, 0x02,  // Vertex normal: 2
+    0x00, 0x00, 0x00, 0x01,  // Face count: 1
+    0x00, 0x00, 0x00, 0x03,  // Face 1 layout indices length: 3
+    0x00, 0x00, 0x00, 0x00,  // Layout index: 0
+    0x00, 0x00, 0x00, 0x01,  // Layout index: 1
+    0x00, 0x00, 0x00, 0x02,  // Layout index: 2
+    0x00, 0x08,  // Maximum influence per vertex
+    0x00, 0x00, 0x00, 0x03,  // Skin weights structure count: 3 (for each vertex)
+    0x00, 0x00, 0x00, 0x03,  // Weights length: 3 (for each influencing joint)
+    0x3d, 0xcc, 0xcc, 0xcd,  // 0.1f
+    0x3e, 0x99, 0x99, 0x9a,  // 0.3f
+    0x3f, 0x19, 0x99, 0x9a,  // 0.6f
+    0x00, 0x00, 0x00, 0x03,  // Influencing joint count: 3 (for each weight)
+    0x00, 0x00,  // Joint: 0
+    0x00, 0x01,  // Joint: 1
+    0x00, 0x02,  // Joint: 2
+    0x00, 0x00, 0x00, 0x02,  // Weights length: 2 (for each influencing joint)
+    0x3e, 0x99, 0x99, 0x9a,  // 0.3f
+    0x3f, 0x33, 0x33, 0x33,  // 0.7f
+    0x00, 0x00, 0x00, 0x02,  // Influencing joint count: 2 (for each weight)
+    0x00, 0x03,  // Joint: 3
+    0x00, 0x04,  // Joint: 4
+    0x00, 0x00, 0x00, 0x02,  // Weights length: 2 (for each influencing joint)
+    0x3e, 0x4c, 0xcc, 0xcd,  // 0.2f
+    0x3f, 0x4c, 0xcc, 0xcd,  // 0.8f
+    0x00, 0x00, 0x00, 0x02,  // Influencing joint count: 2 (for each weight)
+    0x00, 0x05,  // Joint: 5
+    0x00, 0x06,  // Joint: 6
+    0x00, 0x00, 0x00, 0x02,  // Number of blendshapes
+    0x00, 0x00, 0x00, 0x03,  // Blend shape deltas X values length
+    0x3f, 0x80, 0x00, 0x00,  // 1.0f
+    0x40, 0x00, 0x00, 0x00,  // 2.0f
+    0x40, 0x40, 0x00, 0x00,  // 3.0f
+    0x00, 0x00, 0x00, 0x03,  // Blend shape deltas Y values length
+    0x3f, 0x80, 0x00, 0x00,  // 1.0f
+    0x40, 0x00, 0x00, 0x00,  // 2.0f
+    0x40, 0x40, 0x00, 0x00,  // 3.0f
+    0x00, 0x00, 0x00, 0x03,  // Blend shape deltas Z values length
+    0x3f, 0x80, 0x00, 0x00,  // 1.0f
+    0x40, 0x00, 0x00, 0x00,  // 2.0f
+    0x40, 0x40, 0x00, 0x00,  // 3.0f
+    0x00, 0x00, 0x00, 0x03,  // Vertex position indices length (for each delta)
+    0x00, 0x00, 0x00, 0x00,  // Vertex position: 0
+    0x00, 0x00, 0x00, 0x01,  // Vertex position: 1
+    0x00, 0x00, 0x00, 0x02,  // Vertex position: 2
+    0x00, 0x02,  // Blend shape index in Definition
+    0x00, 0x00, 0x00, 0x02,  // Blend shape deltas X values length
+    0x40, 0x80, 0x00, 0x00,  // 4.0f
+    0x40, 0xa0, 0x00, 0x00,  // 5.0f
+    0x00, 0x00, 0x00, 0x02,  // Blend shape deltas Y values length
+    0x40, 0x80, 0x00, 0x00,  // 4.0f
+    0x40, 0xa0, 0x00, 0x00,  // 5.0f
+    0x00, 0x00, 0x00, 0x02,  // Blend shape deltas Z values length
+    0x40, 0x80, 0x00, 0x00,  // 4.0f
+    0x40, 0xa0, 0x00, 0x00,  // 5.0f
+    0x00, 0x00, 0x00, 0x02,  // Vertex position indices length (for each delta)
+    0x00, 0x00, 0x00, 0x00,  // Vertex position: 0
+    0x00, 0x00, 0x00, 0x02,  // Vertex position: 2
+    0x00, 0x03  // Blend shape index in Definition
+};
+
+const unsigned char RawV25::machineLearnedBehavior[] = {
+    0x00, 0x00, 0x00, 0x09,  // Raw control names length
+    0x00, 0x00, 0x00, 0x02,  // Raw control name 0 length
+    0x4d, 0x41,  // Raw control name 0 : MA
+    0x00, 0x00, 0x00, 0x02,  // Raw control name 1 length
+    0x4d, 0x42,  // Raw control name 1 : MB
+    0x00, 0x00, 0x00, 0x02,  // Raw control name 2 length
+    0x4d, 0x43,  // Raw control name 2 : MC
+    0x00, 0x00, 0x00, 0x02,  // Raw control name 3 length
+    0x4d, 0x44,  // Raw control name 3 : MD
+    0x00, 0x00, 0x00, 0x02,  // Raw control name 4 length
+    0x4d, 0x45,  // Raw control name 4 : ME
+    0x00, 0x00, 0x00, 0x02,  // Raw control name 5 length
+    0x4d, 0x46,  // Raw control name 5 : MF
+    0x00, 0x00, 0x00, 0x02,  // Raw control name 6 length
+    0x4d, 0x47,  // Raw control name 6 : MG
+    0x00, 0x00, 0x00, 0x02,  // Raw control name 7 length
+    0x4d, 0x48,  // Raw control name 7 : MH
+    0x00, 0x00, 0x00, 0x02,  // Raw control name 8 length
+    0x4d, 0x49,  // Raw control name 8 : MI
+    0x00, 0x00, 0x00, 0x02,  // Neural network indices lod to row mapping length
+    0x00, 0x00,  // Map from LOD-0 to row 0 in below defined matrix
+    0x00, 0x01,  // Map from LOD-1 to row 1 in below defined matrix
+    0x00, 0x00, 0x00, 0x02,  // Neural network indices per LOD row count
+    0x00, 0x00, 0x00, 0x04,  // Indices matrix row-0
+    0x00, 0x00,  // Neural network index: 0
+    0x00, 0x01,  // Neural network index: 1
+    0x00, 0x02,  // Neural network index: 2
+    0x00, 0x03,  // Neural network index: 3
+    0x00, 0x00, 0x00, 0x02,  // Indices matrix row-1
+    0x00, 0x04,  // Neural network index: 4
+    0x00, 0x05,  // Neural network index: 5
+    0x00, 0x00, 0x00, 0x03,  // Region names length
+    0x00, 0x00, 0x00, 0x02,  // Region names length for mesh 0
+    0x00, 0x00, 0x00, 0x02,  // Region name 0 length
+    0x52, 0x41,  // Region name 0 : RA
+    0x00, 0x00, 0x00, 0x02,  // Region name 1 length
+    0x52, 0x42,  // Region name 1 : RB
+    0x00, 0x00, 0x00, 0x02,  // Region names length for mesh 1
+    0x00, 0x00, 0x00, 0x02,  // Region name 0 length
+    0x52, 0x43,  // Region name 0 : RC
+    0x00, 0x00, 0x00, 0x02,  // Region name 1 length
+    0x52, 0x44,  // Region name 1 : RD
+    0x00, 0x00, 0x00, 0x02,  // Region names length for mesh 2
+    0x00, 0x00, 0x00, 0x02,  // Region name 0 length
+    0x52, 0x45,  // Region name 0 : RE
+    0x00, 0x00, 0x00, 0x02,  // Region name 1 length
+    0x52, 0x46,  // Region name 1 : RF
+    0x00, 0x00, 0x00, 0x03,  // Mesh count
+    0x00, 0x00, 0x00, 0x02,  // Region count for Mesh-0
+    0x00, 0x00, 0x00, 0x01,  // Neural network index count for Mesh-0 Region-0
+    0x00, 0x00,  // Neural network index: 0
+    0x00, 0x00, 0x00, 0x01,  // Neural network index count for Mesh-0 Region-1
+    0x00, 0x01,  // Neural network index: 1
+    0x00, 0x00, 0x00, 0x02,  // Region count for Mesh-1
+    0x00, 0x00, 0x00, 0x01,  // Neural network index count for Mesh-1 Region-0
+    0x00, 0x02,  // Neural network index: 2
+    0x00, 0x00, 0x00, 0x01,  // Neural network index count for Mesh-1 Region-1
+    0x00, 0x03,  // Neural network index: 3
+    0x00, 0x00, 0x00, 0x02,  // Region count for Mesh-2
+    0x00, 0x00, 0x00, 0x01,  // Neural network index count for Mesh-2 Region-0
+    0x00, 0x04,  // Neural network index: 4
+    0x00, 0x00, 0x00, 0x01,  // Neural network index count for Mesh-2 Region-1
+    0x00, 0x05,  // Neural network index: 5
+    0x00, 0x00, 0x00, 0x06,  // Neural network count
+    // Mesh-0 Region-0 neural network
+    0x00, 0x00, 0x00, 0x5a,  // Mesh-0 Region-0 neural network size
+    0x00, 0x00, 0x00, 0x01,  // Mesh-0 Region-0 neural network output index count
+    0x00, 0x09,  // Mesh-0 Region-0 neural network output index-9
+    0x00, 0x00, 0x00, 0x02,  // Mesh-0 Region-1 neural network input index count
+    0x00, 0x00,  // Mesh-0 Region-0 neural network input index-0
+    0x00, 0x01,  // Mesh-0 Region-0 neural network input index-1
+    0x00, 0x00, 0x00, 0x02,  // Mesh-0 Region-0 neural network layer count
+    0x00, 0x00, 0x00, 0x02,  // Mesh-0 Region-0 neural network layer-0 bias count
+    0x3f, 0x80, 0x00, 0x00,  // 1.0f
+    0x3f, 0x80, 0x00, 0x00,  // 1.0f
+    0x00, 0x00, 0x00, 0x04,  // Mesh-0 Region-0 neural network layer-0 weight count
+    0x3f, 0x00, 0x00, 0x00,  // 0.5f
+    0x3f, 0x00, 0x00, 0x00,  // 0.5f
+    0x3f, 0x00, 0x00, 0x00,  // 0.5f
+    0x3f, 0x00, 0x00, 0x00,  // 0.5f
+    0x00, 0x01,  // Mesh-0 Region-0 neural network layer-0 activation function ID
+    0x00, 0x00, 0x00, 0x01,  // Mesh-0 Region-0 neural network layer-1 activation function parameter count
+    0x3f, 0x00, 0x00, 0x00,  // 0.5f
+    0x00, 0x00, 0x00, 0x01,  // Mesh-0 Region-0 neural network layer-1 bias count
+    0x3f, 0x80, 0x00, 0x00,  // 1.0f
+    0x00, 0x00, 0x00, 0x02,  // Mesh-0 Region-0 neural network layer-1 weight count
+    0x3f, 0x00, 0x00, 0x00,  // 0.5f
+    0x3f, 0x00, 0x00, 0x00,  // 0.5f
+    0x00, 0x01,  // Mesh-0 Region-0 neural network layer-1 activation function ID
+    0x00, 0x00, 0x00, 0x01,  // Mesh-0 Region-0 neural network layer-1 activation function parameter count
+    0x3f, 0x00, 0x00, 0x00,  // 0.5f
+    // Mesh-0 Region-1 neural network
+    0x00, 0x00, 0x00, 0x5a,  // Mesh-0 Region-1 neural network size
+    0x00, 0x00, 0x00, 0x01,  // Mesh-0 Region-1 neural network output index count
+    0x00, 0x0a,  // Mesh-0 Region-1 neural network output index-10
+    0x00, 0x00, 0x00, 0x02,  // Mesh-0 Region-1 neural network input index count
+    0x00, 0x02,  // Mesh-0 Region-1 neural network input index-2
+    0x00, 0x03,  // Mesh-0 Region-1 neural network input index-3
+    0x00, 0x00, 0x00, 0x02,  // Mesh-0 Region-1 neural network layer count
+    0x00, 0x00, 0x00, 0x02,  // Mesh-0 Region-1 neural network layer-0 bias count
+    0x3f, 0x00, 0x00, 0x00,  // 0.5f
+    0x3f, 0x00, 0x00, 0x00,  // 0.5f
+    0x00, 0x00, 0x00, 0x04,  // Mesh-0 Region-1 neural network layer-0 weight count
+    0x3f, 0x80, 0x00, 0x00,  // 1.0f
+    0x3f, 0x80, 0x00, 0x00,  // 1.0f
+    0x3f, 0x80, 0x00, 0x00,  // 1.0f
+    0x3f, 0x80, 0x00, 0x00,  // 1.0f
+    0x00, 0x01,  // Mesh-0 Region-1 neural network layer-0 activation function ID
+    0x00, 0x00, 0x00, 0x01,  // Mesh-0 Region-1 neural network layer-1 activation function parameter count
+    0x3f, 0x80, 0x00, 0x00,  // 1.0f
+    0x00, 0x00, 0x00, 0x01,  // Mesh-0 Region-1 neural network layer-1 bias count
+    0x3f, 0x00, 0x00, 0x00,  // 0.5f
+    0x00, 0x00, 0x00, 0x02,  // Mesh-0 Region-1 neural network layer-1 weight count
+    0x3f, 0x80, 0x00, 0x00,  // 1.0f
+    0x3f, 0x80, 0x00, 0x00,  // 1.0f
+    0x00, 0x01,  // Mesh-0 Region-1 neural network layer-1 activation function ID
+    0x00, 0x00, 0x00, 0x01,  // Mesh-0 Region-1 neural network layer-1 activation function parameter count
+    0x3f, 0x80, 0x00, 0x00,  // 1.0f
+    // Mesh-1 Region-0 neural network
+    0x00, 0x00, 0x00, 0x5a,  // Mesh-1 Region-0 neural network size
+    0x00, 0x00, 0x00, 0x01,  // Mesh-1 Region-0 neural network output index count
+    0x00, 0x0b,  // Mesh-1 Region-0 neural network output index-11
+    0x00, 0x00, 0x00, 0x02,  // Mesh-1 Region-1 neural network input index count
+    0x00, 0x04,  // Mesh-1 Region-0 neural network input index-4
+    0x00, 0x05,  // Mesh-1 Region-0 neural network input index-5
+    0x00, 0x00, 0x00, 0x02,  // Mesh-1 Region-0 neural network layer count
+    0x00, 0x00, 0x00, 0x02,  // Mesh-1 Region-0 neural network layer-0 bias count
+    0x3f, 0x00, 0x00, 0x00,  // 0.5f
+    0x3f, 0x00, 0x00, 0x00,  // 0.5f
+    0x00, 0x00, 0x00, 0x04,  // Mesh-1 Region-0 neural network layer-0 weight count
+    0x3f, 0x80, 0x00, 0x00,  // 1.0f
+    0x3f, 0x80, 0x00, 0x00,  // 1.0f
+    0x3f, 0x80, 0x00, 0x00,  // 1.0f
+    0x3f, 0x80, 0x00, 0x00,  // 1.0f
+    0x00, 0x01,  // Mesh-1 Region-0 neural network layer-0 activation function ID
+    0x00, 0x00, 0x00, 0x01,  // Mesh-1 Region-0 neural network layer-1 activation function parameter count
+    0x3f, 0x80, 0x00, 0x00,  // 1.0f
+    0x00, 0x00, 0x00, 0x01,  // Mesh-1 Region-0 neural network layer-1 bias count
+    0x3f, 0x00, 0x00, 0x00,  // 0.5f
+    0x00, 0x00, 0x00, 0x02,  // Mesh-1 Region-0 neural network layer-1 weight count
+    0x3f, 0x80, 0x00, 0x00,  // 1.0f
+    0x3f, 0x80, 0x00, 0x00,  // 1.0f
+    0x00, 0x01,  // Mesh-1 Region-0 neural network layer-1 activation function ID
+    0x00, 0x00, 0x00, 0x01,  // Mesh-1 Region-0 neural network layer-1 activation function parameter count
+    0x3f, 0x80, 0x00, 0x00,  // 1.0f
+    // Mesh-1 Region-1 neural network
+    0x00, 0x00, 0x00, 0x5a,  // Mesh-1 Region-1 neural network size
+    0x00, 0x00, 0x00, 0x01,  // Mesh-1 Region-1 neural network output index count
+    0x00, 0x0c,  // Mesh-1 Region-1 neural network output index-12
+    0x00, 0x00, 0x00, 0x02,  // Mesh-1 Region-1 neural network input index count
+    0x00, 0x06,  // Mesh-1 Region-1 neural network input index-6
+    0x00, 0x07,  // Mesh-1 Region-1 neural network input index-7
+    0x00, 0x00, 0x00, 0x02,  // Mesh-1 Region-1 neural network layer count
+    0x00, 0x00, 0x00, 0x02,  // Mesh-1 Region-1 neural network layer-0 bias count
+    0x3f, 0x80, 0x00, 0x00,  // 1.0f
+    0x3f, 0x80, 0x00, 0x00,  // 1.0f
+    0x00, 0x00, 0x00, 0x04,  // Mesh-1 Region-1 neural network layer-0 weight count
+    0x3f, 0x00, 0x00, 0x00,  // 0.5f
+    0x3f, 0x00, 0x00, 0x00,  // 0.5f
+    0x3f, 0x00, 0x00, 0x00,  // 0.5f
+    0x3f, 0x00, 0x00, 0x00,  // 0.5f
+    0x00, 0x01,  // Mesh-1 Region-1 neural network layer-0 activation function ID
+    0x00, 0x00, 0x00, 0x01,  // Mesh-1 Region-1 neural network layer-1 activation function parameter count
+    0x3f, 0x00, 0x00, 0x00,  // 0.5f
+    0x00, 0x00, 0x00, 0x01,  // Mesh-1 Region-1 neural network layer-1 bias count
+    0x3f, 0x80, 0x00, 0x00,  // 1.0f
+    0x00, 0x00, 0x00, 0x02,  // Mesh-1 Region-1 neural network layer-1 weight count
+    0x3f, 0x00, 0x00, 0x00,  // 0.5f
+    0x3f, 0x00, 0x00, 0x00,  // 0.5f
+    0x00, 0x01,  // Mesh-1 Region-1 neural network layer-1 activation function ID
+    0x00, 0x00, 0x00, 0x01,  // Mesh-1 Region-1 neural network layer-1 activation function parameter count
+    0x3f, 0x00, 0x00, 0x00,  // 0.5f
+    // Mesh-2 Region-0 neural network
+    0x00, 0x00, 0x00, 0x5a,  // Mesh-2 Region-0 neural network size
+    0x00, 0x00, 0x00, 0x01,  // Mesh-2 Region-0 neural network output index count
+    0x00, 0x0d,  // Mesh-2 Region-0 neural network output index-13
+    0x00, 0x00, 0x00, 0x02,  // Mesh-2 Region-1 neural network input index count
+    0x00, 0x08,  // Mesh-2 Region-0 neural network input index-8
+    0x00, 0x00,  // Mesh-2 Region-0 neural network input index-0
+    0x00, 0x00, 0x00, 0x02,  // Mesh-2 Region-0 neural network layer count
+    0x00, 0x00, 0x00, 0x02,  // Mesh-2 Region-0 neural network layer-0 bias count
+    0x3f, 0x80, 0x00, 0x00,  // 1.0f
+    0x3f, 0x80, 0x00, 0x00,  // 1.0f
+    0x00, 0x00, 0x00, 0x04,  // Mesh-2 Region-0 neural network layer-0 weight count
+    0x3f, 0x00, 0x00, 0x00,  // 0.5f
+    0x3f, 0x00, 0x00, 0x00,  // 0.5f
+    0x3f, 0x00, 0x00, 0x00,  // 0.5f
+    0x3f, 0x00, 0x00, 0x00,  // 0.5f
+    0x00, 0x01,  // Mesh-2 Region-0 neural network layer-0 activation function ID
+    0x00, 0x00, 0x00, 0x01,  // Mesh-2 Region-0 neural network layer-1 activation function parameter count
+    0x3f, 0x00, 0x00, 0x00,  // 0.5f
+    0x00, 0x00, 0x00, 0x01,  // Mesh-2 Region-0 neural network layer-1 bias count
+    0x3f, 0x80, 0x00, 0x00,  // 1.0f
+    0x00, 0x00, 0x00, 0x02,  // Mesh-2 Region-0 neural network layer-1 weight count
+    0x3f, 0x00, 0x00, 0x00,  // 0.5f
+    0x3f, 0x00, 0x00, 0x00,  // 0.5f
+    0x00, 0x01,  // Mesh-2 Region-0 neural network layer-1 activation function ID
+    0x00, 0x00, 0x00, 0x01,  // Mesh-2 Region-0 neural network layer-1 activation function parameter count
+    0x3f, 0x00, 0x00, 0x00,  // 0.5f
+    // Mesh-2 Region-1 neural network
+    0x00, 0x00, 0x00, 0x5a,  // Mesh-2 Region-1 neural network size
+    0x00, 0x00, 0x00, 0x01,  // Mesh-2 Region-1 neural network output index count
+    0x00, 0x0e,  // Mesh-2 Region-1 neural network output index-14
+    0x00, 0x00, 0x00, 0x02,  // Mesh-2 Region-1 neural network input index count
+    0x00, 0x04,  // Mesh-2 Region-1 neural network input index-4
+    0x00, 0x07,  // Mesh-2 Region-1 neural network input index-7
+    0x00, 0x00, 0x00, 0x02,  // Mesh-2 Region-1 neural network layer count
+    0x00, 0x00, 0x00, 0x02,  // Mesh-2 Region-1 neural network layer-0 bias count
+    0x3f, 0x00, 0x00, 0x00,  // 0.5f
+    0x3f, 0x00, 0x00, 0x00,  // 0.5f
+    0x00, 0x00, 0x00, 0x04,  // Mesh-2 Region-1 neural network layer-0 weight count
+    0x3f, 0x80, 0x00, 0x00,  // 1.0f
+    0x3f, 0x80, 0x00, 0x00,  // 1.0f
+    0x3f, 0x80, 0x00, 0x00,  // 1.0f
+    0x3f, 0x80, 0x00, 0x00,  // 1.0f
+    0x00, 0x01,  // Mesh-2 Region-1 neural network layer-0 activation function ID
+    0x00, 0x00, 0x00, 0x01,  // Mesh-2 Region-1 neural network layer-1 activation function parameter count
+    0x3f, 0x80, 0x00, 0x00,  // 1.0f
+    0x00, 0x00, 0x00, 0x01,  // Mesh-2 Region-1 neural network layer-1 bias count
+    0x3f, 0x00, 0x00, 0x00,  // 0.5f
+    0x00, 0x00, 0x00, 0x02,  // Mesh-2 Region-1 neural network layer-1 weight count
+    0x3f, 0x80, 0x00, 0x00,  // 1.0f
+    0x3f, 0x80, 0x00, 0x00,  // 1.0f
+    0x00, 0x01,  // Mesh-2 Region-1 neural network layer-1 activation function ID
+    0x00, 0x00, 0x00, 0x01,  // Mesh-2 Region-1 neural network layer-1 activation function parameter count
+    0x3f, 0x80, 0x00, 0x00  // 1.0f
+};
+
+const unsigned char RawV25::rbfBehavior[] = {
+    0x00, 0x00, 0x00, 0x02,  // Solver LOD mapping lods length (2)
+    0x00, 0x00,  // Solver LOD mapping lod 0 (0)
+    0x00, 0x01,  // Solver LOD mapping lod 1 (1)
+    0x00, 0x00, 0x00, 0x02,  // Solver LOD mapping indices length (2)
+    0x00, 0x00, 0x00, 0x02,  // Solver LODs lod 0 indices length (2)
+    0x00, 0x00,  // Solver index 0 (0)
+    0x00, 0x01,  // Solver index 1 (1)
+    0x00, 0x00, 0x00, 0x02,  // Solver LODs lod 1 indices length (2)
+    0x00, 0x01,  // Solver index 0 (1)
+    0x00, 0x02,  // Solver index 1 (2)
+    0x00, 0x00, 0x00, 0x03,  // Solvers length (3)
+    // Solver 0
+    0x00, 0x00, 0x00, 0x49,  // Solver size (73)
+    0x00, 0x00, 0x00, 0x03,  // Solver name length (3)
+    0x52, 0x53, 0x41,  // Solver name (RSA)
+    0x00, 0x00, 0x00, 0x02,  // Raw control indices length (2)
+    0x00, 0x0b,  // Raw control index 0 (11)
+    0x00, 0x0c,  // Raw control index 1 (12)
+    0x00, 0x00, 0x00, 0x03,  // Solver pose indices length (3)
+    0x00, 0x00,  // Pose index 0 (0)
+    0x00, 0x01,  // Pose index 1 (1)
+    0x00, 0x02,  // Pose index 2 (2)
+    0x00, 0x00, 0x00, 0x06,  // Raw control values length (6)
+    0x40, 0x00, 0x00, 0x00,  // Pose 0 Raw control value 0 (2.0)
+    0x00, 0x00, 0x00, 0x00,  // Pose 0 Raw control value 1 (0.0)
+    0x3f, 0x80, 0x00, 0x00,  // Pose 1 Raw control value 0 (1.0)
+    0x3f, 0x80, 0x00, 0x00,  // Pose 1 Raw control value 1 (1.0)
+    0x40, 0x40, 0x00, 0x00,  // Pose 2 Raw control value 0 (3.0)
+    0xc0, 0x40, 0x00, 0x00,  // Pose 2 Raw control value 1 (-3.0)
+    0x3f, 0x80, 0x00, 0x00,  // Solver radius (1.0)
+    0x3f, 0x80, 0x00, 0x00,  // Solver weight threshold (1.0)
+    0x00, 0x00,  // Solver type (0)
+    0x00, 0x00,  // Solver automatic radius (0)
+    0x00, 0x01,  // Solver distance method (1)
+    0x00, 0x00,  // Solver normalize method (0)
+    0x00, 0x02,  // Solver function type (2)
+    0x00, 0x00,  // Solver TwistAxis method (0)
+    // Solver 1
+    0x00, 0x00, 0x00, 0x35,  // Solver size (53)
+    0x00, 0x00, 0x00, 0x03,  // Solver name length (3)
+    0x52, 0x53, 0x42,  // Solver name (RSB)
+    0x00, 0x00, 0x00, 0x01,  // Raw control indices length (1)
+    0x00, 0x03,  // Raw control index 0 (3)
+    0x00, 0x00, 0x00, 0x02,  // Solver pose indices length (2)
+    0x00, 0x03,  // Pose index 0 (3)
+    0x00, 0x04,  // Pose index 1 (4)
+    0x00, 0x00, 0x00, 0x02,  // Raw control values length (2)
+    0x00, 0x00, 0x00, 0x00,  // Pose 2 Raw control value 0 (0.0)
+    0x40, 0x80, 0x00, 0x00,  // Pose 3 Raw control value 1 (4.0)
+    0x40, 0x00, 0x00, 0x00,  // Solver radius 0 (2.0)
+    0x40, 0x00, 0x00, 0x00,  // Solver weight threshold 0 (2.0)
+    0x00, 0x01,  // Solver type (1)
+    0x00, 0x00,  // Solver automatic radius (0)
+    0x00, 0x03,  // Solver distance method (3)
+    0x00, 0x01,  // Solver normalize method (1)
+    0x00, 0x02,  // Solver function type (2)
+    0x00, 0x01,  // Solver TwistAxis method (1)
+    // Solver 2
+    0x00, 0x00, 0x00, 0x49,  // Solver size (73)
+    0x00, 0x00, 0x00, 0x03,  // Solver name length (3)
+    0x52, 0x53, 0x43,  // Solver name (RSC)
+    0x00, 0x00, 0x00, 0x02,  // Raw control indices length (2)
+    0x00, 0x16,  // Raw control index 0 (22)
+    0x00, 0x17,  // Raw control index 0 (23)
+    0x00, 0x00, 0x00, 0x03,  // Solver pose indices length (3)
+    0x00, 0x05,  // Pose index 0 (5)
+    0x00, 0x06,  // Pose index 1 (6)
+    0x00, 0x07,  // Pose index 2 (7)
+    0x00, 0x00, 0x00, 0x06,  // Raw control values length (6)
+    0x40, 0x00, 0x00, 0x00,  // Pose 5 Raw control value 0(2.0)
+    0x00, 0x00, 0x00, 0x00,  // Pose 5 Raw control value 1(0.0)
+    0x3f, 0x80, 0x00, 0x00,  // Pose 6 Raw control value 0(1.0)
+    0x3f, 0x80, 0x00, 0x00,  // Pose 6 Raw control value 1(1.0)
+    0x40, 0x40, 0x00, 0x00,  // Pose 7 Raw control value 0(3.0)
+    0xc0, 0x40, 0x00, 0x00,  // Pose 7 Raw control value 1(-3.0f)
+    0x3f, 0x80, 0x00, 0x00,  // Solver radius (1.0)
+    0x3f, 0x80, 0x00, 0x00,  // Solver weight threshold (0.5)
+    0x00, 0x00,  // Solver type (0)
+    0x00, 0x00,  // Solver automatic radius (0)
+    0x00, 0x01,  // Solver distance method (1)
+    0x00, 0x00,  // Solver normalize method (0)
+    0x00, 0x00,  // Solver function type (0)
+    0x00, 0x00,  // Solver TwistAxis method (0)
+    0x00, 0x00, 0x00, 0x08,  // Pose length (8)
+    // Pose 0
+    0x00, 0x00, 0x00, 0x02,  // pose name length (2)
+    0x52, 0x41,  // pose name (RA)
+    0x00, 0x00, 0x00, 0x00,  // Solver pose scale 0 (0.0)
+    // Pose 1
+    0x00, 0x00, 0x00, 0x02,  // pose name length (2)
+    0x52, 0x42,  // pose name (RB)
+    0x3f, 0x80, 0x00, 0x00,  // Solver pose scale (1.0)
+    // Pose 2
+    0x00, 0x00, 0x00, 0x02,  // pose name length (2)
+    0x52, 0x43,  // pose name (RC)
+    0x40, 0x00, 0x00, 0x00,  // Solver pose scale (2.0)
+    // Pose 3
+    0x00, 0x00, 0x00, 0x02,  // Pose name length (2)
+    0x52, 0x44,  // pose name (RD)
+    0x40, 0x00, 0x00, 0x00,  // Solver pose scale 0 (2.0)
+    // Pose 4
+    0x00, 0x00, 0x00, 0x02,  // pose name length (2)
+    0x52, 0x45,  // pose name(RE)
+    0x3f, 0x80, 0x00, 0x00,  // Solver pose scale (1.0)
+    // Pose 5
+    0x00, 0x00, 0x00, 0x02,  // pose name length (2)
+    0x52, 0x46,  // pose name(RF)
+    0x3f, 0x80, 0x00, 0x00,  // Solver scale 0 (1.0)
+    // Pose 6
+    0x00, 0x00, 0x00, 0x02,  // pose name ,length (2)
+    0x52, 0x47,  // pose name(RG)
+    0x3f, 0x80, 0x00, 0x00,  // Solver pose scale (1.0)
+    // Pose 7
+    0x00, 0x00, 0x00, 0x02,  // pose name length (2)
+    0x52, 0x48,  // pose name(RH)
+    0x3f, 0x00, 0x00, 0x00,  // Solver pose scale (0.5)
+};
+
+const unsigned char RawV25::rbfBehaviorExt[] = {
+    0x00, 0x00, 0x00, 0x09,  // Pose control name count
+    0x00, 0x00, 0x00, 0x02,  // Pose control name 0 length
+    0x50, 0x41,  // Pose control name 0 : PA
+    0x00, 0x00, 0x00, 0x02,  // Pose control name 1 length
+    0x50, 0x42,  // Pose control name 1 : PB
+    0x00, 0x00, 0x00, 0x02,  // Pose control name 2 length
+    0x50, 0x43,  // Pose control name 2 : PC
+    0x00, 0x00, 0x00, 0x02,  // Pose control name 3 length
+    0x50, 0x44,  // Pose control name 3 : PD
+    0x00, 0x00, 0x00, 0x02,  // Pose control name 4 length
+    0x50, 0x45,  // Pose control name 4 : PE
+    0x00, 0x00, 0x00, 0x02,  // Pose control name 5 length
+    0x50, 0x46,  // Pose control name 5 : PF
+    0x00, 0x00, 0x00, 0x02,  // Pose control name 6 length
+    0x50, 0x47,  // Pose control name 6 : PG
+    0x00, 0x00, 0x00, 0x02,  // Pose control name 7 length
+    0x50, 0x48,  // Pose control name 7 : PH
+    0x00, 0x00, 0x00, 0x02,  // Pose control name 8 length
+    0x50, 0x49,  // Pose control name 8 : PI
+    0x00, 0x00, 0x00, 0x08,  // Pose count
+    0x00, 0x00, 0x00, 0x01,  // Pose-0 input control index count
+    0x00, 0x00,  // Pose-0 input control index-0
+    0x00, 0x00, 0x00, 0x01,  // Pose-0 output control index count
+    0x00, 0x08,  // Pose-0 output control index-0
+    0x00, 0x00, 0x00, 0x01,  // Pose-0 output control weight count
+    0x3f, 0x80, 0x00, 0x00,  // Pose-0 output control weight-0 (1.0)
+    0x00, 0x00, 0x00, 0x01,  // Pose-1 input control index count
+    0x00, 0x01,  // Pose-1 input control index-0
+    0x00, 0x00, 0x00, 0x01,  // Pose-1 output control index count
+    0x00, 0x09,  // Pose-1 output control index-0
+    0x00, 0x00, 0x00, 0x01,  // Pose-1 output control weight count
+    0x3f, 0x80, 0x00, 0x00,  // Pose-1 output control weight-0 (1.0)
+    0x00, 0x00, 0x00, 0x01,  // Pose-2 input control index count
+    0x00, 0x02,  // Pose-2 input control index-0
+    0x00, 0x00, 0x00, 0x01,  // Pose-2 output control index count
+    0x00, 0x0a,  // Pose-2 output control index-0
+    0x00, 0x00, 0x00, 0x01,  // Pose-2 output control weight count
+    0x3f, 0x80, 0x00, 0x00,  // Pose-2 output control weight-0 (1.0)
+    0x00, 0x00, 0x00, 0x01,  // Pose-3 input control index count
+    0x00, 0x03,  // Pose-3 input control index-0
+    0x00, 0x00, 0x00, 0x01,  // Pose-3 output control index count
+    0x00, 0x0b,  // Pose-3 output control index-0
+    0x00, 0x00, 0x00, 0x01,  // Pose-3 output control weight count
+    0x3f, 0x80, 0x00, 0x00,  // Pose-3 output control weight-0 (1.0)
+    0x00, 0x00, 0x00, 0x01,  // Pose-4 input control index count
+    0x00, 0x04,  // Pose-4 input control index-0
+    0x00, 0x00, 0x00, 0x01,  // Pose-4 output control index count
+    0x00, 0x0c,  // Pose-4 output control index-0
+    0x00, 0x00, 0x00, 0x01,  // Pose-4 output control weight count
+    0x3f, 0x80, 0x00, 0x00,  // Pose-4 output control weight-0 (1.0)
+    0x00, 0x00, 0x00, 0x01,  // Pose-5 input control index count
+    0x00, 0x05,  // Pose-5 input control index-0
+    0x00, 0x00, 0x00, 0x01,  // Pose-5 output control index count
+    0x00, 0x0d,  // Pose-5 output control index-0
+    0x00, 0x00, 0x00, 0x01,  // Pose-5 output control weight count
+    0x3f, 0x80, 0x00, 0x00,  // Pose-5 output control weight-0 (1.0)
+    0x00, 0x00, 0x00, 0x01,  // Pose-6 input control index count
+    0x00, 0x06,  // Pose-6 input control index-0
+    0x00, 0x00, 0x00, 0x01,  // Pose-6 output control index count
+    0x00, 0x0e,  // Pose-6 output control index-0
+    0x00, 0x00, 0x00, 0x01,  // Pose-6 output control weight count
+    0x3f, 0x80, 0x00, 0x00,  // Pose-6 output control weight-0 (1.0)
+    0x00, 0x00, 0x00, 0x01,  // Pose-7 input control index count
+    0x00, 0x07,  // Pose-7 input control index-0
+    0x00, 0x00, 0x00, 0x02,  // Pose-7 output control index count
+    0x00, 0x0f,  // Pose-7 output control index-0
+    0x00, 0x10,  // Pose-7 output control index-1
+    0x00, 0x00, 0x00, 0x02,  // Pose-7 output control weight count
+    0x3f, 0x00, 0x00, 0x00,  // Pose-7 output control weight-0 (0.5)
+    0x3f, 0x00, 0x00, 0x00  // Pose-7 output control weight-1 (0.5)
+};
+
+const unsigned char RawV25::jointBehaviorMetadata[] = {
+    0x00, 0x00, 0x00, 0x09,  // joint representations length (9)
+    0x00, 0x00,  // joint 0 translation (Vector)
+    0x00, 0x00,  // joint 0 rotation (EulerAngles)
+    0x00, 0x00,  // joint 0 scale (Vector)
+    0x00, 0x00,  // joint 1 translation (Vector)
+    0x00, 0x00,  // joint 1 rotation (EulerAngles)
+    0x00, 0x00,  // joint 1 scale (Vector)
+    0x00, 0x00,  // joint 2 translation (Vector)
+    0x00, 0x01,  // joint 2 rotation (Quaternion)
+    0x00, 0x00,  // joint 2 scale (Vector)
+    0x00, 0x00,  // joint 3 translation (Vector)
+    0x00, 0x01,  // joint 3 rotation (Quaternion)
+    0x00, 0x00,  // joint 3 scale (Vector)
+    0x00, 0x00,  // joint 4 translation (Vector)
+    0x00, 0x00,  // joint 4 rotation (EulerAngles)
+    0x00, 0x00,  // joint 4 scale (Vector)
+    0x00, 0x00,  // joint 5 translation (Vector)
+    0x00, 0x00,  // joint 5 rotation (EulerAngles)
+    0x00, 0x00,  // joint 5 scale (Vector)
+    0x00, 0x00,  // joint 6 translation (Vector)
+    0x00, 0x00,  // joint 6 rotation (EulerAngles)
+    0x00, 0x00,  // joint 6 scale (Vector)
+    0x00, 0x00,  // joint 7 translation (Vector)
+    0x00, 0x01,  // joint 7 rotation (Quaternion)
+    0x00, 0x00,  // joint 7 translation (Vector)
+    0x00, 0x00,  // joint 8 translation (Vector)
+    0x00, 0x00,  // joint 8 rotation (EulerAngles)
+    0x00, 0x00,  // joint 8 translation (Vector)
+};
+
+const unsigned char RawV25::twistSwingBehavior[] = {
+    0x00, 0x00, 0x00, 0x03,  // twist setups length (3)
+    0x00, 0x00, 0x00, 0x02,  // setup 0 twist blend weights length (2)
+    0x3f, 0x80, 0x00, 0x00,  // setup 0 blend weight 0 (1.0)
+    0x40, 0x00, 0x00, 0x00,  // setup 0 blend weight 1 (2.0)
+    0x00, 0x00, 0x00, 0x02,  // setup 0 twist output joint indices length (2)
+    0x00, 0x00,  // setup 0 twist output index 0 (0)
+    0x00, 0x01,  // setup 0 twist output index 1 (1)
+    0x00, 0x00, 0x00, 0x04,  // setup 0 twist input control indices length (4)
+    0x00, 0x05,  // setup 0 twist input index 0 (5)
+    0x00, 0x06,  // setup 0 twist input index 1 (6)
+    0x00, 0x07,  // setup 0 twist input index 2 (7)
+    0x00, 0x08,  // setup 0 twist input index 3 (8)
+    0x00, 0x00,  // setup 0 twist axis (X)
+    0x00, 0x00, 0x00, 0x02,  // setup 1 twist blend weights length (2)
+    0xc0, 0x00, 0x00, 0x00,  // setup 1 blend weight 0 (-2.0)
+    0xbf, 0x80, 0x00, 0x00,  // setup 1 blend weight 1 (-1.0)
+    0x00, 0x00, 0x00, 0x02,  // setup 1 twist output joint indices length (2)
+    0x00, 0x04,  // setup 1 twist output index 0 (4)
+    0x00, 0x06,  // setup 1 twist output index 1 (6)
+    0x00, 0x00, 0x00, 0x04,  // setup 1 twist input control indices length (4)
+    0x00, 0x0b,  // setup 1 twist input index 0 (11)
+    0x00, 0x0c,  // setup 1 twist input index 1 (12)
+    0x00, 0x0d,  // setup 1 twist input index 2 (13)
+    0x00, 0x0e,  // setup 1 twist input index 3 (14)
+    0x00, 0x01,  // setup 1 twist axis (Y)
+    0x00, 0x00, 0x00, 0x01,  // setup 2 twist blend weights length (1)
+    0x3f, 0x80, 0x00, 0x00,  // setup 2 blend weight 0 (1.0)
+    0x00, 0x00, 0x00, 0x01,  // setup 2 twist output joint indices length (1)
+    0x00, 0x05,  // setup 2 twist output index 5 (5)
+    0x00, 0x00, 0x00, 0x04,  // setup 2 twist input control indices length (4)
+    0x00, 0x1b,  // setup 2 twist input index 0 (27)
+    0x00, 0x1c,  // setup 2 twist input index 1 (28)
+    0x00, 0x1d,  // setup 2 twist input index 2 (29)
+    0x00, 0x1e,  // setup 2 twist input index 3 (30)
+    0x00, 0x02,  // setup 2 twist axis (Z)
+    0x00, 0x00, 0x00, 0x03,  // swing setups length (3)
+    0x00, 0x00, 0x00, 0x02,  // setup 0 swing blend weights length (2)
+    0x3f, 0x80, 0x00, 0x00,  // setup 0 blend weight 0 (1.0)
+    0x40, 0x00, 0x00, 0x00,  // setup 0 blend weight 1 (2.0)
+    0x00, 0x00, 0x00, 0x02,  // setup 0 swing output joint indices length (2)
+    0x00, 0x00,  // setup 0 swing output index 0 (0)
+    0x00, 0x01,  // setup 0 swing output index 1 (1)
+    0x00, 0x00, 0x00, 0x04,  // setup 0 swing input control indices length (4)
+    0x00, 0x05,  // setup 0 swing input index 0 (5)
+    0x00, 0x06,  // setup 0 swing input index 1 (6)
+    0x00, 0x07,  // setup 0 swing input index 2 (7)
+    0x00, 0x08,  // setup 0 swing input index 3 (8)
+    0x00, 0x00,  // setup 0 twist axis (X)
+    0x00, 0x00, 0x00, 0x02,  // setup 1 swing blend weights length (2)
+    0xc0, 0x00, 0x00, 0x00,  // setup 1 blend weight 0 (-2.0)
+    0xbf, 0x80, 0x00, 0x00,  // setup 1 blend weight 1 (-1.0)
+    0x00, 0x00, 0x00, 0x02,  // setup 1 swing output joint indices length (2)
+    0x00, 0x04,  // setup 1 swing output index 0 (4)
+    0x00, 0x06,  // setup 1 swing output index 1 (6)
+    0x00, 0x00, 0x00, 0x04,  // setup 1 swing input control indices length (4)
+    0x00, 0x0b,  // setup 1 swing input index 0 (11)
+    0x00, 0x0c,  // setup 1 swing input index 1 (12)
+    0x00, 0x0d,  // setup 1 swing input index 2 (13)
+    0x00, 0x0e,  // setup 1 swing input index 3 (14)
+    0x00, 0x01,  // setup 1 twist axis (Y)
+    0x00, 0x00, 0x00, 0x01,  // setup 2 swing blend weights length (1)
+    0x3f, 0x80, 0x00, 0x00,  // setup 2 blend weight 0 (1.0)
+    0x00, 0x00, 0x00, 0x01,  // setup 2 swing output joint indices length (1)
+    0x00, 0x05,  // setup 2 swing output index 5 (5)
+    0x00, 0x00, 0x00, 0x04,  // setup 2 swing input control indices length (4)
+    0x00, 0x1b,  // setup 2 swing input index 0 (27)
+    0x00, 0x1c,  // setup 2 swing input index 1 (28)
+    0x00, 0x1d,  // setup 2 swing input index 2 (29)
+    0x00, 0x1e,  // setup 2 swing input index 3 (30)
+    0x00, 0x02  // setup 2 twist axis (Z)
+};
+
+std::vector<char> RawV25::getBytes() {
+    #if !defined(__clang__) && defined(__GNUC__) && (__GNUC__ >= 12)
+        #pragma GCC diagnostic push
+        #pragma GCC diagnostic ignored "-Wstringop-overflow"
+    #endif
+    std::vector<char> bytes;
+    // Header
+    bytes.insert(bytes.end(), header, header + sizeof(header));
+    // Descriptor
+    bytes.insert(bytes.end(), descriptor, descriptor + sizeof(descriptor));
+    // Definition
+    bytes.insert(bytes.end(), definition, definition + sizeof(definition));
+    // Behavior
+    // > Controls
+    bytes.insert(bytes.end(), controls, controls + sizeof(controls));
+    bytes.insert(bytes.end(), conditionals, conditionals + sizeof(conditionals));
+    bytes.insert(bytes.end(), psds, psds + sizeof(psds));
+    // > Joints
+    bytes.insert(bytes.end(), joints, joints + sizeof(joints));
+    // > BlendShapes
+    bytes.insert(bytes.end(), blendshapes, blendshapes + sizeof(blendshapes));
+    // > AnimatedMaps
+    bytes.insert(bytes.end(), animatedmaps, animatedmaps + sizeof(animatedmaps));
+    bytes.insert(bytes.end(), conditionals, conditionals + sizeof(conditionals));
+    // Geometry
+    bytes.insert(bytes.end(), geometry, geometry + sizeof(geometry));
+    // Machine learned behavior
+    bytes.insert(bytes.end(), machineLearnedBehavior, machineLearnedBehavior + sizeof(machineLearnedBehavior));
+    // RBF behavior
+    bytes.insert(bytes.end(), rbfBehavior, rbfBehavior + sizeof(rbfBehavior));
+    // RBF behavior ext
+    bytes.insert(bytes.end(), rbfBehaviorExt, rbfBehaviorExt + sizeof(rbfBehaviorExt));
+    // JointBehavior meta data
+    bytes.insert(bytes.end(), jointBehaviorMetadata, jointBehaviorMetadata + sizeof(jointBehaviorMetadata));
+    // Twist swing behavior
+    bytes.insert(bytes.end(), twistSwingBehavior, twistSwingBehavior + sizeof(twistSwingBehavior));
+    return bytes;
+    #if !defined(__clang__) && defined(__GNUC__) && (__GNUC__ >= 12)
+        #pragma GCC diagnostic pop
+    #endif
+}
+
+std::vector<char> RawV24DowngradedFromV25::getBytes() {
+    auto bytes = RawV25::getBytes();
+    bytes[5] = 0x00;
+    bytes[6] = 0x04;
+    return bytes;
+}
+
+#ifdef __clang__
+    #pragma clang diagnostic push
+    #pragma clang diagnostic ignored "-Wexit-time-destructors"
+#endif
+// Descriptor
+const pma::String<char> DecodedV25::name = "test";
+const Archetype DecodedV25::archetype = Archetype::other;
+const Gender DecodedV25::gender = Gender::other;
+const std::uint16_t DecodedV25::age = 42u;
+const pma::Vector<DecodedV25::StringPair> DecodedV25::metadata = {
+    {"key-A", "value-A"},
+    {"key-B", "value-B"}
+};
+const TranslationUnit DecodedV25::translationUnit = TranslationUnit::m;
+const RotationUnit DecodedV25::rotationUnit = RotationUnit::radians;
+const CoordinateSystem DecodedV25::coordinateSystem = {
+    Direction::right,
+    Direction::up,
+    Direction::front
+};
+const std::uint16_t DecodedV25::lodCount[] = {
+    2u,  // MaxLOD-0 - MinLOD-1
+    1u,  // MaxLOD-1 - MinLOD-1
+    1u  // MaxLOD-0 - MinLOD-0
+};
+const std::uint16_t DecodedV25::maxLODs[] = {
+    0u,  // MaxLOD-0 - MinLOD-1
+    1u,  // MaxLOD-1 - MinLOD-0
+    0u  // MaxLOD-0 - MinLOD-0
+};
+const pma::String<char> DecodedV25::complexity = "A";
+const pma::String<char> DecodedV25::dbName = "testDB";
+
+// Definition
+const pma::Vector<pma::String<char> > DecodedV25::guiControlNames = {
+    "GA", "GB", "GC", "GD", "GE", "GF", "GG", "GH", "GI"
+};
+const pma::Vector<pma::String<char> > DecodedV25::rawControlNames = {
+    "RA", "RB", "RC", "RD", "RE", "RF", "RG", "RH", "RI"
+};
+const DecodedV25::VectorOfCharStringMatrix DecodedV25::jointNames = {
+    {  // MaxLOD-0 - MinLOD-1
+        {"JA", "JB", "JC", "JD", "JE", "JF", "JG", "JH", "JI"},
+        {"JA", "JB", "JC", "JD", "JG", "JI"}
+    },
+    {  // MaxLOD-1 - MinLOD-0
+        {"JA", "JB", "JC", "JD", "JG", "JI"}
+    },
+    {  // MaxLOD-0 - MinLOD-0
+        {"JA", "JB", "JC", "JD", "JE", "JF", "JG", "JH", "JI"},
+    }
+};
+const DecodedV25::VectorOfCharStringMatrix DecodedV25::blendShapeNames = {
+    {  // MaxLOD-0 - MinLOD-1
+        {"BA", "BB", "BC", "BD", "BE", "BF", "BG", "BH", "BI"},
+        {"BC", "BF", "BH", "BI"}
+    },
+    {  // MaxLOD-1 - MinLOD-1
+        {"BC", "BF", "BH", "BI"}
+    },
+    {  // MaxLOD-0 - MinLOD-0
+        {"BA", "BB", "BC", "BD", "BE", "BF", "BG", "BH", "BI"},
+    }
+};
+const DecodedV25::VectorOfCharStringMatrix DecodedV25::animatedMapNames = {
+    {  // MaxLOD-0 - MinLOD-1
+        {"AA", "AB", "AC", "AD", "AE", "AF", "AG", "AH", "AI", "AJ"},
+        {"AC", "AF", "AH", "AI"}
+    },
+    {  // MaxLOD-1 - MinLOD-1
+        {"AC", "AF", "AH", "AI"}
+    },
+    {  // MaxLOD-0 - MinLOD-0
+        {"AA", "AB", "AC", "AD", "AE", "AF", "AG", "AH", "AI", "AJ"},
+    }
+};
+const DecodedV25::VectorOfCharStringMatrix DecodedV25::meshNames = {
+    {  // MaxLOD-0 - MinLOD-1
+        {"MA", "MB"},
+        {"MC"}
+    },
+    {  // MaxLOD-1 - MinLOD-1
+        {"MC"}
+    },
+    {  // MaxLOD-0 - MinLOD-0
+        {"MA", "MB"}
+    }
+};
+const pma::Vector<pma::Matrix<std::uint16_t> > DecodedV25::meshBlendShapeIndices = {
+    {  // MaxLOD-0 - MinLOD-1
+        {0, 1, 2, 3, 4, 5, 6},
+        {7, 8}
+    },
+    {  // MaxLOD-1 - MinLOD-1
+        {0, 1}
+    },
+    {  // MaxLOD-0 - MinLOD-0
+        {0, 1, 2, 3, 4, 5, 6},
+    }
+};
+const pma::Matrix<std::uint16_t> DecodedV25::jointHierarchy = {
+    {  // MaxLOD-0 - MinLOD-1
+        {0, 0, 0, 1, 1, 4, 2, 4, 2}
+    },
+    {  // MaxLOD-1 - MinLOD-1
+        {0, 0, 0, 1, 2, 2}
+    },
+    {  // MaxLOD-0 - MinLOD-0
+        {0, 0, 0, 1, 1, 4, 2, 4, 2}
+    }
+};
+const pma::Vector<pma::Matrix<Vector3> > DecodedV25::neutralJointTranslations = {
+    {  // MaxLOD-0 - MinLOD-1
+        {
+            {1.0f, 1.0f, 1.0f},
+            {2.0f, 2.0f, 2.0f},
+            {3.0f, 3.0f, 3.0f},
+            {4.0f, 4.0f, 4.0f},
+            {5.0f, 5.0f, 5.0f},
+            {6.0f, 6.0f, 6.0f},
+            {7.0f, 7.0f, 7.0f},
+            {8.0f, 8.0f, 8.0f},
+            {9.0f, 9.0f, 9.0f}
+        },
+        {
+            {1.0f, 1.0f, 1.0f},
+            {2.0f, 2.0f, 2.0f},
+            {3.0f, 3.0f, 3.0f},
+            {4.0f, 4.0f, 4.0f},
+            {7.0f, 7.0f, 7.0f},
+            {9.0f, 9.0f, 9.0f}
+        }
+    },
+    {  // MaxLOD-1 - MinLOD-1
+        {
+            {1.0f, 1.0f, 1.0f},
+            {2.0f, 2.0f, 2.0f},
+            {3.0f, 3.0f, 3.0f},
+            {4.0f, 4.0f, 4.0f},
+            {7.0f, 7.0f, 7.0f},
+            {9.0f, 9.0f, 9.0f}
+        }
+    },
+    {  // MaxLOD-0 - MinLOD-0
+        {
+            {1.0f, 1.0f, 1.0f},
+            {2.0f, 2.0f, 2.0f},
+            {3.0f, 3.0f, 3.0f},
+            {4.0f, 4.0f, 4.0f},
+            {5.0f, 5.0f, 5.0f},
+            {6.0f, 6.0f, 6.0f},
+            {7.0f, 7.0f, 7.0f},
+            {8.0f, 8.0f, 8.0f},
+            {9.0f, 9.0f, 9.0f}
+        }
+    }
+};
+const pma::Vector<pma::Matrix<Vector3> > DecodedV25::neutralJointRotations = {
+    {  // MaxLOD-0 - MinLOD-1
+        {
+            {1.0f, 1.0f, 1.0f},
+            {2.0f, 2.0f, 2.0f},
+            {3.0f, 3.0f, 3.0f},
+            {4.0f, 4.0f, 4.0f},
+            {5.0f, 5.0f, 5.0f},
+            {6.0f, 6.0f, 6.0f},
+            {7.0f, 7.0f, 7.0f},
+            {8.0f, 8.0f, 8.0f},
+            {9.0f, 9.0f, 9.0f}
+        },
+        {
+            {1.0f, 1.0f, 1.0f},
+            {2.0f, 2.0f, 2.0f},
+            {3.0f, 3.0f, 3.0f},
+            {4.0f, 4.0f, 4.0f},
+            {7.0f, 7.0f, 7.0f},
+            {9.0f, 9.0f, 9.0f}
+        }
+    },
+    {  // MaxLOD-1 - MinLOD-1
+        {
+            {1.0f, 1.0f, 1.0f},
+            {2.0f, 2.0f, 2.0f},
+            {3.0f, 3.0f, 3.0f},
+            {4.0f, 4.0f, 4.0f},
+            {7.0f, 7.0f, 7.0f},
+            {9.0f, 9.0f, 9.0f}
+        }
+    },
+    {  // MaxLOD-0 - MinLOD-0
+        {
+            {1.0f, 1.0f, 1.0f},
+            {2.0f, 2.0f, 2.0f},
+            {3.0f, 3.0f, 3.0f},
+            {4.0f, 4.0f, 4.0f},
+            {5.0f, 5.0f, 5.0f},
+            {6.0f, 6.0f, 6.0f},
+            {7.0f, 7.0f, 7.0f},
+            {8.0f, 8.0f, 8.0f},
+            {9.0f, 9.0f, 9.0f}
+        }
+    }
+};
+
+// Behavior
+const std::uint16_t DecodedV25::guiControlCount = 9u;
+const std::uint16_t DecodedV25::rawControlCount = 9u;
+const std::uint16_t DecodedV25::psdCount = 12u;
+// Behavior->Conditionals
+const pma::Vector<pma::Matrix<std::uint16_t> > DecodedV25::conditionalInputIndices = {
+    {  // MaxLOD-0 - MinLOD-1
+        {0, 1, 1, 2, 3, 3, 4, 4, 4, 5, 6, 7, 7, 8, 8},
+        {0, 1, 1, 2, 3, 3},
+    },
+    {  // MaxLOD-1 - MinLOD-1
+        {0, 1, 1, 2, 3, 3},
+    },
+    {  // MaxLOD-0 - MinLOD-0
+        {0, 1, 1, 2, 3, 3, 4, 4, 4, 5, 6, 7, 7, 8, 8}
+    }
+};
+const pma::Vector<pma::Matrix<std::uint16_t> > DecodedV25::conditionalOutputIndices = {
+    {  // MaxLOD-0 - MinLOD-1
+        {0, 1, 1, 2, 3, 3, 4, 4, 4, 5, 6, 7, 7, 8, 8},
+        {0, 1, 1, 2, 3, 3},
+    },
+    {  // MaxLOD-1 - MinLOD-1
+        {0, 1, 1, 2, 3, 3},
+    },
+    {  // MaxLOD-0 - MinLOD-0
+        {0, 1, 1, 2, 3, 3, 4, 4, 4, 5, 6, 7, 7, 8, 8}
+    }
+};
+const pma::Vector<pma::Matrix<float> > DecodedV25::conditionalFromValues = {
+    {  // MaxLOD-0 - MinLOD-1
+        {0.0f, 0.0f, 0.6f, 0.4f, 0.1f, 0.7f, 0.0f, 0.4f, 0.7f, 0.5f, 0.0f, 0.1f, 0.6f, 0.2f, 0.0f},
+        {0.0f, 0.0f, 0.6f, 0.4f, 0.1f, 0.7f}
+    },
+    {  // MaxLOD-1 - MinLOD-1
+        {0.0f, 0.0f, 0.6f, 0.4f, 0.1f, 0.7f}
+    },
+    {  // MaxLOD-0 - MinLOD-0
+        {0.0f, 0.0f, 0.6f, 0.4f, 0.1f, 0.7f, 0.0f, 0.4f, 0.7f, 0.5f, 0.0f, 0.1f, 0.6f, 0.2f, 0.0f}
+    }
+};
+const pma::Vector<pma::Matrix<float> > DecodedV25::conditionalToValues = {
+    {  // MaxLOD-0 - MinLOD-1
+        {1.0f, 0.6f, 1.0f, 0.9f, 0.7f, 1.0f, 0.4f, 0.7f, 1.0f, 1.0f, 1.0f, 0.6f, 1.0f, 0.8f, 1.0f},
+        {1.0f, 0.6f, 1.0f, 0.9f, 0.7f, 1.0f}
+    },
+    {  // MaxLOD-1 - MinLOD-1
+        {1.0f, 0.6f, 1.0f, 0.9f, 0.7f, 1.0f}
+    },
+    {  // MaxLOD-1 - MinLOD-1
+        {1.0f, 0.6f, 1.0f, 0.9f, 0.7f, 1.0f, 0.4f, 0.7f, 1.0f, 1.0f, 1.0f, 0.6f, 1.0f, 0.8f, 1.0f}
+    }
+};
+const pma::Vector<pma::Matrix<float> > DecodedV25::conditionalSlopeValues = {
+    {  // MaxLOD-0 - MinLOD-1
+        {1.0f, 0.9f, 0.9f, 0.8f, 0.7f, 0.7f, 0.6f, 0.6f, 0.6f, 0.5f, 0.6f, 0.7f, 0.7f, 0.8f, 0.9f},
+        {1.0f, 0.9f, 0.9f, 0.8f, 0.7f, 0.7f}
+    },
+    {  // MaxLOD-1 - MinLOD-1
+        {1.0f, 0.9f, 0.9f, 0.8f, 0.7f, 0.7f}
+    },
+    {  // MaxLOD-0 - MinLOD-0
+        {1.0f, 0.9f, 0.9f, 0.8f, 0.7f, 0.7f, 0.6f, 0.6f, 0.6f, 0.5f, 0.6f, 0.7f, 0.7f, 0.8f, 0.9f}
+    }
+};
+const pma::Vector<pma::Matrix<float> > DecodedV25::conditionalCutValues = {
+    {  // MaxLOD-0 - MinLOD-1
+        {0.0f, 0.5f, 0.5f, 0.4f, 0.3f, 0.3f, 1.0f, 1.0f, 1.0f, 0.2f, 0.4f, 0.8f, 0.8f, 1.0f, 0.2f},
+        {0.0f, 0.5f, 0.5f, 0.4f, 0.3f, 0.3f}
+    },
+    {  // MaxLOD-1 - MinLOD-1
+        {0.0f, 0.5f, 0.5f, 0.4f, 0.3f, 0.3f}
+    },
+    {  // MaxLOD-0 - MinLOD-0
+        {0.0f, 0.5f, 0.5f, 0.4f, 0.3f, 0.3f, 1.0f, 1.0f, 1.0f, 0.2f, 0.4f, 0.8f, 0.8f, 1.0f, 0.2f}
+    }
+};
+// Behavior->PSDs
+const pma::Vector<std::uint16_t> DecodedV25::psdRowIndices = {
+    8, 8, 8, 9, 9, 10, 10, 10, 11, 12, 13, 13, 13, 14, 14, 15, 16, 18, 18, 18, 18, 19, 19, 20
+};
+const pma::Vector<std::uint16_t> DecodedV25::psdColumnIndices = {
+    0, 3, 6, 2, 5, 2, 3, 7, 3, 2, 0, 1, 2, 3, 6, 0, 4, 0, 3, 4, 5, 6, 7, 2
+};
+const pma::Vector<float> DecodedV25::psdValues = {
+    1.0f, 0.9f, 0.9f, 0.6f, 1.0f, 0.8f, 0.9f, 0.8f, 1.0f, 0.3f, 1.0f, 0.9f, 1.0f, 0.9f, 0.5f, 0.5f, 0.9f, 0.7f, 0.6f, 1.0f, 1.0f,
+    1.0f, 0.6f, 1.0f
+};
+// Behavior->Joints
+const pma::Vector<std::uint16_t> DecodedV25::jointRowCount = {
+    81u,  // MaxLOD-0 - MinLOD-1
+    54u,  // MaxLOD-1 - MinLOD-1
+    81u  // MaxLOD-0 - MinLOD-0
+};
+const std::uint16_t DecodedV25::jointColumnCount = 10u;
+const pma::Vector<pma::Matrix<std::uint16_t> > DecodedV25::jointVariableIndices = {
+    {  // MaxLOD-0 - MinLOD-1
+        {2, 3, 4, 5, 12, 13, 14, 18, 20, 36, 38, 39, 40, 41, 45, 46, 48, 49, 50, 55, 56, 57, 58, 59, 63, 71},
+        {2, 3, 4, 5, 12, 13, 14, 18, 20, 39, 40, 41, 48, 49, 50, 55, 56, 57, 58, 59}
+    },
+    {  // MaxLOD-1 - MinLOD-1
+        {2, 3, 4, 5, 12, 13, 14, 18, 20, 37, 38, 39, 40, 41}
+    },
+    {  // MaxLOD-0 - MinLOD-0
+        {2, 3, 4, 5, 12, 13, 14, 18, 20, 36, 38, 39, 40, 41, 45, 46, 48, 49, 50, 55, 56, 57, 58, 59, 63, 71}
+    }
+};
+const pma::Vector<pma::Matrix<std::uint16_t> > DecodedV25::jointGroupLODs = {
+    {  // Joint Group 0
+        {3, 3},  // MaxLOD-0 - MaxLOD-1
+        {3},  // MaxLOD-1 - MaxLOD-1
+        {3},  // MaxLOD-0 - MaxLOD-1
+    },
+    {  // Joint group 1
+        {4, 2},  // MaxLOD-0 - MaxLOD-1
+        {2},  // MaxLOD-1 - MaxLOD-1
+        {4}  // MaxLOD-0 - MaxLOD-0
+    },
+    {  // Joint group 2
+        {3, 2},  // MaxLOD-0 - MinLOD-1
+        {2},  // MaxLOD-1 - MinLOD-1
+        {3}  // MaxLOD-0 - MinLOD-0
+    },
+    {  // Joint group 3
+        {3, 0},  // MaxLOD-0 - MinLOD-1
+        {0},  // MaxLOD-1 - MinLOD-1
+        {3}  // MaxLOD-0 - MinLOD-0
+    }
+};
+const pma::Vector<pma::Vector<pma::Matrix<std::uint16_t> > > DecodedV25::jointGroupInputIndices = {
+    {  // Joint Group 0
+        {  // MaxLOD-0 - MaxLOD-1
+            {0, 1, 2, 3, 6, 7, 8},
+            {0, 1, 2, 3, 6, 7, 8}
+        },
+        {  // MaxLOD-1 - MaxLOD-1
+            {0, 1, 2, 3, 6, 7, 8}
+        },
+        {  // MaxLOD-0 - MaxLOD-0
+            {0, 1, 2, 3, 6, 7, 8}
+        }
+    },
+    {  // Joint Group 1
+        {  // MaxLOD-0 - MaxLOD-1
+            {3, 4, 7, 8, 9},
+            {3, 4, 7, 8, 9}
+        },
+        {  // MaxLOD-1 - MaxLOD-1
+            {3, 4, 7, 8, 9}
+        },
+        {  // MaxLOD-0 - MaxLOD-0
+            {3, 4, 7, 8, 9}
+        }
+    },
+    {  // Joint Group 2
+        {  // MaxLOD-0 - MaxLOD-1
+            {4, 5, 8, 9},
+            {4, 5, 8, 9}
+        },
+        {  // MaxLOD-1 - MaxLOD-1
+            {4, 5, 8, 9}
+        },
+        {  // MaxLOD-0 - MaxLOD-0
+            {4, 5, 8, 9}
+        }
+    },
+    {  // Joint Group 3
+        {  // MaxLOD-0 - MaxLOD-1
+            {2, 5, 6, 8},
+            {2, 5, 6, 8}
+        },
+        {  // MaxLOD-1 - MaxLOD-1
+            {}
+        },
+        {  // MaxLOD-0 - MaxLOD-0
+            {2, 5, 6, 8}
+        }
+    }
+};
+const pma::Vector<pma::Vector<pma::Matrix<std::uint16_t> > > DecodedV25::jointGroupOutputIndices = {
+    {  // Joint Group 0
+        {  // MaxLOD-0 - MaxLOD-1
+            {2, 3, 5},
+            {2, 3, 5}
+        },
+        {  // MaxLOD-1 - MaxLOD-1
+            {2, 3, 5}
+        },
+        {  // MaxLOD-0 - MaxLOD-0
+            {2, 3, 5},
+        }
+    },
+    {  // Joint Group 1
+        {  // MaxLOD-0 - MaxLOD-1
+            {18, 20, 36, 38},
+            {18, 20}
+        },
+        {  // MaxLOD-1 - MaxLOD-1
+            {18, 20}
+        },
+        {  // MaxLOD-0 - MaxLOD-0
+            {18, 20, 36, 38}
+        }
+    },
+    {  // Joint Group 2
+        {  // MaxLOD-0 - MaxLOD-1
+            {55, 56, 63},
+            {55, 56}
+        },
+        {  // MaxLOD-1 - MaxLOD-1
+            {37, 38}
+        },
+        {  // MaxLOD-0 - MaxLOD-0
+            {55, 56, 63}
+        }
+    },
+    {  // Joint Group 3
+        {  // MaxLOD-0 - MaxLOD-1
+            {45, 46, 71},
+            {}
+        },
+        {  // MaxLOD-1 - MaxLOD-1
+            {}
+        },
+        {  // MaxLOD-0 - MaxLOD-0
+            {45, 46, 71}
+        }
+    }
+};
+const pma::Vector<pma::Vector<pma::Matrix<float> > > DecodedV25::jointGroupValues = {
+    {  // Joint Group 0
+        {  // MaxLOD-0 - MaxLOD-1
+            {
+                0.00f, 0.05f, 0.10f, 0.15f, 0.20f, 0.25f, 0.30f,
+                0.35f, 0.40f, 0.45f, 0.50f, 0.55f, 0.60f, 0.65f,
+                0.70f, 0.75f, 0.80f, 0.85f, 0.90f, 0.95f, 1.00f
+            },
+            {
+                0.00f, 0.05f, 0.10f, 0.15f, 0.20f, 0.25f, 0.30f,
+                0.35f, 0.40f, 0.45f, 0.50f, 0.55f, 0.60f, 0.65f,
+                0.70f, 0.75f, 0.80f, 0.85f, 0.90f, 0.95f, 1.00f
+            }
+        },
+        {  // MaxLOD-1 - MinLOD-1
+            {
+                0.00f, 0.05f, 0.10f, 0.15f, 0.20f, 0.25f, 0.30f,
+                0.35f, 0.40f, 0.45f, 0.50f, 0.55f, 0.60f, 0.65f,
+                0.70f, 0.75f, 0.80f, 0.85f, 0.90f, 0.95f, 1.00f
+            }
+        },
+        {  // MaxLOD-0 - MinLOD-0
+            {
+                0.00f, 0.05f, 0.10f, 0.15f, 0.20f, 0.25f, 0.30f,
+                0.35f, 0.40f, 0.45f, 0.50f, 0.55f, 0.60f, 0.65f,
+                0.70f, 0.75f, 0.80f, 0.85f, 0.90f, 0.95f, 1.00f
+            }
+        }
+    },
+    {  // Joint group 1
+        {  // MaxLOD-0 - MaxLOD-1
+            {
+                0.01f, 0.02f, 0.03f, 0.04f, 0.05f,
+                0.06f, 0.07f, 0.08f, 0.09f, 0.10f,
+                0.11f, 0.12f, 0.13f, 0.14f, 0.15f,
+                0.16f, 0.17f, 0.18f, 0.19f, 0.20f
+            },
+            {
+                0.01f, 0.02f, 0.03f, 0.04f, 0.05f,
+                0.06f, 0.07f, 0.08f, 0.09f, 0.10f
+            }
+        },
+        {  // MaxLOD-1 - MinLOD-1
+            {
+                0.01f, 0.02f, 0.03f, 0.04f, 0.05f,
+                0.06f, 0.07f, 0.08f, 0.09f, 0.10f
+            }
+        },
+        {  // MaxLOD-0 - MinLOD-0
+            {
+                0.01f, 0.02f, 0.03f, 0.04f, 0.05f,
+                0.06f, 0.07f, 0.08f, 0.09f, 0.10f,
+                0.11f, 0.12f, 0.13f, 0.14f, 0.15f,
+                0.16f, 0.17f, 0.18f, 0.19f, 0.20f
+            }
+        }
+    },
+    {  // Joint group 2
+        {  // MaxLOD-0 - MaxLOD-1
+            {
+                0.31f, 0.36f, 0.42f, 0.47f,
+                0.53f, 0.58f, 0.64f, 0.69f,
+                0.75f, 0.80f, 0.86f, 0.91f
+            },
+            {
+                0.31f, 0.36f, 0.42f, 0.47f,
+                0.53f, 0.58f, 0.64f, 0.69f
+            }
+        },
+        {  // MaxLOD-1 - MinLOD-1
+            {
+                0.31f, 0.36f, 0.42f, 0.47f,
+                0.53f, 0.58f, 0.64f, 0.69f
+            }
+        },
+        {  // MaxLOD-0 - MinLOD-0
+            {
+                0.31f, 0.36f, 0.42f, 0.47f,
+                0.53f, 0.58f, 0.64f, 0.69f,
+                0.75f, 0.80f, 0.86f, 0.91f
+            }
+        }
+    },
+    {  // Joint group 3
+        {  // MaxLOD-0 - MaxLOD-1
+            {
+                0.31f, 0.36f, 0.42f, 0.47f,
+                0.53f, 0.58f, 0.64f, 0.69f,
+                0.75f, 0.80f, 0.86f, 0.91f
+            },
+            {
+            }
+        },
+        {  // MaxLOD-1 - MinLOD-1
+            {
+            }
+        },
+        {  // MaxLOD-0 - MinLOD-0
+            {
+                0.31f, 0.36f, 0.42f, 0.47f,
+                0.53f, 0.58f, 0.64f, 0.69f,
+                0.75f, 0.80f, 0.86f, 0.91f
+            }
+        }
+    }
+};
+const pma::Vector<pma::Vector<pma::Matrix<std::uint16_t> > > DecodedV25::jointGroupJointIndices = {
+    {  // Joint Group 0
+        {  // MaxLOD-0 - MaxLOD-1
+            {0},
+            {0}
+        },
+        {  // MaxLOD-1 - MinLOD-1
+            {0}
+        },
+        {  // MaxLOD-0 - MinLOD-0
+            {0}
+        }
+    },
+    {  // Joint Group 1
+        {  // MaxLOD-0 - MaxLOD-1
+            {2, 4},
+            {2}
+        },
+        {  // MaxLOD-1 - MinLOD-1
+            {2}
+        },
+        {  // MaxLOD-0 - MinLOD-0
+            {2, 4}
+        }
+    },
+    {  // Joint Group 2
+        {  // MaxLOD-0 - MaxLOD-1
+            {6, 7},
+            {6}
+        },
+        {  // MaxLOD-1 - MinLOD-1
+            {4}
+        },
+        {  // MaxLOD-0 - MinLOD-0
+            {6, 7}
+        }
+    },
+    {  // Joint Group 3
+        {  // MaxLOD-0 - MaxLOD-1
+            {5, 7},
+            {}
+        },
+        {  // MaxLOD-1 - MinLOD-1
+            {}
+        },
+        {  // MaxLOD-0 - MinLOD-0
+            {5, 7}
+        }
+    }
+};
+// Behavior->BlendShapes
+const pma::Matrix<std::uint16_t> DecodedV25::blendShapeLODs = {
+    {
+        {7, 4},  // MaxLOD-0 - MaxLOD-1
+        {4},  // MaxLOD-1 - MinLOD-1
+        {7}  // MaxLOD-0 - MinLOD-0
+    }
+};
+const pma::Vector<pma::Matrix<std::uint16_t> > DecodedV25::blendShapeInputIndices = {
+    {  // MaxLOD-0 - MaxLOD-1
+        {0, 1, 2, 3, 6, 7, 8},
+        {0, 1, 2, 3}
+    },
+    {  // MaxLOD-1 - MinLOD-1
+        {0, 1, 2, 3}
+    },
+    {  // MaxLOD-0 - MinLOD-0
+        {0, 1, 2, 3, 6, 7, 8}
+    }
+};
+const pma::Vector<pma::Matrix<std::uint16_t> > DecodedV25::blendShapeOutputIndices = {
+    {  // MaxLOD-0 - MaxLOD-1
+        {0, 1, 2, 3, 6, 7, 8},
+        {0, 1, 2, 3}
+    },
+    {  // MaxLOD-1 - MinLOD-1
+        {0, 1, 2, 3}
+    },
+    {  // MaxLOD-0 - MinLOD-0
+        {0, 1, 2, 3, 6, 7, 8}
+    }
+};
+// Behavior->AnimatedMaps
+const pma::Vector<std::uint16_t> DecodedV25::animatedMapCount = {
+    10,  // MaxLOD-0 - MaxLOD-1
+    4,  // MaxLOD-1 - MinLOD-1
+    10  // MaxLOD-0 - MinLOD-0
+};
+const pma::Matrix<std::uint16_t> DecodedV25::animatedMapLODs = {
+    {
+        {15, 6},  // MaxLOD-0 - MaxLOD-1
+        {6},  // MaxLOD-1 - MinLOD-1
+        {15}  // MaxLOD-0 - MinLOD-0
+    }
+};
+// Geometry
+const pma::Vector<std::uint32_t> DecodedV25::meshCount = {
+    3u,  // MaxLOD-0 - MaxLOD-1
+    1u,  // MaxLOD-1 - MinLOD-1
+    2u  // MaxLOD-0 - MinLOD-0
+};
+const pma::Vector<pma::Matrix<Vector3> > DecodedV25::vertexPositions = {
+    {  // MaxLOD-0 - MaxLOD-1
+        {  // Mesh-0
+            {7.0f, 7.0f, 7.0f},
+            {8.0f, 8.0f, 8.0f},
+            {9.0f, 9.0f, 9.0f}
+        },
+        {  // Mesh-1
+            {4.0f, 4.0f, 4.0f},
+            {5.0f, 5.0f, 5.0f},
+            {6.0f, 6.0f, 6.0f}
+        },
+        {  // Mesh-2
+            {1.0f, 1.0f, 1.0f},
+            {2.0f, 2.0f, 2.0f},
+            {3.0f, 3.0f, 3.0f}
+        }
+    },
+    {  // MaxLOD-1 - MinLOD-1
+        {  // Mesh-0 (Mesh-2 under MaxLOD-0)
+            {1.0f, 1.0f, 1.0f},
+            {2.0f, 2.0f, 2.0f},
+            {3.0f, 3.0f, 3.0f}
+        }
+    },
+    {  // MaxLOD-0 - MinLOD-0
+        {  // Mesh-0
+            {7.0f, 7.0f, 7.0f},
+            {8.0f, 8.0f, 8.0f},
+            {9.0f, 9.0f, 9.0f}
+        },
+        {  // Mesh-1
+            {4.0f, 4.0f, 4.0f},
+            {5.0f, 5.0f, 5.0f},
+            {6.0f, 6.0f, 6.0f}
+        }
+    }
+};
+const pma::Vector<pma::Matrix<TextureCoordinate> > DecodedV25::vertexTextureCoordinates = {
+    {  // MaxLOD-0 - MaxLOD-1
+        {  // Mesh-0
+            {7.0f, 7.0f},
+            {8.0f, 8.0f},
+            {9.0f, 9.0f}
+        },
+        {  // Mesh-1
+            {4.0f, 4.0f},
+            {5.0f, 5.0f},
+            {6.0f, 6.0f}
+        },
+        {  // Mesh-2
+            {1.0f, 1.0f},
+            {2.0f, 2.0f},
+            {3.0f, 3.0f}
+        }
+    },
+    {  // MaxLOD-1 - MinLOD-1
+        {  // Mesh-0 (Mesh-2 under MaxLOD-0)
+            {1.0f, 1.0f},
+            {2.0f, 2.0f},
+            {3.0f, 3.0f}
+        }
+    },
+    {  // MaxLOD-0 - MinLOD-0
+        {  // Mesh-0
+            {7.0f, 7.0f},
+            {8.0f, 8.0f},
+            {9.0f, 9.0f}
+        },
+        {  // Mesh-1
+            {4.0f, 4.0f},
+            {5.0f, 5.0f},
+            {6.0f, 6.0f}
+        }
+    }
+};
+const pma::Vector<pma::Matrix<Vector3> > DecodedV25::vertexNormals = {
+    {  // MaxLOD-0 - MaxLOD-1
+        {  // Mesh-0
+            {7.0f, 7.0f, 7.0f},
+            {8.0f, 8.0f, 8.0f},
+            {9.0f, 9.0f, 9.0f}
+        },
+        {  // Mesh-1
+            {4.0f, 4.0f, 4.0f},
+            {5.0f, 5.0f, 5.0f},
+            {6.0f, 6.0f, 6.0f}
+        },
+        {  // Mesh-2
+            {1.0f, 1.0f, 1.0f},
+            {2.0f, 2.0f, 2.0f},
+            {3.0f, 3.0f, 3.0f}
+        }
+    },
+    {  // MaxLOD-1 - MinLOD-1
+        {  // Mesh-0 (Mesh-2 under MaxLOD-0)
+            {1.0f, 1.0f, 1.0f},
+            {2.0f, 2.0f, 2.0f},
+            {3.0f, 3.0f, 3.0f}
+        }
+    },
+    {  // MaxLOD-0 - MinLOD-0
+        {  // Mesh-0
+            {7.0f, 7.0f, 7.0f},
+            {8.0f, 8.0f, 8.0f},
+            {9.0f, 9.0f, 9.0f}
+        },
+        {  // Mesh-1
+            {4.0f, 4.0f, 4.0f},
+            {5.0f, 5.0f, 5.0f},
+            {6.0f, 6.0f, 6.0f}
+        }
+    }
+};
+const pma::Vector<pma::Matrix<VertexLayout> > DecodedV25::vertexLayouts = {
+    {  // MaxLOD-0 - MaxLOD-1
+        {  // Mesh-0
+            {0, 0, 0},
+            {1, 1, 1},
+            {2, 2, 2}
+        },
+        {  // Mesh-1
+            {0, 0, 0},
+            {1, 1, 1},
+            {2, 2, 2}
+        },
+        {  // Mesh-2
+            {0, 0, 0},
+            {1, 1, 1},
+            {2, 2, 2}
+        }
+    },
+    {  // MaxLOD-1 - MinLOD-1
+        {  // Mesh-0 (Mesh-2 under MaxLOD-0)
+            {0, 0, 0},
+            {1, 1, 1},
+            {2, 2, 2}
+        }
+    },
+    {  // MaxLOD-0 - MinLOD-0
+        {  // Mesh-0
+            {0, 0, 0},
+            {1, 1, 1},
+            {2, 2, 2}
+        },
+        {  // Mesh-1
+            {0, 0, 0},
+            {1, 1, 1},
+            {2, 2, 2}
+        }
+    }};
+const pma::Matrix<pma::Matrix<std::uint32_t> > DecodedV25::faces = {
+    {  // MaxLOD-0 - MaxLOD-1
+        {  // Mesh-0
+            {0, 1, 2}
+        },
+        {  // Mesh-1
+            {0, 1, 2}
+        },
+        {  // Mesh-2
+            {0, 1, 2}
+        }
+    },
+    {  // MaxLOD-1 - MinLOD-1
+        {  // Mesh-0 (Mesh-2 under MaxLOD-0)
+            {0, 1, 2}
+        }
+    },
+    {  // MaxLOD-0 - MinLOD-0
+        {  // Mesh-0
+            {0, 1, 2}
+        },
+        {  // Mesh-1
+            {0, 1, 2}
+        }
+    }};
+const pma::Matrix<std::uint16_t> DecodedV25::maxInfluencePerVertex = {
+    {  // MaxLOD-0 - MaxLOD-1
+        8u,  // Mesh-0
+        8u,  // Mesh-1
+        8u  // Mesh-2
+    },
+    {  // MaxLOD-1 - MinLOD-1
+        8u  // Mesh-0 (Mesh-2 under MaxLOD-0)
+    },
+    {  // MaxLOD-0 - MinLOD-0
+        8u,  // Mesh-0
+        8u  // Mesh-1
+    }
+};
+const pma::Matrix<pma::Matrix<float> > DecodedV25::skinWeightsValues = {
+    {  // MaxLOD-0 - MinLOD-1
+        {  // Mesh-0
+            {0.7f, 0.1f, 0.2f},
+            {0.5f, 0.5f},
+            {0.4f, 0.6f}
+        },
+        {  // Mesh-1
+            {0.4f, 0.3f, 0.3f},
+            {0.8f, 0.2f},
+            {0.1f, 0.9f}
+        },
+        {  // Mesh-2
+            {0.1f, 0.3f, 0.6f},
+            {0.3f, 0.7f},
+            {0.2f, 0.8f}
+        }
+    },
+    {  // MaxLOD-1 - MinLOD-1
+        {  // Mesh-0 (Mesh-2 under MaxLOD-0)
+            {0.1f, 0.3f, 0.6f},
+            {1.0f},  // 0.3f normalized to 1.0f
+            {1.0f}  // 0.8f normalized to 1.0f
+        }
+    },
+    {  // MaxLOD-0 - MinLOD-0
+        {  // Mesh-0
+            {0.7f, 0.1f, 0.2f},
+            {0.5f, 0.5f},
+            {0.4f, 0.6f}
+        },
+        {  // Mesh-1
+            {0.4f, 0.3f, 0.3f},
+            {0.8f, 0.2f},
+            {0.1f, 0.9f}
+        }
+    }
+};
+const pma::Matrix<pma::Matrix<std::uint16_t> > DecodedV25::skinWeightsJointIndices = {
+    {  // MaxLOD-0 - MaxLOD-1
+        {  // Mesh-0
+            {0, 1, 2},
+            {3, 4},
+            {5, 6}
+        },
+        {  // Mesh-1
+            {0, 1, 2},
+            {3, 4},
+            {5, 6}
+        },
+        {  // Mesh-2
+            {0, 1, 2},
+            {3, 4},
+            {5, 6}
+        }
+    },
+    {  // MaxLOD-1 - MinLOD-1
+        {  // Mesh-0 (Mesh-2 under MaxLOD-0)
+            {0, 1, 2},
+            {3},
+            {4}
+        }
+    },
+    {  // MaxLOD-0 - MinLOD-0
+        {  // Mesh-0
+            {0, 1, 2},
+            {3, 4},
+            {5, 6}
+        },
+        {  // Mesh-1
+            {0, 1, 2},
+            {3, 4},
+            {5, 6}
+        }
+    }
+};
+const pma::Vector<pma::Matrix<std::uint16_t> > DecodedV25::correctiveBlendShapeIndices = {
+    {  // MaxLOD-0 - MaxLOD-1
+        {  // Mesh-0
+            2
+        },
+        {  // Mesh-1
+            2
+        },
+        {  // Mesh-2
+            2, 3
+        }
+    },
+    {  // MaxLOD-1 - MinLOD-1
+        {  // Mesh-0 (Mesh-2 under MaxLOD-0)
+            2
+        }
+    },
+    {  // MaxLOD-0 - MinLOD-0
+        {  // Mesh-0
+            2
+        },
+        {  // Mesh-1
+            2
+        }
+    }
+};
+const pma::Matrix<pma::Matrix<Vector3> > DecodedV25::correctiveBlendShapeDeltas = {
+    {  // MaxLOD-0 - MaxLOD-1
+        {  // Mesh-0
+            {  // Blendshape-0
+                {7.0f, 7.0f, 7.0f},
+                {8.0f, 8.0f, 8.0f},
+                {9.0f, 9.0f, 9.0f}
+            }
+        },
+        {  // Mesh-1
+            {  // Blendshape-0
+                {4.0f, 4.0f, 4.0f},
+                {5.0f, 5.0f, 5.0f},
+                {6.0f, 6.0f, 6.0f}
+            }
+        },
+        {  // Mesh-2
+            {  // Blendshape-0
+                {1.0f, 1.0f, 1.0f},
+                {2.0f, 2.0f, 2.0f},
+                {3.0f, 3.0f, 3.0f}
+            },
+            {  // Blendshape-1
+                {4.0f, 4.0f, 4.0f},
+                {5.0f, 5.0f, 5.0f}
+            }
+        }
+    },
+    {  // MaxLOD-1 - MinLOD-1
+        {  // Mesh-0 (Mesh-2 under MaxLOD-0)
+            {  // Blendshape-0
+                {1.0f, 1.0f, 1.0f},
+                {2.0f, 2.0f, 2.0f},
+                {3.0f, 3.0f, 3.0f}
+            }
+        }
+    },
+    {  // MaxLOD-0 - MinLOD-0
+        {  // Mesh-0
+            {  // Blendshape-0
+                {7.0f, 7.0f, 7.0f},
+                {8.0f, 8.0f, 8.0f},
+                {9.0f, 9.0f, 9.0f}
+            }
+        },
+        {  // Mesh-1
+            {  // Blendshape-0
+                {4.0f, 4.0f, 4.0f},
+                {5.0f, 5.0f, 5.0f},
+                {6.0f, 6.0f, 6.0f}
+            }
+        }
+    }
+};
+const pma::Matrix<pma::Matrix<std::uint32_t> > DecodedV25::correctiveBlendShapeVertexIndices = {
+    {  // MaxLOD-0 - MaxLOD-1
+        {  // Mesh-0
+            {0, 1, 2},  // Blendshape-0
+        },
+        {  // Mesh-1
+            {0, 1, 2},  // Blendshape-0
+        },
+        {  // Mesh-2
+            {0, 1, 2},  // Blendshape-0
+            {0, 2}  // Blendshape-1
+        }
+    },
+    {  // MaxLOD-1 - MinLOD-1
+        {  // Mesh-0 (Mesh-2 under MaxLOD-0)
+            {0, 1, 2}  // Blendshape-0
+        }
+    },
+    {  // MaxLOD-0 - MinLOD-0
+        {  // Mesh-0
+            {0, 1, 2},  // Blendshape-0
+        },
+        {  // Mesh-1
+            {0, 1, 2},  // Blendshape-0
+        }
+    }
+};
+// Machine learned behavior
+const pma::Vector<pma::String<char> > DecodedV25::mlControlNames = {
+    "MA", "MB", "MC", "MD", "ME", "MF", "MG", "MH", "MI"
+};
+const pma::Matrix<std::uint16_t>  DecodedV25::neuralNetworkIndicesPerLOD = {
+    {  // MaxLOD-0 - MaxLOD-1
+        0,  // Mesh-0 Region-0
+        1,  // Mesh-0 Region-1
+        2,  // Mesh-1 Region-0
+        3,  // Mesh-1 Region-1
+        4,  // Mesh-2 Region-0
+        5  // Mesh-2 Region-1
+    },
+    {  // MaxLOD-1 - MinLOD-1
+        0,  // Mesh-0 (Mesh-2 under MaxLOD-0) Region-0
+        1  // Mesh-0 (Mesh-2 under MaxLOD-0) Region-1
+    },
+    {  // MaxLOD-0 - MinLOD-0
+        0,  // Mesh-0 Region-0
+        1,  // Mesh-0 Region-1
+        2,  // Mesh-1 Region-0
+        3  // Mesh-1 Region-1
+    }
+};
+const DecodedV25::VectorOfCharStringMatrix DecodedV25::regionNames = {
+    {  // MaxLOD-0 - MinLOD-1
+        {  // Mesh-0
+            "RA", "RB"
+        },
+        {  // Mesh-1
+            "RC", "RD"
+        },
+        {  // Mesh-2
+            "RE", "RF"
+        }
+    },
+    {  // MaxLOD-1 - MinLOD-1
+        {  // Mesh-0 (Mesh-2 under MaxLOD-0)
+            "RE", "RF"
+        }
+    },
+    {  // MaxLOD-0 - MinLOD-0
+        {  // Mesh-0
+            "RA", "RB"
+        },
+        {  // Mesh-1
+            "RC", "RD"
+        }
+    }
+};
+const pma::Matrix<pma::Matrix<std::uint16_t> > DecodedV25::neuralNetworkIndicesPerMeshRegion = {
+    {  // MaxLOD-0 - MaxLOD-1
+        {  // Mesh-0
+            {  // Region-0
+                0
+            },
+            {  // Region-1
+                1
+            }
+        },
+        {  // Mesh-1
+            {  // Region-0
+                2
+            },
+            {  // Region-1
+                3
+            }
+        },
+        {  // Mesh-2
+            {  // Region-0
+                4
+            },
+            {  // Region-1
+                5
+            }
+        }
+    },
+    {  // MaxLOD-1 - MinLOD-1
+        {  // Mesh-0 (Mesh-2 under MaxLOD-0)
+            {  // Region-0
+                0  // (4 under MaxLOD-0)
+            },
+            {  // Region-1
+                1  // (5 under MaxLOD-0)
+            }
+        }
+    },
+    {  // MaxLOD-0 - MinLOD-0
+        {  // Mesh-0
+            {  // Region-0
+                0
+            },
+            {  // Region-1
+                1
+            }
+        },
+        {  // Mesh-1
+            {  // Region-0
+                2
+            },
+            {  // Region-1
+                3
+            }
+        }
+    }
+};
+const pma::Vector<pma::Matrix<std::uint16_t> > DecodedV25::neuralNetworkInputIndices = {
+    {  // MaxLOD-0 - MaxLOD-1
+        {  // Mesh-0 Region-0
+            0, 1
+        },
+        {  // Mesh-0 Region-1
+            2, 3
+        },
+        {  // Mesh-1 Region-0
+            4, 5
+        },
+        {  // Mesh-1 Region-1
+            6, 7
+        },
+        {  // Mesh-2 Region-0
+            8, 0
+        },
+        {  // Mesh-2 Region-1
+            4, 7
+        }
+    },
+    {  // MaxLOD-1 - MinLOD-1
+        {  // Mesh-0 (Mesh-2 under MaxLOD-0) Region-0
+            8, 0
+        },
+        {  // Mesh-0 (Mesh-2 under MaxLOD-0) Region-1
+            4, 7
+        }
+    },
+    {  // MaxLOD-0 - MinLOD-0
+        {  // Mesh-0 Region-0
+            0, 1
+        },
+        {  // Mesh-0 Region-1
+            2, 3
+        },
+        {  // Mesh-1 Region-0
+            4, 5
+        },
+        {  // Mesh-1 Region-1
+            6, 7
+        }
+    }
+};
+const pma::Vector<pma::Matrix<std::uint16_t> > DecodedV25::neuralNetworkOutputIndices = {
+    {  // MaxLOD-0 - MaxLOD-1
+        {  // Mesh-0 Region-0
+            9
+        },
+        {  // Mesh-0 Region-1
+            10
+        },
+        {  // Mesh-1 Region-0
+            11
+        },
+        {  // Mesh-1 Region-1
+            12
+        },
+        {  // Mesh-2 Region-0
+            13
+        },
+        {  // Mesh-2 Region-1
+            14
+        }
+    },
+    {  // MaxLOD-1 - MinLOD-1
+        {  // Mesh-0 (Mesh-2 under MaxLOD-0) Region-0
+            13
+        },
+        {  // Mesh-0 (Mesh-2 under MaxLOD-0) Region-1
+            14
+        }
+    },
+    {  // MaxLOD-0 - MinLOD-0
+        {  // Mesh-0 Region-0
+            9
+        },
+        {  // Mesh-0 Region-1
+            10
+        },
+        {  // Mesh-1 Region-0
+            11
+        },
+        {  // Mesh-1 Region-1
+            12
+        }
+    }
+};
+const pma::Matrix<std::uint16_t> DecodedV25::neuralNetworkLayerCount = {
+    {  // MaxLOD-0 - MaxLOD-1
+        2,  // Mesh-0 Region-0
+        2,  // Mesh-0 Region-1
+        2,  // Mesh-1 Region-0
+        2,  // Mesh-1 Region-1
+        2,  // Mesh-2 Region-0
+        2  // Mesh-2 Region-1
+    },
+    {  // MaxLOD-1 - MinLOD-1
+        2,  // Mesh-0 (Mesh-2 under MaxLOD-0) Region-0
+        2  // Mesh-0 (Mesh-2 under MaxLOD-0) Region-1
+    },
+    {  // MaxLOD-0 - MinLOD-0
+        2,  // Mesh-0 Region-0
+        2,  // Mesh-0 Region-1
+        2,  // Mesh-1 Region-0
+        2  // Mesh-1 Region-1
+    }
+};
+const pma::Vector<pma::Matrix<std::uint16_t> > DecodedV25::neuralNetworkActivationFunction = {
+    {  // MaxLOD-0 - MaxLOD-1
+        {  // Mesh-0 Region-0
+            1, 1
+        },
+        {  // Mesh-0 Region-1
+            1, 1
+        },
+        {  // Mesh-1  Region-0
+            1, 1
+        },
+        {  // Mesh-1 Region-1
+            1, 1
+        },
+        {  // Mesh-2 Region-0
+            1, 1
+        },
+        {  // Mesh-2 Region-1
+            1, 1
+        }
+    },
+    {  // MaxLOD-1 - MinLOD-1
+        {  // Mesh-0 (Mesh-2 under MaxLOD-0) Region-0
+            1, 1
+        },
+        {  // Mesh-0 (Mesh-2 under MaxLOD-0) Region-1
+            1, 1
+        }
+    },
+    {  // MaxLOD-0 - MinLOD-0
+        {  // Mesh-0  Region-0
+            1, 1
+        },
+        {  // Mesh-0 Region-1
+            1, 1
+        },
+        {  // Mesh-1 Region-0
+            1, 1
+        },
+        {  // Mesh-1 Region-1
+            1, 1
+        }
+    }
+};
+const pma::Matrix<pma::Matrix<float> > DecodedV25::neuralNetworkActivationFunctionParameters = {
+    {  // MaxLOD-0 - MaxLOD-1
+        {  // Mesh-0  Region-0
+            {0.5f},  // Layer-0
+            {0.5f}  // Layer-1
+        },
+        {  // Mesh-0 Region-1
+            {1.0f},  // Layer-0
+            {1.0f}  // Layer-1
+        },
+        {  // Mesh-1 Region-0
+            {1.0f},  // Layer-0
+            {1.0f}  // Layer-1
+        },
+        {  // Mesh-1 Region-1
+            {0.5f},  // Layer-0
+            {0.5f}  // Layer-1
+        },
+        {  // Mesh-2 Region-0
+            {0.5f},  // Layer-0
+            {0.5f}  // Layer-1
+        },
+        {  // Mesh-2 Region-1
+            {1.0f},  // Layer-0
+            {1.0f}  // Layer-1
+        }
+    },
+    {  // MaxLOD-1 - MinLOD-1
+        {  // Mesh-0 (Mesh-2 under MaxLOD-0) Region-0
+            {0.5f},  // Layer-0
+            {0.5f}  // Layer-1
+        },
+        {  // Mesh-0 (Mesh-2 under MaxLOD-0) Region-1
+            {1.0f},  // Layer-0
+            {1.0f}  // Layer-1
+        }
+    },
+    {  // MaxLOD-0 - MinLOD-0
+        {  // Mesh-0 Region-0
+            {0.5f},  // Layer-0
+            {0.5f}  // Layer-1
+        },
+        {  // Mesh-0 Region-1
+            {1.0f},  // Layer-0
+            {1.0f}  // Layer-1
+        },
+        {  // Mesh-1 Region-0
+            {1.0f},  // Layer-0
+            {1.0f}  // Layer-1
+        },
+        {  // Mesh-1 Region-1
+            {0.5f},  // Layer-0
+            {0.5f}  // Layer-1
+        }
+    }
+};
+const pma::Matrix<pma::Matrix<float> > DecodedV25::neuralNetworkBiases = {
+    {  // MaxLOD-0 - MaxLOD-1
+        {  // Mesh-0 Region-0
+            {1.0f, 1.0f},  // Layer-0
+            {1.0f}  // Layer-1
+        },
+        {  // Mesh-0 Region-1
+            {0.5f, 0.5f},  // Layer-0
+            {0.5f}  // Layer-1
+        },
+        {  // Mesh-1 Region-0
+            {0.5f, 0.5f},  // Layer-0
+            {0.5f}  // Layer-1
+        },
+        {  // Mesh-1 Region-1
+            {1.0f, 1.0f},  // Layer-0
+            {1.0f}  // Layer-1
+        },
+        {  // Mesh-2 Region-0
+            {1.0f, 1.0f},  // Layer-0
+            {1.0f}  // Layer-1
+        },
+        {  // Mesh-2 Region-1
+            {0.5f, 0.5f},  // Layer-0
+            {0.5f}  // Layer-1
+        }
+    },
+    {  // MaxLOD-1 - MinLOD-1
+        {  // Mesh-0 (Mesh-2 under MaxLOD-0) Region-0
+            {1.0f, 1.0f},  // Layer-0
+            {1.0f}  // Layer-1
+        },
+        {  // Mesh-0 (Mesh-2 under MaxLOD-0) Region-1
+            {0.5f, 0.5f},  // Layer-0
+            {0.5f}  // Layer-1
+        }
+    },
+    {  // MaxLOD-0 - MinLOD-0
+        {  // Mesh-0 Region-0
+            {1.0f, 1.0f},  // Layer-0
+            {1.0f}  // Layer-1
+        },
+        {  // Mesh-0 Region-1
+            {0.5f, 0.5f},  // Layer-0
+            {0.5f}  // Layer-1
+        },
+        {  // Mesh-1 Region-0
+            {0.5f, 0.5f},  // Layer-0
+            {0.5f}  // Layer-1
+        },
+        {  // Mesh-1 Region-1
+            {1.0f, 1.0f},  // Layer-0
+            {1.0f}  // Layer-1
+        }
+    }
+};
+const pma::Matrix<pma::Matrix<float> > DecodedV25::neuralNetworkWeights = {
+    {  // MaxLOD-0 - MaxLOD-1
+        {  // Mesh-0 Region-0
+            {0.5f, 0.5f, 0.5f, 0.5f},  // Layer-0
+            {0.5f, 0.5f}  // Layer-1
+        },
+        {  // Mesh-0 Region-1
+            {1.0f, 1.0f, 1.0f, 1.0f},  // Layer-0
+            {1.0f, 1.0f}  // Layer-1
+        },
+        {  // Mesh-1 Region-0
+            {1.0f, 1.0f, 1.0f, 1.0f},  // Layer-0
+            {1.0f, 1.0f}  // Layer-1
+        },
+        {  // Mesh-1 Region-1
+            {0.5f, 0.5f, 0.5f, 0.5f},  // Layer-0
+            {0.5f, 0.5f}  // Layer-1
+        },
+        {  // Mesh-2 Region-0
+            {0.5f, 0.5f, 0.5f, 0.5f},  // Layer-0
+            {0.5f, 0.5f}  // Layer-1
+        },
+        {  // Mesh-2 Region-1
+            {1.0f, 1.0f, 1.0f, 1.0f},  // Layer-0
+            {1.0f, 1.0f}  // Layer-1
+        }
+    },
+    {  // MaxLOD-1 - MinLOD-1
+        {  // Mesh-0 (Mesh-2 under MaxLOD-0) Region-0
+            {0.5f, 0.5f, 0.5f, 0.5f},  // Layer-0
+            {0.5f, 0.5f}  // Layer-1
+        },
+        {  // Mesh-0 Region-1
+            {1.0f, 1.0f, 1.0f, 1.0f},  // Layer-0
+            {1.0f, 1.0f}  // Layer-1
+        }
+    },
+    {  // MaxLOD-0 - MinLOD-0
+        {  // Mesh-0 Region-0
+            {0.5f, 0.5f, 0.5f, 0.5f},  // Layer-0
+            {0.5f, 0.5f}  // Layer-1
+        },
+        {  // Mesh-0 Region-1
+            {1.0f, 1.0f, 1.0f, 1.0f},  // Layer-0
+            {1.0f, 1.0f}  // Layer-1
+        },
+        {  // Mesh-1 Region-0
+            {1.0f, 1.0f, 1.0f, 1.0f},  // Layer-0
+            {1.0f, 1.0f}  // Layer-1
+        },
+        {  // Mesh-1 Region-1
+            {0.5f, 0.5f, 0.5f, 0.5f},  // Layer-0
+            {0.5f, 0.5f}  // Layer-1
+        }
+    }
+};
+
+const pma::Matrix<std::uint16_t> DecodedV25::solverIndicesPerLOD {
+    {  // MaxLOD-0 - MinLOD-1
+        0,
+        1,
+        2
+    },
+    {  // MaxLOD-1 - MinLOD-1
+        1,
+        2
+    },
+    {  // MaxLOD-0 - MinLOD-0
+        0,
+        1
+    }
+};
+
+const pma::Vector<pma::String<char> > DecodedV25::solverNames {
+    "RSA", "RSB", "RSC"
+};
+
+const pma::Matrix<std::uint16_t> DecodedV25::solverRawControlIndices {
+    // Solver 0
+    {11, 12},
+    // Solver 1
+    {3},
+    // Solver 2
+    {22, 23}
+};
+
+const pma::Matrix<std::uint16_t> DecodedV25::solverPoseIndices {
+    // Solver 0
+    {0, 1, 2},
+    // Solver 1
+    {3, 4},
+    // Solver 2
+    {5, 6, 7},
+};
+
+const pma::Vector<float> DecodedV25::solverRadius {
+    1.0f, 2.0f, 1.0f
+};
+
+const pma::Vector<float> DecodedV25::solverWeightThreshold {
+    1.0f, 2.0f, 1.0f
+};
+
+const pma::Vector<std::uint16_t> DecodedV25::solverType {
+    0, 1, 0
+};
+
+const pma::Vector<std::uint16_t> DecodedV25::solverAutomaticRadius {
+    0, 0, 0
+};
+
+const pma::Vector<std::uint16_t> DecodedV25::solverDistanceMethod {
+    1, 3, 1
+};
+
+const pma::Vector<std::uint16_t> DecodedV25::solverNormalizeMethod {
+    0, 1, 0
+};
+
+const pma::Vector<std::uint16_t> DecodedV25::solverFunctionType {
+    2, 2, 0
+};
+
+const pma::Vector<std::uint16_t> DecodedV25::solverTwistAxis {
+    0, 1, 0
+};
+
+const pma::Vector<pma::String<char> > DecodedV25::poseNames {
+    "RA", "RB", "RC", "RD", "RE", "RF", "RG", "RH"
+};
+
+const Vector<float> DecodedV25::poseScale {
+    0.0f,  // Pose 0 (RA)
+    1.0f,  // Pose 1 (RB)
+    2.0f,  // Pose 2 (RC)
+    2.0f,  // Pose 3 (RD)
+    1.0f,  // Pose 4 (RE)
+    1.0f,  // Pose 5 (RF)
+    1.0f,  // Pose 6 (RG)
+    0.5f  // Pose 7 (RH)
+};
+
+const pma::Vector<std::uint16_t> DecodedV25::poseDistanceMethod {
+    0,  // Pose 0 (RA)
+    1,  // Pose 1 (RB)
+    2,  // Pose 2 (RC)
+    3,  // Pose 0 (RD)
+    4,  // Pose 1 (RE)
+    1,  // Pose 0 (RF)
+    2,  // Pose 1 (RG)
+    2  // Pose 2 (RH)
+
+};
+
+const pma::Vector<std::uint16_t> DecodedV25::poseFunctionType {
+    5,  // Pose 0 (RA)
+    4,  // Pose 1 (RB)
+    3,  // Pose 2 (RC)
+    2,  // Pose 3 (RD)
+    1,  // Pose 4 (RE)
+    0,  // Pose 5 (RF)
+    1,  // Pose 6 (RG)
+    2  // Pose 7 (RH)
+};
+
+const Matrix<float> DecodedV25::solverRawControlValues {
+    {  // Solver 0
+        // Pose 0 (RA)
+        2.0f,  // Raw control index 11
+        0.0f,  // Raw control index 12
+        // Pose 1 (RB)
+        1.0f,  // Raw control index 11
+        1.0f,  // Raw control index 12
+        // Pose 2 (RC)
+        3.0f,  // Raw control index 11
+        -3.0f,  // Raw control index 12
+    },
+    {  // Solver 1
+        // Pose 0 (RD)
+        0.0f,  // Raw control index 3
+        // Pose 1 (RE)
+        4.0f,  // Raw control index 3
+    },
+    {  // Solver 2
+        // Pose 0 (RF)
+        2.0f,  // Raw control index 22
+        0.0f,  // Raw control index 23
+        // Pose 1 (RG)
+        1.0f,  // Raw control index 22
+        1.0f,  // Raw control index 23
+        // Pose 2 (RH)
+        3.0f,  // Raw control index 22
+        -3.0f,  // Raw control index 23
+    },
+};
+
+const pma::Vector<pma::String<char> > DecodedV25::poseControlNames = {
+    "PA", "PB", "PC", "PD", "PE", "PF", "PG", "PH", "PI"
+};
+
+const Matrix<std::uint16_t> DecodedV25::poseInputControlIndices = {
+    {0},
+    {1},
+    {2},
+    {3},
+    {4},
+    {5},
+    {6},
+    {7},
+};
+
+const Matrix<std::uint16_t> DecodedV25::poseOutputControlIndices = {
+    {8},
+    {9},
+    {10},
+    {11},
+    {12},
+    {13},
+    {14},
+    {15, 16},
+};
+
+const Matrix<float> DecodedV25::poseOutputControlWeights = {
+    {1.0f},
+    {1.0f},
+    {1.0f},
+    {1.0f},
+    {1.0f},
+    {1.0f},
+    {1.0f},
+    {0.5f, 0.5f}
+};
+
+const pma::Matrix<TranslationRepresentation> DecodedV25::jointTranslationRepresentation = {
+    {  // MaxLOD-0 - MinLOD-1
+        TranslationRepresentation::Vector,
+        TranslationRepresentation::Vector,
+        TranslationRepresentation::Vector,
+        TranslationRepresentation::Vector,
+        TranslationRepresentation::Vector,
+        TranslationRepresentation::Vector,
+        TranslationRepresentation::Vector,
+        TranslationRepresentation::Vector,
+        TranslationRepresentation::Vector
+    },
+    {  // MaxLOD-1 - MinLOD-1
+        TranslationRepresentation::Vector,
+        TranslationRepresentation::Vector,
+        TranslationRepresentation::Vector,
+        TranslationRepresentation::Vector,
+        TranslationRepresentation::Vector,
+        TranslationRepresentation::Vector,
+    },
+    {  // MaxLOD-0 - MinLOD-0
+        TranslationRepresentation::Vector,
+        TranslationRepresentation::Vector,
+        TranslationRepresentation::Vector,
+        TranslationRepresentation::Vector,
+        TranslationRepresentation::Vector,
+        TranslationRepresentation::Vector,
+        TranslationRepresentation::Vector,
+        TranslationRepresentation::Vector,
+        TranslationRepresentation::Vector
+    }
+};
+
+const pma::Matrix<RotationRepresentation> DecodedV25::jointRotationRepresentation = {
+    {  // MaxLOD-0 - MinLOD-1
+        RotationRepresentation::EulerAngles,  // JA
+        RotationRepresentation::EulerAngles,  // JB
+        RotationRepresentation::Quaternion,  // JC
+        RotationRepresentation::Quaternion,  // JD
+        RotationRepresentation::EulerAngles,  // JE
+        RotationRepresentation::EulerAngles,  // JF
+        RotationRepresentation::EulerAngles,  // JG
+        RotationRepresentation::Quaternion,  // JH
+        RotationRepresentation::EulerAngles  // JI
+    },
+    {  // MaxLOD-1 - MinLOD-1
+        RotationRepresentation::EulerAngles,  // JA
+        RotationRepresentation::EulerAngles,  // JB
+        RotationRepresentation::Quaternion,  // JC
+        RotationRepresentation::Quaternion,  // JD
+        RotationRepresentation::EulerAngles,  // JG
+        RotationRepresentation::EulerAngles  // JI
+    },
+    {  // MaxLOD-0 - MinLOD-0
+        RotationRepresentation::EulerAngles,  // JA
+        RotationRepresentation::EulerAngles,  // JB
+        RotationRepresentation::Quaternion,  // JC
+        RotationRepresentation::Quaternion,  // JD
+        RotationRepresentation::EulerAngles,  // JE
+        RotationRepresentation::EulerAngles,  // JF
+        RotationRepresentation::EulerAngles,  // JG
+        RotationRepresentation::Quaternion,  // JH
+        RotationRepresentation::EulerAngles  // JI
+    }
+};
+
+const pma::Matrix<ScaleRepresentation> DecodedV25::jointScaleRepresentation = {
+    {  // MaxLOD-0 - MinLOD-1
+        ScaleRepresentation::Vector,
+        ScaleRepresentation::Vector,
+        ScaleRepresentation::Vector,
+        ScaleRepresentation::Vector,
+        ScaleRepresentation::Vector,
+        ScaleRepresentation::Vector,
+        ScaleRepresentation::Vector,
+        ScaleRepresentation::Vector,
+        ScaleRepresentation::Vector
+    },
+    {  // MaxLOD-1 - MinLOD-1
+        ScaleRepresentation::Vector,
+        ScaleRepresentation::Vector,
+        ScaleRepresentation::Vector,
+        ScaleRepresentation::Vector,
+        ScaleRepresentation::Vector,
+        ScaleRepresentation::Vector,
+    },
+    {  // MaxLOD-0 - MinLOD-0
+        ScaleRepresentation::Vector,
+        ScaleRepresentation::Vector,
+        ScaleRepresentation::Vector,
+        ScaleRepresentation::Vector,
+        ScaleRepresentation::Vector,
+        ScaleRepresentation::Vector,
+        ScaleRepresentation::Vector,
+        ScaleRepresentation::Vector,
+        ScaleRepresentation::Vector
+    }
+};
+
+const pma::Vector<pma::Matrix<float> > DecodedV25::swingBlendWeights = {
+    {  // MaxLOD-0 - MinLOD-1
+        {1.0f, 2.0f},
+        {-2.0f, -1.0f},
+        {1.0f}
+    },
+    {  // MaxLOD-1 - MinLOD-0
+        {1.0f, 2.0f},
+        {-1.0f},
+    },
+    {  // MaxLOD-0 - MinLOD-0
+        {1.0f, 2.0f},
+        {-2.0f, -1.0f},
+        {1.0f}
+    }
+};
+
+const pma::Vector<pma::Matrix<std::uint16_t> > DecodedV25::swingOutputJointIndices = {
+    {  // MaxLOD-0 - MinLOD-1
+        {0u, 1u},
+        {4u, 6u},
+        {5u}
+    },
+    {  // MaxLOD-1 - MinLOD-0
+        {0u, 1u},
+        {4u}
+    },
+    {  // MaxLOD-0 - MinLOD-0
+        {0u, 1u},
+        {4u, 6u},
+        {5u}
+    }
+};
+
+const pma::Vector<pma::Matrix<std::uint16_t> > DecodedV25::swingInputControlIndices = {
+    {  // MaxLOD-0 - MinLOD-1
+        {5u, 6u, 7u, 8u},
+        {11u, 12u, 13u, 14u},
+        {27u, 28u, 29u, 30u}
+    },
+    {  // MaxLOD-1 - MinLOD-0
+        {5u, 6u, 7u, 8u},
+        {11u, 12u, 13u, 14u},
+        {27u, 28u, 29u, 30u}
+    },
+    {  // MaxLOD-0 - MinLOD-0
+        {5u, 6u, 7u, 8u},
+        {11u, 12u, 13u, 14u},
+        {27u, 28u, 29u, 30u}
+    }
+};
+
+const pma::Matrix<TwistAxis> DecodedV25::swingTwistAxes = {
+    {  // MaxLOD-0 - MinLOD-1
+        TwistAxis::X,
+        TwistAxis::Y,
+        TwistAxis::Z
+    },
+    {  // MaxLOD-1 - MinLOD-0
+        TwistAxis::X,
+        TwistAxis::Y,
+        TwistAxis::Z
+    },
+    {  // MaxLOD-0 - MinLOD-0
+        TwistAxis::X,
+        TwistAxis::Y,
+        TwistAxis::Z
+    }
+};
+
+const pma::Vector<pma::Matrix<float> > DecodedV25::twistBlendWeights = {
+    {  // MaxLOD-0 - MinLOD-1
+        {1.0f, 2.0f},
+        {-2.0f, -1.0f},
+        {1.0f}
+    },
+    {  // MaxLOD-1 - MinLOD-0
+        {1.0f, 2.0f},
+        {-1.0f},
+    },
+    {  // MaxLOD-0 - MinLOD-0
+        {1.0f, 2.0f},
+        {-2.0f, -1.0f},
+        {1.0f}
+    }
+};
+
+const pma::Vector<pma::Matrix<std::uint16_t> > DecodedV25::twistOutputJointIndices = {
+    {  // MaxLOD-0 - MinLOD-1
+        {0u, 1u},
+        {4u, 6u},
+        {5u}
+    },
+    {  // MaxLOD-1 - MinLOD-0
+        {0u, 1u},
+        {4u}
+    },
+    {  // MaxLOD-0 - MinLOD-0
+        {0u, 1u},
+        {4u, 6u},
+        {5u}
+    }
+};
+
+const pma::Vector<pma::Matrix<std::uint16_t> > DecodedV25::twistInputControlIndices = {
+    {  // MaxLOD-0 - MinLOD-1
+        {5u, 6u, 7u, 8u},
+        {11u, 12u, 13u, 14u},
+        {27u, 28u, 29u, 30u}
+    },
+    {  // MaxLOD-1 - MinLOD-0
+        {5u, 6u, 7u, 8u},
+        {11u, 12u, 13u, 14u},
+        {27u, 28u, 29u, 30u}
+    },
+    {  // MaxLOD-0 - MinLOD-0
+        {5u, 6u, 7u, 8u},
+        {11u, 12u, 13u, 14u},
+        {27u, 28u, 29u, 30u}
+    }
+};
+
+const pma::Matrix<TwistAxis> DecodedV25::twistTwistAxes = {
+    {  // MaxLOD-0 - MinLOD-1
+        TwistAxis::X,
+        TwistAxis::Y,
+        TwistAxis::Z
+    },
+    {  // MaxLOD-1 - MinLOD-0
+        TwistAxis::X,
+        TwistAxis::Y,
+        TwistAxis::Z
+    },
+    {  // MaxLOD-0 - MinLOD-0
+        TwistAxis::X,
+        TwistAxis::Y,
+        TwistAxis::Z
+    }
+};
+
+#ifdef __clang__
+    #pragma clang diagnostic pop
+#endif
+
+std::size_t DecodedV25::lodConstraintToIndex(std::uint16_t maxLOD, std::uint16_t minLOD) {
+    // Relies on having only TWO available LODs (0, 1)
+    return (minLOD == 1u ? maxLOD : 2ul);
+}
+
+RawJoints DecodedV25::getJoints(std::uint16_t currentMaxLOD, std::uint16_t currentMinLOD, pma::MemoryResource* memRes) {
+    const auto srcIndex = lodConstraintToIndex(currentMaxLOD, currentMinLOD);
+    RawJoints joints{memRes};
+    joints.rowCount = jointRowCount[srcIndex];
+    joints.colCount = jointColumnCount;
+    for (std::size_t i = 0ul; i < jointGroupLODs.size(); ++i) {
+        RawJointGroup jntGrp{memRes};
+        jntGrp.lods.assign(jointGroupLODs[i][srcIndex].begin(),
+                           jointGroupLODs[i][srcIndex].end());
+        jntGrp.inputIndices.assign(jointGroupInputIndices[i][srcIndex][0ul].begin(),
+                                   jointGroupInputIndices[i][srcIndex][0ul].end());
+        jntGrp.outputIndices.assign(jointGroupOutputIndices[i][srcIndex][0ul].begin(),
+                                    jointGroupOutputIndices[i][srcIndex][0ul].end());
+        jntGrp.values.assign(jointGroupValues[i][srcIndex][0ul].begin(),
+                             jointGroupValues[i][srcIndex][0ul].end());
+        jntGrp.jointIndices.assign(jointGroupJointIndices[i][srcIndex][0ul].begin(),
+                                   jointGroupJointIndices[i][srcIndex][0ul].end());
+        joints.jointGroups.push_back(std::move(jntGrp));
+    }
+
+    return joints;
+}
+
+RawBlendShapeChannels DecodedV25::getBlendShapes(std::uint16_t currentMaxLOD,
+                                                 std::uint16_t currentMinLOD,
+                                                 pma::MemoryResource* memRes) {
+    RawBlendShapeChannels blendShapes{memRes};
+    const auto srcIndex = lodConstraintToIndex(currentMaxLOD, currentMinLOD);
+    blendShapes.lods.assign(blendShapeLODs[srcIndex].begin(),
+                            blendShapeLODs[srcIndex].end());
+    blendShapes.inputIndices.assign(blendShapeInputIndices[srcIndex][0ul].begin(),
+                                    blendShapeInputIndices[srcIndex][0ul].end());
+    blendShapes.outputIndices.assign(blendShapeOutputIndices[srcIndex][0ul].begin(),
+                                     blendShapeOutputIndices[srcIndex][0ul].end());
+    return blendShapes;
+}
+
+RawConditionalTable DecodedV25::getConditionals(std::uint16_t currentMaxLOD,
+                                                std::uint16_t currentMinLOD,
+                                                pma::MemoryResource* memRes) {
+    RawConditionalTable conditionals{memRes};
+    const auto srcIndex = lodConstraintToIndex(currentMaxLOD, currentMinLOD);
+    conditionals.inputIndices.assign(conditionalInputIndices[srcIndex][0ul].begin(),
+                                     conditionalInputIndices[srcIndex][0ul].end());
+    conditionals.outputIndices.assign(conditionalOutputIndices[srcIndex][0ul].begin(),
+                                      conditionalOutputIndices[srcIndex][0ul].end());
+    conditionals.fromValues.assign(conditionalFromValues[srcIndex][0ul].begin(),
+                                   conditionalFromValues[srcIndex][0ul].end());
+    conditionals.toValues.assign(conditionalToValues[srcIndex][0ul].begin(),
+                                 conditionalToValues[srcIndex][0ul].end());
+    conditionals.slopeValues.assign(conditionalSlopeValues[srcIndex][0ul].begin(),
+                                    conditionalSlopeValues[srcIndex][0ul].end());
+    conditionals.cutValues.assign(conditionalCutValues[srcIndex][0ul].begin(),
+                                  conditionalCutValues[srcIndex][0ul].end());
+    return conditionals;
+}
+
+RawAnimatedMaps DecodedV25::getAnimatedMaps(std::uint16_t currentMaxLOD, std::uint16_t currentMinLOD,
+                                            pma::MemoryResource* memRes) {
+    RawAnimatedMaps animatedMaps{memRes};
+    const auto srcIndex = lodConstraintToIndex(currentMaxLOD, currentMinLOD);
+    animatedMaps.lods.assign(animatedMapLODs[srcIndex].begin(),
+                             animatedMapLODs[srcIndex].end());
+    animatedMaps.conditionals = getConditionals(currentMaxLOD, currentMinLOD, memRes);
+    return animatedMaps;
+}
+
+}  // namespace dna
diff --git a/Engine/Plugins/Animation/RigLogic/Source/RigLogicLibTest/Private/dnatests/Fixturesv25.h b/Engine/Plugins/Animation/RigLogic/Source/RigLogicLibTest/Private/dnatests/Fixturesv25.h
new file mode 100644
--- /dev/null
+++ b/Engine/Plugins/Animation/RigLogic/Source/RigLogicLibTest/Private/dnatests/Fixturesv25.h
@@ -0,0 +1,176 @@
+// Copyright Epic Games, Inc. All Rights Reserved.
+
+#pragma once
+
+#include "dnatests/Defs.h"
+
+#include "dna/DNA.h"
+#include "dna/Reader.h"
+
+#include <pma/TypeDefs.h>
+
+namespace dna {
+
+struct RawV25 {
+    static const unsigned char header[155];
+    static const unsigned char descriptor[87];
+    static const unsigned char definition[794];
+    static const unsigned char controls[2];
+    static const unsigned char conditionals[324];
+    static const unsigned char psds[204];
+    static const unsigned char joints[444];
+    static const unsigned char blendshapes[44];
+    static const unsigned char animatedmaps[8];
+    static const unsigned char geometry[1080];
+    static const unsigned char machineLearnedBehavior[762];
+    static const unsigned char rbfBehavior[327];
+    static const unsigned char rbfBehaviorExt[228];
+    static const unsigned char jointBehaviorMetadata[58];
+    static const unsigned char twistSwingBehavior[200];
+
+    static std::vector<char> getBytes();
+};
+
+struct RawV24DowngradedFromV25 {
+    static std::vector<char> getBytes();
+};
+
+struct DecodedV25 {
+    using VectorOfCharStringMatrix = pma::Vector<pma::Matrix<pma::String<char> > >;
+    using StringPair = std::pair<pma::String<char>, pma::String<char> >;
+
+    // Descriptor
+    static const pma::String<char> name;
+    static const Archetype archetype;
+    static const Gender gender;
+    static const std::uint16_t age;
+
+    static const pma::Vector<StringPair> metadata;
+    static const TranslationUnit translationUnit;
+    static const RotationUnit rotationUnit;
+    static const CoordinateSystem coordinateSystem;
+    static const std::uint16_t lodCount[3ul];
+    static const std::uint16_t maxLODs[3ul];
+    static const pma::String<char> complexity;
+    static const pma::String<char> dbName;
+
+    // Definition
+    static const pma::Vector<pma::String<char> > guiControlNames;
+    static const pma::Vector<pma::String<char> > rawControlNames;
+    static const VectorOfCharStringMatrix jointNames;
+    static const VectorOfCharStringMatrix blendShapeNames;
+    static const VectorOfCharStringMatrix animatedMapNames;
+    static const VectorOfCharStringMatrix meshNames;
+    static const pma::Vector<pma::Matrix<std::uint16_t> > meshBlendShapeIndices;
+    static const pma::Matrix<std::uint16_t> jointHierarchy;
+    static const pma::Vector<pma::Matrix<Vector3> > neutralJointTranslations;
+    static const pma::Vector<pma::Matrix<Vector3> > neutralJointRotations;
+
+    // Behavior
+    static const std::uint16_t guiControlCount;
+    static const std::uint16_t rawControlCount;
+    static const std::uint16_t psdCount;
+    static const pma::Vector<pma::Matrix<std::uint16_t> > conditionalInputIndices;
+    static const pma::Vector<pma::Matrix<std::uint16_t> > conditionalOutputIndices;
+    static const pma::Vector<pma::Matrix<float> > conditionalFromValues;
+    static const pma::Vector<pma::Matrix<float> > conditionalToValues;
+    static const pma::Vector<pma::Matrix<float> > conditionalSlopeValues;
+    static const pma::Vector<pma::Matrix<float> > conditionalCutValues;
+    static const pma::Vector<std::uint16_t> psdRowIndices;
+    static const pma::Vector<std::uint16_t> psdColumnIndices;
+    static const pma::Vector<float> psdValues;
+    static const pma::Vector<std::uint16_t> jointRowCount;
+    static const std::uint16_t jointColumnCount;
+    static const pma::Vector<pma::Matrix<std::uint16_t> > jointVariableIndices;
+    static const pma::Vector<pma::Matrix<std::uint16_t> > jointGroupLODs;
+    static const pma::Matrix<pma::Matrix<std::uint16_t> > jointGroupInputIndices;
+    static const pma::Matrix<pma::Matrix<std::uint16_t> > jointGroupOutputIndices;
+    static const pma::Matrix<pma::Matrix<float> > jointGroupValues;
+    static const pma::Matrix<pma::Matrix<std::uint16_t> > jointGroupJointIndices;
+    static const pma::Matrix<std::uint16_t> blendShapeLODs;
+    static const pma::Vector<pma::Matrix<std::uint16_t> > blendShapeInputIndices;
+    static const pma::Vector<pma::Matrix<std::uint16_t> > blendShapeOutputIndices;
+    static const pma::Vector<std::uint16_t> animatedMapCount;
+    static const pma::Matrix<std::uint16_t> animatedMapLODs;
+
+    // Geometry
+    static const pma::Vector<std::uint32_t> meshCount;
+    static const pma::Vector<pma::Matrix<Vector3> > vertexPositions;
+    static const pma::Vector<pma::Matrix<TextureCoordinate> > vertexTextureCoordinates;
+    static const pma::Vector<pma::Matrix<Vector3> > vertexNormals;
+    static const pma::Vector<pma::Matrix<VertexLayout> > vertexLayouts;
+    static const pma::Matrix<pma::Matrix<std::uint32_t> > faces;
+    static const pma::Matrix<std::uint16_t> maxInfluencePerVertex;
+    static const pma::Matrix<pma::Matrix<float> > skinWeightsValues;
+    static const pma::Matrix<pma::Matrix<std::uint16_t> > skinWeightsJointIndices;
+    static const pma::Vector<pma::Matrix<std::uint16_t> > correctiveBlendShapeIndices;
+    static const pma::Matrix<pma::Matrix<Vector3> > correctiveBlendShapeDeltas;
+    static const pma::Matrix<pma::Matrix<std::uint32_t> > correctiveBlendShapeVertexIndices;
+
+    // Machine learned behavior
+    static const pma::Vector<pma::String<char> > mlControlNames;
+    static const VectorOfCharStringMatrix regionNames;
+    static const pma::Matrix<std::uint16_t> neuralNetworkIndicesPerLOD;
+    static const pma::Matrix<pma::Matrix<std::uint16_t> > neuralNetworkIndicesPerMeshRegion;
+    static const pma::Vector<pma::Matrix<std::uint16_t> > neuralNetworkInputIndices;
+    static const pma::Vector<pma::Matrix<std::uint16_t> > neuralNetworkOutputIndices;
+    static const pma::Matrix<std::uint16_t> neuralNetworkLayerCount;
+    static const pma::Vector<pma::Matrix<std::uint16_t> > neuralNetworkActivationFunction;
+    static const pma::Matrix<pma::Matrix<float> > neuralNetworkActivationFunctionParameters;
+    static const pma::Matrix<pma::Matrix<float> > neuralNetworkBiases;
+    static const pma::Matrix<pma::Matrix<float> > neuralNetworkWeights;
+
+    // RBF behavior
+    static const pma::Vector<pma::String<char> > poseNames;
+    static const pma::Vector<std::uint16_t> poseDistanceMethod;
+    static const pma::Vector<std::uint16_t> poseFunctionType;
+    static const pma::Vector<float> poseScale;
+
+    static const pma::Matrix<std::uint16_t> solverIndicesPerLOD;
+    static const pma::Vector<pma::String<char> > solverNames;
+    static const pma::Vector<float> solverRadius;
+    static const pma::Vector<float> solverWeightThreshold;
+    static const pma::Vector<std::uint16_t> solverType;
+    static const pma::Vector<std::uint16_t> solverAutomaticRadius;
+    static const pma::Vector<std::uint16_t> solverDistanceMethod;
+    static const pma::Vector<std::uint16_t> solverNormalizeMethod;
+    static const pma::Vector<std::uint16_t> solverFunctionType;
+    static const pma::Vector<std::uint16_t> solverTwistAxis;
+    static const pma::Matrix<std::uint16_t> solverRawControlIndices;
+    static const pma::Matrix<std::uint16_t> solverPoseIndices;
+    static const pma::Matrix<float> solverRawControlValues;
+
+    // RBF behavior ext
+    static const pma::Vector<pma::String<char> > poseControlNames;
+    static const pma::Matrix<std::uint16_t> poseInputControlIndices;
+    static const pma::Matrix<std::uint16_t> poseOutputControlIndices;
+    static const pma::Matrix<float> poseOutputControlWeights;
+
+    // Joint behavior metadata
+    static const pma::Matrix<TranslationRepresentation> jointTranslationRepresentation;
+    static const pma::Matrix<RotationRepresentation> jointRotationRepresentation;
+    static const pma::Matrix<ScaleRepresentation> jointScaleRepresentation;
+
+    // Twist swing setups
+    static const pma::Vector<pma::Matrix<std::uint16_t> > swingInputControlIndices;
+    static const pma::Vector<pma::Matrix<std::uint16_t> > swingOutputJointIndices;
+    static const pma::Vector<pma::Matrix<float> > swingBlendWeights;
+    static const pma::Matrix<TwistAxis> swingTwistAxes;
+    static const pma::Vector<pma::Matrix<std::uint16_t> > twistInputControlIndices;
+    static const pma::Vector<pma::Matrix<std::uint16_t> > twistOutputJointIndices;
+    static const pma::Vector<pma::Matrix<float> > twistBlendWeights;
+    static const pma::Matrix<TwistAxis> twistTwistAxes;
+
+    static std::size_t lodConstraintToIndex(std::uint16_t maxLOD, std::uint16_t minLOD);
+    static RawJoints getJoints(std::uint16_t currentMaxLOD, std::uint16_t currentMinLOD, pma::MemoryResource* memRes);
+    static RawBlendShapeChannels getBlendShapes(std::uint16_t currentMaxLOD,
+                                                std::uint16_t currentMinLOD,
+                                                pma::MemoryResource* memRes);
+    static RawConditionalTable getConditionals(std::uint16_t currentMaxLOD,
+                                               std::uint16_t currentMinLOD,
+                                               pma::MemoryResource* memRes);
+    static RawAnimatedMaps getAnimatedMaps(std::uint16_t currentMaxLOD, std::uint16_t currentMinLOD, pma::MemoryResource* memRes);
+
+};
+
+}  // namespace dna
diff --git a/Engine/Plugins/Animation/RigLogic/Source/RigLogicLibTest/Private/dnatests/TestStreamReadWriteIntegration.cpp b/Engine/Plugins/Animation/RigLogic/Source/RigLogicLibTest/Private/dnatests/TestStreamReadWriteIntegration.cpp
--- a/Engine/Plugins/Animation/RigLogic/Source/RigLogicLibTest/Private/dnatests/TestStreamReadWriteIntegration.cpp
+++ b/Engine/Plugins/Animation/RigLogic/Source/RigLogicLibTest/Private/dnatests/TestStreamReadWriteIntegration.cpp
@@ -10,6 +10,7 @@
 #include "dnatests/Fixturesv22.h"
 #include "dnatests/Fixturesv23.h"
 #include "dnatests/Fixturesv24.h"
+#include "dnatests/Fixturesv25.h"
 
 #include "dna/DataLayer.h"
 #include "dna/BinaryStreamReader.h"
@@ -427,6 +428,7 @@ static void verifyRBFBehavior(RBFBehaviorReader* reader) {
     ASSERT_EQ(solverCount, DecodedDNA::solverIndicesPerLOD[index].size());
 
     const auto poseCount = reader->getRBFPoseCount();
+    ASSERT_EQ(poseCount, DecodedDNA::poseScale.size());
     for (std::uint16_t pi = {}; pi < poseCount; ++pi) {
         ASSERT_EQ(reader->getRBFPoseName(pi), StringView{DecodedDNA::poseNames[pi]});
         ASSERT_EQ(reader->getRBFPoseScale(pi), DecodedDNA::poseScale[pi]);
@@ -458,16 +460,45 @@ static void verifyRBFBehavior(RBFBehaviorReader* reader) {
         auto solverPoseIndices = reader->getRBFSolverPoseIndices(si);
         const auto& expectedSolverPoseIndices = DecodedDNA::solverPoseIndices[esi];
         ASSERT_EQ(solverPoseIndices.size(), expectedSolverPoseIndices.size());
-        ASSERT_ELEMENTS_EQ(solverPoseIndices, expectedSolverPoseIndices, solverPoseIndices.size());
+        ASSERT_ELEMENTS_EQ(solverPoseIndices, expectedSolverPoseIndices, expectedSolverPoseIndices.size());
 
         auto solverRawControlValues = reader->getRBFSolverRawControlValues(si);
         const auto& expectedSolverRawControlValues = DecodedDNA::solverRawControlValues[esi];
         ASSERT_EQ(solverRawControlValues.size(), expectedSolverRawControlValues.size());
-        ASSERT_ELEMENTS_EQ(solverRawControlValues, expectedSolverRawControlValues, solverRawControlValues.size());
+        ASSERT_ELEMENTS_EQ(solverRawControlValues, expectedSolverRawControlValues, expectedSolverRawControlValues.size());
     }
 
 }
 
+template<class TAPICopyParameters>
+static void verifyRBFBehaviorExt(RBFBehaviorReader* reader) {
+    using DecodedDNA = typename TAPICopyParameters::DecodedData;
+
+    const auto poseControlCount = reader->getRBFPoseControlCount();
+    ASSERT_EQ(poseControlCount, DecodedDNA::poseControlNames.size());
+    for (std::uint16_t pci = {}; pci < poseControlCount; ++pci) {
+        ASSERT_EQ(reader->getRBFPoseControlName(pci), StringView{DecodedDNA::poseControlNames[pci]});
+    }
+
+    const auto poseCount = reader->getRBFPoseCount();
+    for (std::uint16_t pi = {}; pi < poseCount; ++pi) {
+        auto poseInputControlIndices = reader->getRBFPoseInputControlIndices(pi);
+        const auto& expectedPoseInputControlIndices = DecodedDNA::poseInputControlIndices[pi];
+        ASSERT_EQ(poseInputControlIndices.size(), expectedPoseInputControlIndices.size());
+        ASSERT_ELEMENTS_EQ(poseInputControlIndices, expectedPoseInputControlIndices, expectedPoseInputControlIndices.size());
+
+        auto poseOutputControlIndices = reader->getRBFPoseOutputControlIndices(pi);
+        const auto& expectedPoseOutputControlIndices = DecodedDNA::poseOutputControlIndices[pi];
+        ASSERT_EQ(poseOutputControlIndices.size(), expectedPoseOutputControlIndices.size());
+        ASSERT_ELEMENTS_EQ(poseOutputControlIndices, expectedPoseOutputControlIndices, expectedPoseOutputControlIndices.size());
+
+        auto poseOutputControlWeights = reader->getRBFPoseOutputControlWeights(pi);
+        const auto& expectedPoseOutputControlWeights = DecodedDNA::poseOutputControlWeights[pi];
+        ASSERT_EQ(poseOutputControlWeights.size(), expectedPoseOutputControlWeights.size());
+        ASSERT_ELEMENTS_EQ(poseOutputControlWeights, expectedPoseOutputControlWeights, expectedPoseOutputControlWeights.size());
+    }
+}
+
 template<class TAPICopyParameters>
 static void verifyJointBehaviorMetadata(JointBehaviorMetadataReader* reader) {
     using DecodedDNA = typename TAPICopyParameters::DecodedData;
@@ -578,6 +609,24 @@ struct ReaderDataVerifier<APICopyParameters<Reader, Writer, RawV24, DecodedV24,
 
 };
 
+template<class Reader, class Writer, std::uint16_t MaxLOD, std::uint16_t MinLOD, std::uint16_t CurrentLOD>
+struct ReaderDataVerifier<APICopyParameters<Reader, Writer, RawV25, DecodedV25, MaxLOD, MinLOD, CurrentLOD> > {
+
+    static void assertHasAllData(Reader* reader) {
+        using TAPICopyParameters = APICopyParameters<Reader, Writer, RawV25, DecodedV25, MaxLOD, MinLOD, CurrentLOD>;
+        verifyDescriptor<TAPICopyParameters>(reader);
+        verifyDefinition<TAPICopyParameters>(reader);
+        verifyBehavior<TAPICopyParameters>(reader);
+        verifyGeometry<TAPICopyParameters>(reader);
+        verifyMachineLearnedBehavior<TAPICopyParameters>(reader);
+        verifyRBFBehavior<TAPICopyParameters>(reader);
+        verifyRBFBehaviorExt<TAPICopyParameters>(reader);
+        verifyJointBehaviorMetadata<TAPICopyParameters>(reader);
+        verifyTwistSwingBehavior<TAPICopyParameters>(reader);
+    }
+
+};
+
 using TAPICopyTestParameters = ::testing::Types<
     APICopyParameters<dna::BinaryStreamReader, dna::BinaryStreamWriter, RawV21, DecodedV21, 0u, 1u, 0u>
     , APICopyParameters<dna::BinaryStreamReader, dna::BinaryStreamWriter, RawV21, DecodedV21, 0u, 1u, 1u>
@@ -595,11 +644,16 @@ using TAPICopyTestParameters = ::testing::Types<
     , APICopyParameters<dna::BinaryStreamReader, dna::BinaryStreamWriter, RawV24, DecodedV24, 0u, 1u, 1u>
     , APICopyParameters<dna::BinaryStreamReader, dna::BinaryStreamWriter, RawV24, DecodedV24, 0u, 0u, 0u>
     , APICopyParameters<dna::BinaryStreamReader, dna::BinaryStreamWriter, RawV24, DecodedV24, 1u, 1u, 0u>
+    , APICopyParameters<dna::BinaryStreamReader, dna::BinaryStreamWriter, RawV25, DecodedV25, 0u, 1u, 0u>
+    , APICopyParameters<dna::BinaryStreamReader, dna::BinaryStreamWriter, RawV25, DecodedV25, 0u, 1u, 1u>
+    , APICopyParameters<dna::BinaryStreamReader, dna::BinaryStreamWriter, RawV25, DecodedV25, 0u, 0u, 0u>
+    , APICopyParameters<dna::BinaryStreamReader, dna::BinaryStreamWriter, RawV25, DecodedV25, 1u, 1u, 0u>
     #ifdef DNA_BUILD_WITH_JSON_SUPPORT
         , APICopyParameters<dna::JSONStreamReader, dna::JSONStreamWriter, RawV21, DecodedV21, 0u, 1u, 0u>
         , APICopyParameters<dna::JSONStreamReader, dna::JSONStreamWriter, RawV22, DecodedV22, 0u, 1u, 0u>
         , APICopyParameters<dna::JSONStreamReader, dna::JSONStreamWriter, RawV23, DecodedV23, 0u, 1u, 0u>
         , APICopyParameters<dna::JSONStreamReader, dna::JSONStreamWriter, RawV24, DecodedV24, 0u, 1u, 0u>
+        , APICopyParameters<dna::JSONStreamReader, dna::JSONStreamWriter, RawV25, DecodedV25, 0u, 1u, 0u>
     #endif  // DNA_BUILD_WITH_JSON_SUPPORT
     >;
 TYPED_TEST_SUITE(StreamReadWriteAPICopyIntegrationTest, TAPICopyTestParameters, );
@@ -679,6 +733,8 @@ using TRawCopyTestParameters = ::testing::Types<
     RawCopyParameters<RawV23, RawV23, UnknownLayerPolicy::Ignore, 2, 3>,
     RawCopyParameters<RawV24, RawV24, UnknownLayerPolicy::Preserve, 2, 4>,
     RawCopyParameters<RawV24, RawV24, UnknownLayerPolicy::Ignore, 2, 4>,
+    RawCopyParameters<RawV25, RawV25, UnknownLayerPolicy::Preserve, 2, 5>,
+    RawCopyParameters<RawV25, RawV25, UnknownLayerPolicy::Ignore, 2, 5>,
     // File format conversion tests
     RawCopyParameters<RawV21, RawV22WithUnknownDataIgnoredAndDNARewritten, UnknownLayerPolicy::Preserve, 2, 2>,
     RawCopyParameters<RawV21, RawV22WithUnknownDataIgnoredAndDNARewritten, UnknownLayerPolicy::Ignore, 2, 2>,
@@ -691,7 +747,9 @@ using TRawCopyTestParameters = ::testing::Types<
     RawCopyParameters<RawV23, RawV22DowngradedFromV23, UnknownLayerPolicy::Preserve, 2, 2>,
     RawCopyParameters<RawV23, RawV22WithUnknownDataIgnoredAndDNARewritten, UnknownLayerPolicy::Ignore, 2, 2>,
     RawCopyParameters<RawV24, RawV23DowngradedFromV24, UnknownLayerPolicy::Preserve, 2, 3>,
-    RawCopyParameters<RawV24, RawV23, UnknownLayerPolicy::Ignore, 2, 3>
+    RawCopyParameters<RawV24, RawV23, UnknownLayerPolicy::Ignore, 2, 3>,
+    RawCopyParameters<RawV25, RawV24DowngradedFromV25, UnknownLayerPolicy::Preserve, 2, 4>,
+    RawCopyParameters<RawV25, RawV24, UnknownLayerPolicy::Ignore, 2, 4>
     >;
 TYPED_TEST_SUITE(StreamReadWriteRawCopyIntegrationTest, TRawCopyTestParameters, );
 
@@ -771,6 +829,7 @@ using TReadWriteMultipleParameters = ::testing::Types<
     ReadWriteMultipleParameters<RawV22>,
     ReadWriteMultipleParameters<RawV23>,
     ReadWriteMultipleParameters<RawV24>,
+    ReadWriteMultipleParameters<RawV25>,
     ReadWriteMultipleParameters<RawV22Empty>,
     ReadWriteMultipleParameters<RawV22WithUnknownDataIgnoredAndDNARewritten>,
     ReadWriteMultipleParameters<RawV2xNewerWithUnknownDataIgnoredAndDNARewritten>,
@@ -866,4 +925,30 @@ TYPED_TEST(StreamReadWriteMultipleIntegrationTest, ReadWriteTwoDNAsToSameStream)
     ASSERT_EQ(copiedCloneBytes, copiedCloneRewrittenBytes);
 }
 
+TEST(StreamReadWriteMultipleIntegrationTest, DNAv25LayerIsBackFilledFromv24) {
+    const auto bytes = RawV24::getBytes();
+    auto source = pma::makeScoped<trio::MemoryStream>();
+    source->write(bytes.data(), bytes.size());
+    source->seek(0);
+    auto reader = pma::makeScoped<BinaryStreamReader>(source.get(),
+                                                      DataLayer::All,
+                                                      UnknownLayerPolicy::Preserve,
+                                                      static_cast<std::uint16_t>(0));
+    reader->read();
+
+    ASSERT_TRUE(dna::Status::isOk());
+    ASSERT_EQ(reader->getRBFPoseControlCount(), reader->getRBFPoseCount());
+    for (std::uint16_t pi = {}; pi < reader->getRBFPoseCount(); ++pi) {
+        const auto inputControlIndices = reader->getRBFPoseInputControlIndices(pi);
+        const auto outputControlIndices = reader->getRBFPoseOutputControlIndices(pi);
+        const auto outputControlWeights = reader->getRBFPoseOutputControlWeights(pi);
+        ASSERT_EQ(inputControlIndices.size(), 0ul);
+        ASSERT_EQ(outputControlIndices.size(), 1ul);
+        ASSERT_EQ(outputControlWeights.size(), 1ul);
+        const auto offset = reader->getRawControlCount() + reader->getPSDCount() + reader->getMLControlCount();
+        ASSERT_EQ(outputControlIndices[0], offset + pi);
+        ASSERT_EQ(outputControlWeights[0], 1.0f);
+    }
+}
+
 }  // namespace dna
diff --git a/Engine/Plugins/Animation/RigLogic/Source/RigLogicLibTest/Private/rltests/dna/FakeReader.h b/Engine/Plugins/Animation/RigLogic/Source/RigLogicLibTest/Private/rltests/dna/FakeReader.h
--- a/Engine/Plugins/Animation/RigLogic/Source/RigLogicLibTest/Private/rltests/dna/FakeReader.h
+++ b/Engine/Plugins/Animation/RigLogic/Source/RigLogicLibTest/Private/rltests/dna/FakeReader.h
@@ -570,6 +570,26 @@ class FakeReader : public Reader {
             return {};
         }
 
+        std::uint16_t getRBFPoseControlCount() const override {
+            return {};
+        }
+
+        StringView getRBFPoseControlName(std::uint16_t poseControlIndex) const override {
+            return {};
+        }
+
+        ConstArrayView<std::uint16_t> getRBFPoseInputControlIndices(std::uint16_t poseIndex) const override {
+            return {};
+        }
+
+        ConstArrayView<std::uint16_t> getRBFPoseOutputControlIndices(std::uint16_t poseIndex) const override {
+            return {};
+        }
+
+        ConstArrayView<float> getRBFPoseOutputControlWeights(std::uint16_t poseIndex) const override {
+            return {};
+        }
+
         std::uint16_t getRBFSolverCount() const override {
             return {};
         }
diff --git a/Engine/Plugins/Animation/RigLogic/Source/RigLogicLibTest/Private/rltests/rbf/cpu/RBFFixtures.cpp b/Engine/Plugins/Animation/RigLogic/Source/RigLogicLibTest/Private/rltests/rbf/cpu/RBFFixtures.cpp
--- a/Engine/Plugins/Animation/RigLogic/Source/RigLogicLibTest/Private/rltests/rbf/cpu/RBFFixtures.cpp
+++ b/Engine/Plugins/Animation/RigLogic/Source/RigLogicLibTest/Private/rltests/rbf/cpu/RBFFixtures.cpp
@@ -17,7 +17,6 @@ namespace unoptimized {
 const std::uint16_t lodCount = 3u;
 
 const std::uint16_t rawControlCount = 8u;
-const std::uint16_t rbfControlCount = 15u;
 
 const pma::Matrix<std::uint16_t> solverIndicesPerLOD = {
     {0u, 1u},
@@ -77,6 +76,17 @@ const pma::Matrix<std::uint16_t> solverRawControlIndices = {
     {4u, 5u, 6u, 7u},
 };
 
+const std::uint16_t poseControlCount = 16u;
+
+const Matrix<std::uint16_t> poseInputControlIndices =
+{{}, {}, {}, {}, {}, {}, {}, {}, {}, {}, {}, {}, {}, {}, {1u}};
+
+const Matrix<std::uint16_t> poseOutputControlIndices =
+{{8u}, {9u}, {10u}, {11u}, {12u}, {13u}, {14u}, {15u}, {16u}, {17u}, {18u}, {19u}, {20u}, {21u}, {22u, 23u}};
+
+const Matrix<float> poseOutputControlWeights =
+{{1.0f}, {1.0f}, {1.0f}, {1.0f}, {1.0f}, {1.0f}, {1.0f}, {1.0f}, {1.0f}, {1.0f}, {1.0f}, {1.0f}, {1.0f}, {1.0f}, {0.5f, 0.5f}};
+
 const pma::Matrix<float> solverRawControlValues = {
     {
         0.000000000000000f, 0.000000000000000f, 0.000000000000000f, 1.000000000000000f,
@@ -111,9 +121,6 @@ const LODSpec<std::uint16_t> lods = {
 
 const Matrix<std::uint16_t> solverRawControlInputIndices = {{0u, 1u, 2u, 3u}, {4u, 5u, 6u, 7u}};
 
-const Matrix<std::uint16_t> solverRawControlOutputIndices =
-{{8u, 9u, 10u, 11u, 12u, 13u, 14u, 15u, 16u, 17u, 18u, 19u}, {8u, 20u, 21u, 22u}};
-
 const std::uint16_t maximumInputCount = 4u;
 
 const std::uint16_t maxTargetCount = 12u;
@@ -237,7 +244,8 @@ const Matrix<float> valuesPerLOD = {
         0.0198993f,
         0.0f,
         0.453696f,
-        0.480578f,
+        0.09189290998017785f,
+        0.09189290998017785f,
     },
     {
         0.0f,
diff --git a/Engine/Plugins/Animation/RigLogic/Source/RigLogicLibTest/Private/rltests/rbf/cpu/RBFFixtures.h b/Engine/Plugins/Animation/RigLogic/Source/RigLogicLibTest/Private/rltests/rbf/cpu/RBFFixtures.h
--- a/Engine/Plugins/Animation/RigLogic/Source/RigLogicLibTest/Private/rltests/rbf/cpu/RBFFixtures.h
+++ b/Engine/Plugins/Animation/RigLogic/Source/RigLogicLibTest/Private/rltests/rbf/cpu/RBFFixtures.h
@@ -28,11 +28,15 @@ namespace unoptimized {
 
 
 extern const std::uint16_t rawControlCount;
-extern const std::uint16_t rbfControlCount;
 extern const std::uint16_t lodCount;
 extern const pma::Matrix<std::uint16_t> solverIndicesPerLOD;
+extern const pma::Matrix<std::uint16_t> solverRawControlIndices;
 
 extern const pma::Vector<float> poseScales;
+extern const std::uint16_t poseControlCount;
+extern const Matrix<std::uint16_t> poseInputControlIndices;
+extern const Matrix<std::uint16_t> poseOutputControlIndices;
+extern const Matrix<float> poseOutputControlWeights;
 extern const pma::Vector<dna::RBFSolverType> solverTypes;
 extern const pma::Vector<dna::RBFDistanceMethod> solverDistanceMethods;
 extern const pma::Vector<dna::RBFFunctionType> solverFunctionType;
@@ -42,7 +46,6 @@ extern const pma::Vector<dna::AutomaticRadius> solverAutomaticRadius;
 extern const pma::Vector<float> solverRadius;
 extern const pma::Vector<float> solverWeightThreshold;
 extern const pma::Matrix<std::uint16_t> solverPoseIndices;
-extern const pma::Matrix<std::uint16_t> solverRawControlIndices;
 extern const pma::Matrix<float> solverRawControlValues;
 
 }  // namespace unoptimized
@@ -52,7 +55,6 @@ namespace optimized {
 
 extern const LODSpec<std::uint16_t> lods;
 extern const Matrix<std::uint16_t> solverRawControlInputIndices;
-extern const Matrix<std::uint16_t> solverRawControlOutputIndices;
 extern const std::uint16_t maximumInputCount;
 extern const std::uint16_t maxTargetCount;
 
@@ -151,6 +153,22 @@ class RBFReader : public dna::FakeReader {
             return unoptimized::solverTwistAxis[solverIndex];
         }
 
+        std::uint16_t getRBFPoseControlCount() const override {
+            return static_cast<std::uint16_t>(unoptimized::poseControlCount);
+        }
+
+        ConstArrayView<std::uint16_t> getRBFPoseInputControlIndices(std::uint16_t poseIndex) const override {
+            return unoptimized::poseInputControlIndices[poseIndex];
+        }
+
+        ConstArrayView<std::uint16_t> getRBFPoseOutputControlIndices(std::uint16_t poseIndex) const override {
+            return unoptimized::poseOutputControlIndices[poseIndex];
+        }
+
+        ConstArrayView<float> getRBFPoseOutputControlWeights(std::uint16_t poseIndex) const override {
+            return unoptimized::poseOutputControlWeights[poseIndex];
+        }
+
 };
 
 }  // namespace rbf
diff --git a/Engine/Plugins/Animation/RigLogic/Source/RigLogicLibTest/Private/rltests/rbf/cpu/TestRBFEvaluator.cpp b/Engine/Plugins/Animation/RigLogic/Source/RigLogicLibTest/Private/rltests/rbf/cpu/TestRBFEvaluator.cpp
--- a/Engine/Plugins/Animation/RigLogic/Source/RigLogicLibTest/Private/rltests/rbf/cpu/TestRBFEvaluator.cpp
+++ b/Engine/Plugins/Animation/RigLogic/Source/RigLogicLibTest/Private/rltests/rbf/cpu/TestRBFEvaluator.cpp
@@ -93,11 +93,11 @@ TYPED_TEST(RBFBehaviorTest, SolverPerLOD) {
                                                                     rltests::rbf::unoptimized::rawControlCount,
                                                                     0,
                                                                     0,
-                                                                    rltests::rbf::unoptimized::rbfControlCount);
+                                                                    rltests::rbf::unoptimized::poseControlCount);
     auto inputInstance = inputInstanceFactory(&this->memRes);
     auto inputBuffer = inputInstance->getInputBuffer();
     auto outputBuffer =
-        inputBuffer.subview(rltests::rbf::unoptimized::rawControlCount, rltests::rbf::unoptimized::rbfControlCount);
+        inputBuffer.subview(rltests::rbf::unoptimized::rawControlCount, rltests::rbf::unoptimized::poseControlCount);
     const auto& inputValues = rltests::rbf::input::values;
     for (std::size_t i = {}; i < inputValues.size(); ++i) {
         inputBuffer[i] = inputValues[i];
diff --git a/Engine/Plugins/Animation/RigLogic/Source/RigLogicLibTest/Private/rltests/rbf/cpu/TestRBFStorageBuilder.cpp b/Engine/Plugins/Animation/RigLogic/Source/RigLogicLibTest/Private/rltests/rbf/cpu/TestRBFStorageBuilder.cpp
--- a/Engine/Plugins/Animation/RigLogic/Source/RigLogicLibTest/Private/rltests/rbf/cpu/TestRBFStorageBuilder.cpp
+++ b/Engine/Plugins/Animation/RigLogic/Source/RigLogicLibTest/Private/rltests/rbf/cpu/TestRBFStorageBuilder.cpp
@@ -40,8 +40,12 @@ struct Evaluator<T, TF256, TF128>::Accessor {
             ASSERT_EQ(solver->getNormalizeMethod(), rltests::rbf::unoptimized::solverNormalizeMethods[si]);
             ASSERT_EQ(solver->getTwistAxis(), rltests::rbf::unoptimized::solverTwistAxis[si]);
 
-            ASSERT_EQ(result.solverRawControlInputIndices[si], rltests::rbf::optimized::solverRawControlInputIndices[si]);
-            ASSERT_EQ(result.solverRawControlOutputIndices[si], rltests::rbf::optimized::solverRawControlOutputIndices[si]);
+            ASSERT_EQ(result.solverPoseIndices[si], rltests::rbf::unoptimized::solverPoseIndices[si]);
+        }
+        for (std::uint16_t pi = {}; pi < rltests::rbf::unoptimized::poseScales.size(); ++pi) {
+            ASSERT_EQ(result.poseInputControlIndices[pi], rltests::rbf::unoptimized::poseInputControlIndices[pi]);
+            ASSERT_EQ(result.poseOutputControlIndices[pi], rltests::rbf::unoptimized::poseOutputControlIndices[pi]);
+            ASSERT_EQ(result.poseOutputControlWeights[pi], rltests::rbf::unoptimized::poseOutputControlWeights[pi]);
         }
     }
 
diff --git a/Engine/Plugins/Animation/RigLogic/Source/RigLogicModule/Private/RigLogicDNAReader.cpp b/Engine/Plugins/Animation/RigLogic/Source/RigLogicModule/Private/RigLogicDNAReader.cpp
--- a/Engine/Plugins/Animation/RigLogic/Source/RigLogicModule/Private/RigLogicDNAReader.cpp
+++ b/Engine/Plugins/Animation/RigLogic/Source/RigLogicModule/Private/RigLogicDNAReader.cpp
@@ -574,6 +574,26 @@ float RigLogicDNAReader::getRBFPoseScale(std::uint16_t poseIndex) const
 {
 	return Reader->getRBFPoseScale(poseIndex);
 }
+std::uint16_t RigLogicDNAReader::getRBFPoseControlCount() const
+{
+	return Reader->getRBFPoseControlCount();
+}
+dna::StringView RigLogicDNAReader::getRBFPoseControlName(std::uint16_t poseControlIndex) const
+{
+	return Reader->getRBFPoseControlName(poseControlIndex);
+}
+dna::ConstArrayView<std::uint16_t> RigLogicDNAReader::getRBFPoseInputControlIndices(std::uint16_t poseIndex) const
+{
+	return Reader->getRBFPoseInputControlIndices(poseIndex);
+}
+dna::ConstArrayView<std::uint16_t> RigLogicDNAReader::getRBFPoseOutputControlIndices(std::uint16_t poseIndex) const
+{
+	return Reader->getRBFPoseOutputControlIndices(poseIndex);
+}
+dna::ConstArrayView<float> RigLogicDNAReader::getRBFPoseOutputControlWeights(std::uint16_t poseIndex) const
+{
+	return Reader->getRBFPoseOutputControlWeights(poseIndex);
+}
 std::uint16_t RigLogicDNAReader::getRBFSolverCount() const
 {
 	return Reader->getRBFSolverCount();
diff --git a/Engine/Plugins/Animation/RigLogic/Source/RigLogicModule/Private/RigLogicDNAReader.h b/Engine/Plugins/Animation/RigLogic/Source/RigLogicModule/Private/RigLogicDNAReader.h
--- a/Engine/Plugins/Animation/RigLogic/Source/RigLogicModule/Private/RigLogicDNAReader.h
+++ b/Engine/Plugins/Animation/RigLogic/Source/RigLogicModule/Private/RigLogicDNAReader.h
@@ -155,6 +155,11 @@ public:
 	dna::ConstArrayView<std::uint16_t> getRBFPoseAnimatedMapOutputIndices(std::uint16_t poseIndex) const override;
 	dna::ConstArrayView<float> getRBFPoseJointOutputValues(std::uint16_t poseIndex) const override;
 	float getRBFPoseScale(std::uint16_t poseIndex) const override;
+	std::uint16_t getRBFPoseControlCount() const override;
+	dna::StringView getRBFPoseControlName(std::uint16_t poseControlIndex) const override;
+	dna::ConstArrayView<std::uint16_t> getRBFPoseInputControlIndices(std::uint16_t poseIndex) const override;
+	dna::ConstArrayView<std::uint16_t> getRBFPoseOutputControlIndices(std::uint16_t poseIndex) const override;
+	dna::ConstArrayView<float> getRBFPoseOutputControlWeights(std::uint16_t poseIndex) const override;
 	std::uint16_t getRBFSolverCount() const override;
 	std::uint16_t getRBFSolverIndexListCount() const override;
 	dna::ConstArrayView<std::uint16_t> getRBFSolverIndicesForLOD(std::uint16_t lod) const override;
diff --git a/Engine/Plugins/Animation/RigLogic/Source/RigLogicModule/Private/SkelMeshDNAReader.cpp b/Engine/Plugins/Animation/RigLogic/Source/RigLogicModule/Private/SkelMeshDNAReader.cpp
--- a/Engine/Plugins/Animation/RigLogic/Source/RigLogicModule/Private/SkelMeshDNAReader.cpp
+++ b/Engine/Plugins/Animation/RigLogic/Source/RigLogicModule/Private/SkelMeshDNAReader.cpp
@@ -748,6 +748,31 @@ float FSkelMeshDNAReader::GetRBFPoseScale(uint16 PoseIndex) const
 	return BehaviorReader->GetRBFPoseScale(PoseIndex);
 }
 
+uint16 FSkelMeshDNAReader::GetRBFPoseControlCount() const
+{
+	return BehaviorReader->GetRBFPoseControlCount();
+}
+
+FString FSkelMeshDNAReader::GetRBFPoseControlName(uint16 PoseControlIndex) const
+{
+	return BehaviorReader->GetRBFPoseControlName(PoseControlIndex);
+}
+
+TArrayView<const uint16> FSkelMeshDNAReader::GetRBFPoseInputControlIndices(uint16 PoseIndex) const
+{
+	return BehaviorReader->GetRBFPoseInputControlIndices(PoseIndex);
+}
+
+TArrayView<const uint16> FSkelMeshDNAReader::GetRBFPoseOutputControlIndices(uint16 PoseIndex) const
+{
+	return BehaviorReader->GetRBFPoseOutputControlIndices(PoseIndex);
+}
+
+TArrayView<const float> FSkelMeshDNAReader::GetRBFPoseOutputControlWeights(uint16 PoseIndex) const
+{
+	return BehaviorReader->GetRBFPoseOutputControlWeights(PoseIndex);
+}
+
 uint16 FSkelMeshDNAReader::GetRBFSolverCount() const
 {
 	return BehaviorReader->GetRBFSolverCount();
diff --git a/Engine/Plugins/Animation/RigLogic/Source/RigLogicModule/Public/DNAReader.h b/Engine/Plugins/Animation/RigLogic/Source/RigLogicModule/Public/DNAReader.h
--- a/Engine/Plugins/Animation/RigLogic/Source/RigLogicModule/Public/DNAReader.h
+++ b/Engine/Plugins/Animation/RigLogic/Source/RigLogicModule/Public/DNAReader.h
@@ -217,6 +217,11 @@ public:
 	virtual TArrayView<const uint16> GetRBFPoseAnimatedMapOutputIndices(uint16 PoseIndex) const = 0;
 	virtual TArrayView<const float> GetRBFPoseJointOutputValues(uint16 PoseIndex) const = 0;
 	virtual float GetRBFPoseScale(uint16 PoseIndex) const = 0;
+	virtual uint16 GetRBFPoseControlCount() const = 0;
+	virtual FString GetRBFPoseControlName(uint16 PoseControlIndex) const = 0;
+	virtual TArrayView<const uint16> GetRBFPoseInputControlIndices(uint16 PoseIndex) const = 0;
+	virtual TArrayView<const uint16> GetRBFPoseOutputControlIndices(uint16 PoseIndex) const = 0;
+	virtual TArrayView<const float> GetRBFPoseOutputControlWeights(uint16 PoseIndex) const = 0;
 	virtual uint16 GetRBFSolverCount() const = 0;
 	virtual uint16 GetRBFSolverIndexListCount() const = 0;
 	virtual TArrayView<const uint16> GetRBFSolverIndicesForLOD(uint16 LOD) const = 0;
diff --git a/Engine/Plugins/Animation/RigLogic/Source/RigLogicModule/Public/DNAReaderAdapter.h b/Engine/Plugins/Animation/RigLogic/Source/RigLogicModule/Public/DNAReaderAdapter.h
--- a/Engine/Plugins/Animation/RigLogic/Source/RigLogicModule/Public/DNAReaderAdapter.h
+++ b/Engine/Plugins/Animation/RigLogic/Source/RigLogicModule/Public/DNAReaderAdapter.h
@@ -148,6 +148,11 @@ public:
 	TArrayView<const uint16> GetRBFPoseAnimatedMapOutputIndices(uint16 PoseIndex) const override;
 	TArrayView<const float> GetRBFPoseJointOutputValues(uint16 PoseIndex) const override;
 	float GetRBFPoseScale(uint16 PoseIndex) const override;
+	uint16 GetRBFPoseControlCount() const override;
+	FString GetRBFPoseControlName(uint16 PoseControlIndex) const override;
+	TArrayView<const uint16> GetRBFPoseInputControlIndices(uint16 PoseIndex) const override;
+	TArrayView<const uint16> GetRBFPoseOutputControlIndices(uint16 PoseIndex) const override;
+	TArrayView<const float> GetRBFPoseOutputControlWeights(uint16 PoseIndex) const override;
 	uint16 GetRBFSolverCount() const override;
 	uint16 GetRBFSolverIndexListCount() const override;
 	TArrayView<const uint16> GetRBFSolverIndicesForLOD(uint16 LOD) const override;
@@ -1021,6 +1026,39 @@ float FDNAReader<TWrappedReader>::GetRBFPoseScale(uint16 PoseIndex) const
 	return ReaderPtr->getRBFPoseScale(PoseIndex);
 }
 
+template <class TWrappedReader>
+uint16 FDNAReader<TWrappedReader>::GetRBFPoseControlCount() const
+{
+	return ReaderPtr->getRBFPoseControlCount();
+}
+
+template <class TWrappedReader>
+FString FDNAReader<TWrappedReader>::GetRBFPoseControlName(uint16 PoseControlIndex) const
+{
+	return FString(ANSI_TO_TCHAR(ReaderPtr->getRBFPoseControlName(PoseControlIndex).data()));
+}
+
+template <class TWrappedReader>
+TArrayView<const uint16> FDNAReader<TWrappedReader>::GetRBFPoseInputControlIndices(uint16 PoseIndex) const
+{
+	const auto Indices = ReaderPtr->getRBFPoseInputControlIndices(PoseIndex);
+	return TArrayView<const uint16>(Indices.data(), static_cast<int32>(Indices.size()));
+}
+
+template <class TWrappedReader>
+TArrayView<const uint16> FDNAReader<TWrappedReader>::GetRBFPoseOutputControlIndices(uint16 PoseIndex) const
+{
+	const auto Indices = ReaderPtr->getRBFPoseOutputControlIndices(PoseIndex);
+	return TArrayView<const uint16>(Indices.data(), static_cast<int32>(Indices.size()));
+}
+
+template <class TWrappedReader>
+TArrayView<const float> FDNAReader<TWrappedReader>::GetRBFPoseOutputControlWeights(uint16 PoseIndex) const
+{
+	const auto Values = ReaderPtr->getRBFPoseOutputControlWeights(PoseIndex);
+	return TArrayView<const float>(Values.data(), static_cast<int32>(Values.size()));
+}
+
 template <class TWrappedReader>
 uint16 FDNAReader<TWrappedReader>::GetRBFSolverCount() const
 {
diff --git a/Engine/Plugins/Animation/RigLogic/Source/RigLogicModule/Public/RigUnit_RigLogic_Test.h b/Engine/Plugins/Animation/RigLogic/Source/RigLogicModule/Public/RigUnit_RigLogic_Test.h
--- a/Engine/Plugins/Animation/RigLogic/Source/RigLogicModule/Public/RigUnit_RigLogic_Test.h
+++ b/Engine/Plugins/Animation/RigLogic/Source/RigLogicModule/Public/RigUnit_RigLogic_Test.h
@@ -672,6 +672,31 @@ public:
 		return {};
 	}
 
+	uint16 GetRBFPoseControlCount() const override
+	{
+		return {};
+	}
+
+	FString GetRBFPoseControlName(uint16 PoseControlIndex) const override
+	{
+		return {};
+	}
+
+	TArrayView<const uint16> GetRBFPoseInputControlIndices(uint16 PoseIndex) const override
+	{
+		return {};
+	}
+
+	TArrayView<const uint16> GetRBFPoseOutputControlIndices(uint16 PoseIndex) const override
+	{
+		return {};
+	}
+
+	TArrayView<const float> GetRBFPoseOutputControlWeights(uint16 PoseIndex) const override
+	{
+		return {};
+	}
+
 	uint16 GetRBFSolverCount() const override
 	{
 		return {};
diff --git a/Engine/Plugins/Animation/RigLogic/Source/RigLogicModule/Public/SkelMeshDNAReader.h b/Engine/Plugins/Animation/RigLogic/Source/RigLogicModule/Public/SkelMeshDNAReader.h
--- a/Engine/Plugins/Animation/RigLogic/Source/RigLogicModule/Public/SkelMeshDNAReader.h
+++ b/Engine/Plugins/Animation/RigLogic/Source/RigLogicModule/Public/SkelMeshDNAReader.h
@@ -149,6 +149,11 @@ public:
 	TArrayView<const uint16> GetRBFPoseAnimatedMapOutputIndices(uint16 PoseIndex) const override;
 	TArrayView<const float> GetRBFPoseJointOutputValues(uint16 PoseIndex) const override;
 	float GetRBFPoseScale(uint16 PoseIndex) const override;
+	uint16 GetRBFPoseControlCount() const override;
+	FString GetRBFPoseControlName(uint16 PoseControlIndex) const override;
+	TArrayView<const uint16> GetRBFPoseInputControlIndices(uint16 PoseIndex) const override;
+	TArrayView<const uint16> GetRBFPoseOutputControlIndices(uint16 PoseIndex) const override;
+	TArrayView<const float> GetRBFPoseOutputControlWeights(uint16 PoseIndex) const override;
 	uint16 GetRBFSolverCount() const override;
 	uint16 GetRBFSolverIndexListCount() const override;
 	TArrayView<const uint16> GetRBFSolverIndicesForLOD(uint16 LOD) const override;
