Fix RigLogic integration tests with control rig

diff --git a/Engine/Plugins/Animation/RigLogic/Source/RigLogicModule/Private/RigUnit_RigLogic_Test.cpp b/Engine/Plugins/Animation/RigLogic/Source/RigLogicModule/Private/RigUnit_RigLogic_Test.cpp
--- a/Engine/Plugins/Animation/RigLogic/Source/RigLogicModule/Private/RigUnit_RigLogic_Test.cpp
+++ b/Engine/Plugins/Animation/RigLogic/Source/RigLogicModule/Private/RigUnit_RigLogic_Test.cpp
@@ -168,21 +168,22 @@ void FRigUnit_RigLogic::TestAccessor::Exec_MapMaskMultipliers(URigHierarchy* Tes
 void FRigUnit_RigLogic::TestAccessor::AddToTransformArray(float* InArray, FTransform& Transform)
 {
 	uint32 FirstAttributeIndex = 0;
-	
-	FVector Rotation = Transform.GetRotation().Euler();
-	InArray[FirstAttributeIndex + 0] = Rotation.X;
-	InArray[FirstAttributeIndex + 1] = Rotation.Y;
-	InArray[FirstAttributeIndex + 2] = Rotation.Z;
 
 	FVector Translation = Transform.GetTranslation();
-	InArray[FirstAttributeIndex + 3] = Translation.X;
-	InArray[FirstAttributeIndex + 4] = Translation.Y;
-	InArray[FirstAttributeIndex + 5] = Translation.Z;
+	InArray[FirstAttributeIndex + 0] = Translation.X;
+	InArray[FirstAttributeIndex + 1] = Translation.Y;
+	InArray[FirstAttributeIndex + 2] = Translation.Z;
+
+	FQuat Rotation = Transform.GetRotation();
+	InArray[FirstAttributeIndex + 3] = Rotation.X;
+	InArray[FirstAttributeIndex + 4] = Rotation.Y;
+	InArray[FirstAttributeIndex + 5] = Rotation.Z;
+	InArray[FirstAttributeIndex + 6] = Rotation.W;
 
 	FVector Scale = Transform.GetScale3D();
-	InArray[FirstAttributeIndex + 6] = Scale.X;
-	InArray[FirstAttributeIndex + 7] = Scale.Y;
-	InArray[FirstAttributeIndex + 8] = Scale.Z;
+	InArray[FirstAttributeIndex + 7] = Scale.X;
+	InArray[FirstAttributeIndex + 8] = Scale.Y;
+	InArray[FirstAttributeIndex + 9] = Scale.Z;
 }
 
 
@@ -218,19 +219,20 @@ TArrayView<const uint16> FRigUnit_RigLogic::TestAccessor::CreateTwoJointVariable
 	InVariableAttributeIndices[6] = 6;
 	InVariableAttributeIndices[7] = 7;
 	InVariableAttributeIndices[8] = 8;
+	InVariableAttributeIndices[9] = 9;
 
 	if (LOD == 0) //LOD0 includes attributes for both bones
 	{
-		InVariableAttributeIndices[9 + 0] = 9 + 0;
-		InVariableAttributeIndices[9 + 1] = 9 + 1;
-		InVariableAttributeIndices[9 + 2] = 9 + 2;
-		InVariableAttributeIndices[9 + 3] = 9 + 3;
-		InVariableAttributeIndices[9 + 4] = 9 + 4;
-		InVariableAttributeIndices[9 + 5] = 9 + 5;
-		InVariableAttributeIndices[9 + 6] = 9 + 6;
-		InVariableAttributeIndices[9 + 7] = 9 + 7;
-		InVariableAttributeIndices[9 + 8] = 9 + 8;
-	
+		InVariableAttributeIndices[FRigUnit_RigLogic::TestAccessor::MAX_ATTRS_PER_JOINT + 0] = FRigUnit_RigLogic::TestAccessor::MAX_ATTRS_PER_JOINT + 0;
+		InVariableAttributeIndices[FRigUnit_RigLogic::TestAccessor::MAX_ATTRS_PER_JOINT + 1] = FRigUnit_RigLogic::TestAccessor::MAX_ATTRS_PER_JOINT + 1;
+		InVariableAttributeIndices[FRigUnit_RigLogic::TestAccessor::MAX_ATTRS_PER_JOINT + 2] = FRigUnit_RigLogic::TestAccessor::MAX_ATTRS_PER_JOINT + 2;
+		InVariableAttributeIndices[FRigUnit_RigLogic::TestAccessor::MAX_ATTRS_PER_JOINT + 3] = FRigUnit_RigLogic::TestAccessor::MAX_ATTRS_PER_JOINT + 3;
+		InVariableAttributeIndices[FRigUnit_RigLogic::TestAccessor::MAX_ATTRS_PER_JOINT + 4] = FRigUnit_RigLogic::TestAccessor::MAX_ATTRS_PER_JOINT + 4;
+		InVariableAttributeIndices[FRigUnit_RigLogic::TestAccessor::MAX_ATTRS_PER_JOINT + 5] = FRigUnit_RigLogic::TestAccessor::MAX_ATTRS_PER_JOINT + 5;
+		InVariableAttributeIndices[FRigUnit_RigLogic::TestAccessor::MAX_ATTRS_PER_JOINT + 6] = FRigUnit_RigLogic::TestAccessor::MAX_ATTRS_PER_JOINT + 6;
+		InVariableAttributeIndices[FRigUnit_RigLogic::TestAccessor::MAX_ATTRS_PER_JOINT + 7] = FRigUnit_RigLogic::TestAccessor::MAX_ATTRS_PER_JOINT + 7;
+		InVariableAttributeIndices[FRigUnit_RigLogic::TestAccessor::MAX_ATTRS_PER_JOINT + 8] = FRigUnit_RigLogic::TestAccessor::MAX_ATTRS_PER_JOINT + 8;
+		InVariableAttributeIndices[FRigUnit_RigLogic::TestAccessor::MAX_ATTRS_PER_JOINT + 9] = FRigUnit_RigLogic::TestAccessor::MAX_ATTRS_PER_JOINT + 9;
 		return TArrayView<const uint16>(InVariableAttributeIndices, 2 * FRigUnit_RigLogic::TestAccessor::MAX_ATTRS_PER_JOINT);
 	}
 
@@ -764,10 +766,15 @@ IMPLEMENT_RIGUNIT_AUTOMATION_TEST(FRigUnit_RigLogic)
 	DeltaTransformData[0] = 1.f;
 	DeltaTransformData[1] = 0.f;
 	DeltaTransformData[2] = 0.f;
+	// first bone rotation
+	DeltaTransformData[6] = 1.0f;
 	//second bone translation
-	DeltaTransformData[9] = 1.f;
-	DeltaTransformData[10] = 2.f;
-	DeltaTransformData[11] = 7.f;
+	DeltaTransformData[10] = 1.f;
+	DeltaTransformData[11] = 2.f;
+	DeltaTransformData[12] = 7.f;
+	// second bone rotation
+	DeltaTransformData[16] = 1.0f;
+
 	TArrayView<const float> DeltaTransforms = TArrayView<const float>(DeltaTransformData, TransformArraySize);
 	//create variable joint index arrays for two bones
 	SharedRigRuntimeContext->VariableJointIndicesPerLOD.Reset();
@@ -779,7 +786,7 @@ IMPLEMENT_RIGUNIT_AUTOMATION_TEST(FRigUnit_RigLogic)
 	//Assert
 	//Note that BoneB.GlobalTransform.Z should be zero since the scale.Z is zero. Also, translation Y becomes -Y 
 	AddErrorIfFalse(TestHierarchyTwoBones->GetGlobalTransform(0).GetTranslation().Equals(FVector(1.f, 0.f, 0.f)), TEXT("UpdateJoints LOD0 Bone 01 - unexpected transform"));
-	AddErrorIfFalse(TestHierarchyTwoBones->GetGlobalTransform(1).GetTranslation().Equals(FVector(2.f, -2.f, 0.f)), TEXT("UpdateJoints LOD0 Bone 02 - unexpected transform"));
+	AddErrorIfFalse(TestHierarchyTwoBones->GetGlobalTransform(1).GetTranslation().Equals(FVector(2.f, 2.f, 0.f)), TEXT("UpdateJoints LOD0 Bone 02 - unexpected transform"));
 
 	return true;
 }
