- When an AssetImportTask Option has the skeletal Mesh populated, refrain from searching for the skeletal Mesh on the Asset package. This change was introduced because the AssetImportTask that was created through scripting was not respected due to a change recently introduced. We are checking whether the provided Skeletal Mesh is valid, if that's the case we stick with it not searching the Uasset Path for the SKM object

#tests Tested using the python AssetImportTask API and manual imports of DNA files

diff --git a/Engine/Plugins/Animation/RigLogic/Source/RigLogicEditor/Private/DNAAssetImportFactory.cpp b/Engine/Plugins/Animation/RigLogic/Source/RigLogicEditor/Private/DNAAssetImportFactory.cpp
--- a/Engine/Plugins/Animation/RigLogic/Source/RigLogicEditor/Private/DNAAssetImportFactory.cpp
+++ b/Engine/Plugins/Animation/RigLogic/Source/RigLogicEditor/Private/DNAAssetImportFactory.cpp
@@ -138,17 +138,22 @@ UObject* UDNAAssetImportFactory::FactoryCreateFile(UClass* InClass, UObject* InP
 	ImportUI->bIsReimport = false;
 	ImportUI->ReimportMesh = nullptr;
 
-	USkeletalMesh* SkeletalMeshForDNAImport = nullptr;
-	ForEachObjectWithOuter(InParent, [&SkeletalMeshForDNAImport](UObject* Object)
-	{
-		if (USkeletalMesh* SkeletalMesh = Cast<USkeletalMesh>(Object))
-		{
-			SkeletalMeshForDNAImport = SkeletalMesh;
-			return false;
-		}
-		return true;
-	});
-	ImportUI->SkeletalMesh = SkeletalMeshForDNAImport;
+	// If the user provided a SkeletalMesh through the AssetImportTask Options 
+	// refrain from searching for the Skeletal Mesh in the asset package
+
+	if (!IsValid(ImportUI->SkeletalMesh)) {
+		USkeletalMesh* SkeletalMeshForDNAImport = nullptr;
+		ForEachObjectWithOuter(InParent, [&SkeletalMeshForDNAImport](UObject* Object)
+			{
+				if (USkeletalMesh* SkeletalMesh = Cast<USkeletalMesh>(Object))
+				{
+					SkeletalMeshForDNAImport = SkeletalMesh;
+					return false;
+				}
+				return true;
+			});
+		ImportUI->SkeletalMesh = SkeletalMeshForDNAImport;
+	}
 
 	// Show the import dialog only when not in a "yes to all" state or when automating import
 	bool bIsAutomated = IsAutomatedImport();
