Fix unused blend shape filtering in DNA library

diff --git a/Engine/Plugins/Animation/RigLogic/RigLogic.uplugin b/Engine/Plugins/Animation/RigLogic/RigLogic.uplugin
--- a/Engine/Plugins/Animation/RigLogic/RigLogic.uplugin
+++ b/Engine/Plugins/Animation/RigLogic/RigLogic.uplugin
@@ -1,9 +1,9 @@
 {
 	"FileVersion": 3,
 	"Version": 10.2,
-	"VersionName": "10.2.0",
-	"FriendlyName": "RigLogic Plugin v10.2.0",
-	"Description": "3Lateral RigLogic Plugin for Facial Animation v10.2.0",
+	"VersionName": "10.2.1",
+	"FriendlyName": "RigLogic Plugin v10.2.1",
+	"Description": "3Lateral RigLogic Plugin for Facial Animation v10.2.1",
 	"Category": "Animation",
 	"CreatedBy": "3Lateral",
 	"CreatedByURL": "https://www.3lateral.com/",
diff --git a/Engine/Plugins/Animation/RigLogic/Source/RigLogicLib/Private/dna/stream/FilteredBinaryInputArchive.cpp b/Engine/Plugins/Animation/RigLogic/Source/RigLogicLib/Private/dna/stream/FilteredBinaryInputArchive.cpp
--- a/Engine/Plugins/Animation/RigLogic/Source/RigLogicLib/Private/dna/stream/FilteredBinaryInputArchive.cpp
+++ b/Engine/Plugins/Animation/RigLogic/Source/RigLogicLib/Private/dna/stream/FilteredBinaryInputArchive.cpp
@@ -422,7 +422,7 @@ void FilteredBinaryInputArchive::process(RawMachineLearnedBehavior& dest) {
 void FilteredBinaryInputArchive::process(DNA& dest) {
     BaseArchive::process(dest);
     // Don't run control-based post-load filtering for delta DNA files
-    if (!loadedControls.empty()) {
+    if ((dest.layers.unknownPolicy == UnknownLayerPolicy::Ignore) && (!loadedControls.empty())) {
         removeUnreferencedBlendShapes(dest);
     }
 }
@@ -435,7 +435,7 @@ void FilteredBinaryInputArchive::removeUnreferencedBlendShapes(DNA& dest) {
     for (std::size_t iPlusOne = bsc.inputIndices.size(); iPlusOne > 0ul; --iPlusOne) {
         const auto i = iPlusOne - 1ul;
         const auto controlIndex = bsc.inputIndices[i];
-        if (!loadedControls[controlIndex]) {
+        if ((controlIndex > loadedControls.size()) || (!loadedControls[controlIndex])) {
             unreferencedChannels.push_back(bsc.outputIndices[i]);
             // Remove behavior data
             removeByIndex(bsc.inputIndices, i);
@@ -521,12 +521,14 @@ void FilteredBinaryInputArchive::process(RawRBFBehaviorExt& dest) {
         return;
     }
 
-    if (!contains(layerBitmask, DataLayerBitmask::RBFBehavior)) {
-        return;
+    if (contains(layerBitmask, DataLayerBitmask::RBFBehavior)) {
+        process(dest.poseControlNames);
+        process(dest.poses);
+        loadedControls.resize(loadedControls.size() + dest.poseControlNames.size(), true);
+    } else {
+        const auto poseControlCount = processSize();
+        loadedControls.resize(loadedControls.size() + poseControlCount, false);
     }
-
-    process(dest.poseControlNames);
-    process(dest.poses);
 }
 
 void FilteredBinaryInputArchive::process(RawJointBehaviorMetadata& dest) {
diff --git a/Engine/Plugins/Animation/RigLogic/Source/RigLogicLib/Public/dna/version/Version.h b/Engine/Plugins/Animation/RigLogic/Source/RigLogicLib/Public/dna/version/Version.h
--- a/Engine/Plugins/Animation/RigLogic/Source/RigLogicLib/Public/dna/version/Version.h
+++ b/Engine/Plugins/Animation/RigLogic/Source/RigLogicLib/Public/dna/version/Version.h
@@ -4,5 +4,5 @@
 
 #define DNA_MAJOR_VERSION 9
 #define DNA_MINOR_VERSION 4
-#define DNA_PATCH_VERSION 0
-#define DNA_VERSION_STRING "9.4.0"
+#define DNA_PATCH_VERSION 1
+#define DNA_VERSION_STRING "9.4.1"
diff --git a/Engine/Plugins/Animation/RigLogic/Source/RigLogicLib/Public/riglogic/version/Version.h b/Engine/Plugins/Animation/RigLogic/Source/RigLogicLib/Public/riglogic/version/Version.h
--- a/Engine/Plugins/Animation/RigLogic/Source/RigLogicLib/Public/riglogic/version/Version.h
+++ b/Engine/Plugins/Animation/RigLogic/Source/RigLogicLib/Public/riglogic/version/Version.h
@@ -4,5 +4,5 @@
 
 #define RL_MAJOR_VERSION 10
 #define RL_MINOR_VERSION 2
-#define RL_PATCH_VERSION 0
-#define RL_VERSION_STRING "10.2.0"
+#define RL_PATCH_VERSION 1
+#define RL_VERSION_STRING "10.2.1"
