Make sure AnimNode_RigLogic and RigUnit_RigLogic is clamping raw control values that are fed from anim curves directly (but leave joint rotation inputs still unclamped)

diff --git a/Engine/Plugins/Animation/RigLogic/Source/RigLogicModule/Private/AnimNode_RigLogic.cpp b/Engine/Plugins/Animation/RigLogic/Source/RigLogicModule/Private/AnimNode_RigLogic.cpp
--- a/Engine/Plugins/Animation/RigLogic/Source/RigLogicModule/Private/AnimNode_RigLogic.cpp
+++ b/Engine/Plugins/Animation/RigLogic/Source/RigLogicModule/Private/AnimNode_RigLogic.cpp
@@ -186,7 +186,7 @@ void FAnimNode_RigLogic::UpdateControlCurves(const FPoseContext& InputContext, c
 		{
 			if (InControlAttributeCurveElement.Index != INDEX_NONE)
  			{
-				RigInstance->SetRawControl(InControlAttributeCurveElement.Index, InCurveElement.Value);
+				RigInstance->SetRawControl(InControlAttributeCurveElement.Index, FMath::Clamp(InCurveElement.Value, 0.0, 1.0));
 			}
 		});
 
diff --git a/Engine/Plugins/Animation/RigLogic/Source/RigLogicModule/Private/RigUnit_RigLogic.cpp b/Engine/Plugins/Animation/RigLogic/Source/RigLogicModule/Private/RigUnit_RigLogic.cpp
--- a/Engine/Plugins/Animation/RigLogic/Source/RigLogicModule/Private/RigUnit_RigLogic.cpp
+++ b/Engine/Plugins/Animation/RigLogic/Source/RigLogicModule/Private/RigUnit_RigLogic.cpp
@@ -333,7 +333,7 @@ void FRigUnit_RigLogic_Data::CalculateRigLogic(const URigHierarchy* InHierarchy,
 	{
 		const uint32 CurveIndex = InputCurveIndices[ControlIndex];
 		const float Value = InHierarchy->GetCurveValue(CurveIndex);
-		RigInstance->SetRawControl(ControlIndex, Value);
+		RigInstance->SetRawControl(ControlIndex, FMath::Clamp(Value, 0.0, 1.0));
 	}
 
 	for (int32 MappingIndex = 0; MappingIndex < DriverJointsToControlAttributesMap.Num(); ++MappingIndex)
