Cache DNA joint to SkeletalMesh bone indices mapping for each LOD lazily as they are requested, avoiding their possibly frequent and costly recomputation

diff --git a/Engine/Plugins/Animation/RigLogic/Source/RigLogicModule/Private/AnimNode_RigLogic.cpp b/Engine/Plugins/Animation/RigLogic/Source/RigLogicModule/Private/AnimNode_RigLogic.cpp
--- a/Engine/Plugins/Animation/RigLogic/Source/RigLogicModule/Private/AnimNode_RigLogic.cpp
+++ b/Engine/Plugins/Animation/RigLogic/Source/RigLogicModule/Private/AnimNode_RigLogic.cpp
@@ -93,46 +93,65 @@ void FAnimNode_RigLogic::CacheBones_AnyThread(const FAnimationCacheBonesContext&
 		RigInstance = new FRigInstance(LocalRigRuntimeContext->RigLogic.Get());
 	}
 
-	LocalDNAIndexMapping = DNAAsset->GetDNAIndexMapping(Skeleton, SkeletalMesh);
+	RigInstance->SetLOD(Context.AnimInstanceProxy->GetLODLevel());
+
+	TSharedPtr<FDNAIndexMapping> SharedDNAIndexMapping = DNAAsset->GetDNAIndexMapping(Skeleton, SkeletalMesh);
+	if (LocalDNAIndexMapping != SharedDNAIndexMapping)
+	{
+		LocalDNAIndexMapping = SharedDNAIndexMapping;
+		LocalJointMappingsPerLOD.Empty();
+		LocalJointMappingsPerLOD.SetNum(LocalRigRuntimeContext->RigLogic->GetLODCount());
+	}
+
 	// CacheBones is called on LOD switches as well, in which case compact pose bone indices must be remapped
 	const FBoneContainer& RequiredBones = Context.AnimInstanceProxy->GetRequiredBones();
 	if (RequiredBones.IsValid())
 	{
-		// Set current LOD
-		RigInstance->SetLOD(Context.AnimInstanceProxy->GetLODLevel());
 		const uint16 CurrentLOD = RigInstance->GetLOD();
+		auto& CurrentLODJointMapping = LocalJointMappingsPerLOD[CurrentLOD];
 
-		// Populate mapping of DNA joint indices to CompactPoseBoneIndex entries (used in updating joints with results from RigLogic)
-		const TArray<uint16>& VariableJointIndices = LocalRigRuntimeContext->VariableJointIndicesPerLOD[CurrentLOD].Values;
-		JointsMapDNAIndicesToCompactPoseBoneIndices.Empty();
-		JointsMapDNAIndicesToCompactPoseBoneIndices.Reserve(VariableJointIndices.Num());
-		for (const uint16 JointIndex : VariableJointIndices)
+		// Lazily initialize and cache mappings for each LOD as they are requested
+		const int32 BoneCountForLOD = RequiredBones.GetCompactPoseNumBones();
+		if (CurrentLODJointMapping.BoneCount != BoneCountForLOD)
 		{
-			const FMeshPoseBoneIndex MeshPoseBoneIndex = LocalDNAIndexMapping->JointsMapDNAIndicesToMeshPoseBoneIndices[JointIndex];
-			const FCompactPoseBoneIndex CompactPoseBoneIndex = RequiredBones.MakeCompactPoseIndex(MeshPoseBoneIndex);
-			if (CompactPoseBoneIndex != INDEX_NONE) {
-				JointsMapDNAIndicesToCompactPoseBoneIndices.Add({JointIndex, CompactPoseBoneIndex});
-			}
-		}
-
-		// Populate driver joint to raw control attribute mapping (used to feed RigLogic with inputs from the joint hierarchy)
-		SparseDriverJointsToControlAttributesMap.Empty();
-		DenseDriverJointsToControlAttributesMap.Empty();
-		DenseDriverJointsToControlAttributesMap.Reserve(LocalDNAIndexMapping->DriverJointsToControlAttributesMap.Num());
-		// Sparse mapping will likely remain empty so no reservation happens
-		for (const auto& Mapping : LocalDNAIndexMapping->DriverJointsToControlAttributesMap)
-		{
-			const FCompactPoseBoneIndex CompactPoseBoneIndex = RequiredBones.MakeCompactPoseIndex(Mapping.MeshPoseBoneIndex);
-			if (CompactPoseBoneIndex != INDEX_NONE) {
-				if ((Mapping.RotationX != INDEX_NONE) && (Mapping.RotationY != INDEX_NONE) && (Mapping.RotationZ != INDEX_NONE) && (Mapping.RotationW != INDEX_NONE))
+			// Populate mapping of DNA joint indices to CompactPoseBoneIndex entries (used in updating joints with results from RigLogic)
+			const TArray<uint16>& VariableJointIndices = LocalRigRuntimeContext->VariableJointIndicesPerLOD[CurrentLOD].Values;
+			auto& JointsMapDNAIndicesToCompactPoseBoneIndices = CurrentLODJointMapping.JointsMapDNAIndicesToCompactPoseBoneIndices;
+			JointsMapDNAIndicesToCompactPoseBoneIndices.Empty();
+			JointsMapDNAIndicesToCompactPoseBoneIndices.Reserve(VariableJointIndices.Num());
+			for (const uint16 JointIndex : VariableJointIndices)
+			{
+				const FMeshPoseBoneIndex MeshPoseBoneIndex = LocalDNAIndexMapping->JointsMapDNAIndicesToMeshPoseBoneIndices[JointIndex];
+				const FCompactPoseBoneIndex CompactPoseBoneIndex = RequiredBones.MakeCompactPoseIndex(MeshPoseBoneIndex);
+				if (CompactPoseBoneIndex != INDEX_NONE)
 				{
-					DenseDriverJointsToControlAttributesMap.Add({CompactPoseBoneIndex, Mapping.DNAJointIndex, Mapping.RotationX, Mapping.RotationY, Mapping.RotationZ, Mapping.RotationW});
+					JointsMapDNAIndicesToCompactPoseBoneIndices.Add({JointIndex, CompactPoseBoneIndex});
 				}
-				else
+			}
+
+			auto& SparseDriverJointsToControlAttributesMap = CurrentLODJointMapping.SparseDriverJointsToControlAttributesMap;
+			auto& DenseDriverJointsToControlAttributesMap = CurrentLODJointMapping.DenseDriverJointsToControlAttributesMap;
+			// Populate driver joint to raw control attribute mapping (used to feed RigLogic with inputs from the joint hierarchy)
+			SparseDriverJointsToControlAttributesMap.Empty();
+			DenseDriverJointsToControlAttributesMap.Empty();
+			DenseDriverJointsToControlAttributesMap.Reserve(LocalDNAIndexMapping->DriverJointsToControlAttributesMap.Num());
+			// Sparse mapping will likely remain empty so no reservation happens
+			for (const auto& Mapping : LocalDNAIndexMapping->DriverJointsToControlAttributesMap)
+			{
+				const FCompactPoseBoneIndex CompactPoseBoneIndex = RequiredBones.MakeCompactPoseIndex(Mapping.MeshPoseBoneIndex);
+				if (CompactPoseBoneIndex != INDEX_NONE)
 				{
-					SparseDriverJointsToControlAttributesMap.Add({CompactPoseBoneIndex, Mapping.DNAJointIndex, Mapping.RotationX, Mapping.RotationY, Mapping.RotationZ, Mapping.RotationW});
+					if ((Mapping.RotationX != INDEX_NONE) && (Mapping.RotationY != INDEX_NONE) && (Mapping.RotationZ != INDEX_NONE) && (Mapping.RotationW != INDEX_NONE))
+					{
+						DenseDriverJointsToControlAttributesMap.Add({CompactPoseBoneIndex, Mapping.DNAJointIndex, Mapping.RotationX, Mapping.RotationY, Mapping.RotationZ, Mapping.RotationW});
+					}
+					else
+					{
+						SparseDriverJointsToControlAttributesMap.Add({CompactPoseBoneIndex, Mapping.DNAJointIndex, Mapping.RotationX, Mapping.RotationY, Mapping.RotationZ, Mapping.RotationW});
+					}
 				}
 			}
+			CurrentLODJointMapping.BoneCount = BoneCountForLOD;
 		}
 	}
 }
@@ -162,7 +181,6 @@ void FAnimNode_RigLogic::Evaluate_AnyThread(FPoseContext& OutputContext)
 	{
 		UpdateControlCurves(OutputContext, LocalDNAIndexMapping.Get(), LocalRigRuntimeContext->RigLogic->GetNeutralJointValues());
 		CalculateRigLogic(LocalRigRuntimeContext->RigLogic.Get());
-		const uint16 CurrentLOD = RigInstance->GetLOD();
 		UpdateJoints(LocalRigRuntimeContext->RigLogic->GetNeutralJointValues(), RigInstance->GetJointOutputs(), OutputContext);
 		UpdateBlendShapeCurves(LocalDNAIndexMapping.Get(), RigInstance->GetBlendShapeOutputs(), OutputContext);
 		UpdateAnimMapCurves(LocalDNAIndexMapping.Get(), RigInstance->GetAnimatedMapOutputs(), OutputContext);
@@ -179,6 +197,8 @@ void FAnimNode_RigLogic::UpdateControlCurves(const FPoseContext& InputContext, c
 	DECLARE_SCOPE_HIERARCHICAL_COUNTER_FUNC()
 	QUICK_SCOPE_CYCLE_COUNTER(STAT_AnimNode_RigLogic_UpdateControlCurves);
 
+	const uint16 LOD = RigInstance->GetLOD();
+
 	// Combine control attribute curve with input curve to get indexed curve to apply to rig
 	// Curve elements that dont have a control mapping will have INDEX_NONE as their index
 	UE::Anim::FNamedValueArrayUtils::Union(InputContext.Curve, DNAIndexMapping->ControlAttributeCurves,
@@ -191,6 +211,7 @@ void FAnimNode_RigLogic::UpdateControlCurves(const FPoseContext& InputContext, c
 		});
 
 	// The sparse mapping is NOT guaranteed to supply all quaternion attributes, so checks for each attribute mapping are present
+	const auto& SparseDriverJointsToControlAttributesMap = LocalJointMappingsPerLOD[LOD].SparseDriverJointsToControlAttributesMap;
 	for (int32 MappingIndex = 0; MappingIndex < SparseDriverJointsToControlAttributesMap.Num(); ++MappingIndex)
 	{
 		const FCompactPoseBoneControlAttributeMapping& Mapping = SparseDriverJointsToControlAttributesMap[MappingIndex];
@@ -222,6 +243,7 @@ void FAnimNode_RigLogic::UpdateControlCurves(const FPoseContext& InputContext, c
 	}
 
 	// The dense mapping is guaranteed to supply all quaternion attributes, so NO checks for each attribute mapping are present
+	const auto& DenseDriverJointsToControlAttributesMap = LocalJointMappingsPerLOD[LOD].DenseDriverJointsToControlAttributesMap;
 	for (int32 MappingIndex = 0; MappingIndex < DenseDriverJointsToControlAttributesMap.Num(); ++MappingIndex)
 	{
 		const FCompactPoseBoneControlAttributeMapping& Mapping = DenseDriverJointsToControlAttributesMap[MappingIndex];
@@ -239,7 +261,8 @@ void FAnimNode_RigLogic::UpdateControlCurves(const FPoseContext& InputContext, c
 		RigInstance->SetRawControl(Mapping.RotationW, DeltaPoseRotation.w);
 	}
 
-	if (RigInstance->GetNeuralNetworkCount() != 0) {
+	if (RigInstance->GetNeuralNetworkCount() != 0)
+	{
 		UE::Anim::FNamedValueArrayUtils::Union(InputContext.Curve, DNAIndexMapping->NeuralNetworkMaskCurves,
 			[this](const UE::Anim::FCurveElement& InCurveElement, const UE::Anim::FCurveElementIndexed& InControlAttributeCurveElement, UE::Anim::ENamedValueUnionFlags InFlags)
 			{
@@ -263,8 +286,10 @@ void FAnimNode_RigLogic::UpdateJoints(TArrayView<const float> NeutralJointValues
 	DECLARE_SCOPE_HIERARCHICAL_COUNTER_FUNC()
 	QUICK_SCOPE_CYCLE_COUNTER(STAT_AnimNode_RigLogic_UpdateJoints);
 
+	const uint16 LOD = RigInstance->GetLOD();
 	const float* N = NeutralJointValues.GetData();
 	const float* D = DeltaJointValues.GetData();
+	const auto& JointsMapDNAIndicesToCompactPoseBoneIndices = LocalJointMappingsPerLOD[LOD].JointsMapDNAIndicesToCompactPoseBoneIndices;
 	for (const FJointCompactPoseBoneMapping& Mapping : JointsMapDNAIndicesToCompactPoseBoneIndices)
 	{
 		const uint16 AttrIndex = Mapping.JointIndex * ATTR_COUNT_PER_JOINT;
diff --git a/Engine/Plugins/Animation/RigLogic/Source/RigLogicModule/Public/AnimNode_RigLogic.h b/Engine/Plugins/Animation/RigLogic/Source/RigLogicModule/Public/AnimNode_RigLogic.h
--- a/Engine/Plugins/Animation/RigLogic/Source/RigLogicModule/Public/AnimNode_RigLogic.h
+++ b/Engine/Plugins/Animation/RigLogic/Source/RigLogicModule/Public/AnimNode_RigLogic.h
@@ -44,12 +44,14 @@ private:
 	void UpdateAnimMapCurves(const FDNAIndexMapping* DNAIndexMapping, TArrayView<const float> AnimMapOutputs, FPoseContext& OutputContext);
 
 private:
-	struct FJointCompactPoseBoneMapping {
+	struct FJointCompactPoseBoneMapping
+	{
 		uint16 JointIndex;
 		FCompactPoseBoneIndex CompactPoseBoneIndex;
 	};
 
-	struct FCompactPoseBoneControlAttributeMapping {
+	struct FCompactPoseBoneControlAttributeMapping
+	{
 		FCompactPoseBoneIndex CompactPoseBoneIndex;
 		int32 DNAJointIndex;
 		int32 RotationX;
@@ -58,6 +60,14 @@ private:
 		int32 RotationW;
 	};
 
+	struct FCachedJointMapping
+	{
+		TArray<FJointCompactPoseBoneMapping> JointsMapDNAIndicesToCompactPoseBoneIndices;
+		TArray<FCompactPoseBoneControlAttributeMapping> SparseDriverJointsToControlAttributesMap;
+		TArray<FCompactPoseBoneControlAttributeMapping> DenseDriverJointsToControlAttributesMap;
+		int32 BoneCount = -1;
+	};
+
 private:
 	/*
 	 * Max LOD level that post-process AnimBPs are evaluated.
@@ -71,7 +81,5 @@ private:
 	TSharedPtr<FSharedRigRuntimeContext> LocalRigRuntimeContext;
 	TSharedPtr<FDNAIndexMapping> LocalDNAIndexMapping;
 	FRigInstance* RigInstance;
-	TArray<FJointCompactPoseBoneMapping> JointsMapDNAIndicesToCompactPoseBoneIndices;
-	TArray<FCompactPoseBoneControlAttributeMapping> SparseDriverJointsToControlAttributesMap;
-	TArray<FCompactPoseBoneControlAttributeMapping> DenseDriverJointsToControlAttributesMap;
+	TArray<FCachedJointMapping> LocalJointMappingsPerLOD;
 };
