Enable NEON for RigLogic on supported platforms

diff --git a/Engine/Platforms/Android/Plugins/RigLogic/RigLogic_Android.uplugin b/Engine/Platforms/Android/Plugins/RigLogic/RigLogic_Android.uplugin
new file mode 100644
--- /dev/null
+++ b/Engine/Platforms/Android/Plugins/RigLogic/RigLogic_Android.uplugin
@@ -0,0 +1,5 @@
+{
+	"FileVersion": 3,
+	"bIsPluginExtension": true
+}
+
diff --git a/Engine/Platforms/Android/Plugins/RigLogic/Source/RigLogicLib/RigLogicLib_Android.Build.cs b/Engine/Platforms/Android/Plugins/RigLogic/Source/RigLogicLib/RigLogicLib_Android.Build.cs
new file mode 100644
--- /dev/null
+++ b/Engine/Platforms/Android/Plugins/RigLogic/Source/RigLogicLib/RigLogicLib_Android.Build.cs
@@ -0,0 +1,12 @@
+// Copyright Epic Games, Inc. All Rights Reserved.
+namespace UnrealBuildTool.Rules
+{
+	public class RigLogicLib_Android : RigLogicLib
+	{
+		public RigLogicLib_Android(ReadOnlyTargetRules Target) : base(Target)
+		{
+			PrivateDefinitions.Add("RL_BUILD_WITH_NEON=1");
+			PrivateDefinitions.Add("RL_USE_HALF_FLOATS=1");
+		}
+	}
+}
\ No newline at end of file
