Fix RigLogic rbf pose weight calculation when multiple poses output to the same output control

diff --git a/Engine/Plugins/Animation/RigLogic/RigLogic.uplugin b/Engine/Plugins/Animation/RigLogic/RigLogic.uplugin
--- a/Engine/Plugins/Animation/RigLogic/RigLogic.uplugin
+++ b/Engine/Plugins/Animation/RigLogic/RigLogic.uplugin
@@ -1,9 +1,9 @@
 {
 	"FileVersion": 3,
 	"Version": 10.2,
-	"VersionName": "10.2.1",
-	"FriendlyName": "RigLogic Plugin v10.2.1",
-	"Description": "3Lateral RigLogic Plugin for Facial Animation v10.2.1",
+	"VersionName": "10.2.5",
+	"FriendlyName": "RigLogic Plugin v10.2.5",
+	"Description": "3Lateral RigLogic Plugin for Facial Animation v10.2.5",
 	"Category": "Animation",
 	"CreatedBy": "3Lateral",
 	"CreatedByURL": "https://www.3lateral.com/",
diff --git a/Engine/Plugins/Animation/RigLogic/Source/RigLogicLib/Private/riglogic/controls/Controls.cpp b/Engine/Plugins/Animation/RigLogic/Source/RigLogicLib/Private/riglogic/controls/Controls.cpp
--- a/Engine/Plugins/Animation/RigLogic/Source/RigLogicLib/Private/riglogic/controls/Controls.cpp
+++ b/Engine/Plugins/Animation/RigLogic/Source/RigLogicLib/Private/riglogic/controls/Controls.cpp
@@ -36,7 +36,7 @@ void Controls::mapRawToGUI(ControlsInputInstance* instance) const {
 }
 
 void Controls::calculate(ControlsInputInstance* instance) const {
-    psds.calculate(instance->getInputBuffer().data(), guiToRawMapping.getOutputCount());
+    psds.calculate(instance->getInputBuffer().data(), instance->getRawControlCount());
 }
 
 }  // namespace rl4
diff --git a/Engine/Plugins/Animation/RigLogic/Source/RigLogicLib/Private/riglogic/controls/ControlsInputInstance.h b/Engine/Plugins/Animation/RigLogic/Source/RigLogicLib/Private/riglogic/controls/ControlsInputInstance.h
--- a/Engine/Plugins/Animation/RigLogic/Source/RigLogicLib/Private/riglogic/controls/ControlsInputInstance.h
+++ b/Engine/Plugins/Animation/RigLogic/Source/RigLogicLib/Private/riglogic/controls/ControlsInputInstance.h
@@ -26,6 +26,11 @@ class ControlsInputInstance {
         virtual ArrayView<float> getInputBuffer() = 0;
         virtual ConstArrayView<float> getGUIControlBuffer() const = 0;
         virtual ConstArrayView<float> getInputBuffer() const = 0;
+        virtual std::uint16_t getGUIControlCount() const = 0;
+        virtual std::uint16_t getRawControlCount() const = 0;
+        virtual std::uint16_t getPSDControlCount() const = 0;
+        virtual std::uint16_t getMLControlCount() const = 0;
+        virtual std::uint16_t getRBFControlCount() const = 0;
 
 };
 
diff --git a/Engine/Plugins/Animation/RigLogic/Source/RigLogicLib/Private/riglogic/controls/instances/StandardControlsInputInstance.cpp b/Engine/Plugins/Animation/RigLogic/Source/RigLogicLib/Private/riglogic/controls/instances/StandardControlsInputInstance.cpp
--- a/Engine/Plugins/Animation/RigLogic/Source/RigLogicLib/Private/riglogic/controls/instances/StandardControlsInputInstance.cpp
+++ b/Engine/Plugins/Animation/RigLogic/Source/RigLogicLib/Private/riglogic/controls/instances/StandardControlsInputInstance.cpp
@@ -7,14 +7,19 @@
 
 namespace rl4 {
 
-StandardControlsInputInstance::StandardControlsInputInstance(std::uint16_t guiControlCount,
-                                                             std::uint16_t rawControlCount,
-                                                             std::uint16_t psdControlCount,
-                                                             std::uint16_t mlControlCount,
-                                                             std::uint16_t rbfControlCount,
+StandardControlsInputInstance::StandardControlsInputInstance(std::uint16_t guiControlCount_,
+                                                             std::uint16_t rawControlCount_,
+                                                             std::uint16_t psdControlCount_,
+                                                             std::uint16_t mlControlCount_,
+                                                             std::uint16_t rbfControlCount_,
                                                              MemoryResource* memRes) :
-    guiControlBuffer{guiControlCount, {}, memRes},
-    inputBuffer{static_cast<std::size_t>(rawControlCount + psdControlCount + mlControlCount + rbfControlCount), {}, memRes} {
+    guiControlBuffer{guiControlCount_, {}, memRes},
+    inputBuffer{static_cast<std::size_t>(rawControlCount_ + psdControlCount_ + mlControlCount_ + rbfControlCount_), {}, memRes},
+    guiControlCount{guiControlCount_},
+    rawControlCount{rawControlCount_},
+    psdControlCount{psdControlCount_},
+    mlControlCount{mlControlCount_},
+    rbfControlCount{rbfControlCount_} {
 }
 
 ArrayView<float> StandardControlsInputInstance::getGUIControlBuffer() {
@@ -33,4 +38,24 @@ ConstArrayView<float> StandardControlsInputInstance::getInputBuffer() const {
     return ConstArrayView<float>{inputBuffer};
 }
 
+std::uint16_t StandardControlsInputInstance::getGUIControlCount() const {
+    return guiControlCount;
+}
+
+std::uint16_t StandardControlsInputInstance::getRawControlCount() const {
+    return rawControlCount;
+}
+
+std::uint16_t StandardControlsInputInstance::getPSDControlCount() const {
+    return psdControlCount;
+}
+
+std::uint16_t StandardControlsInputInstance::getMLControlCount() const {
+    return mlControlCount;
+}
+
+std::uint16_t StandardControlsInputInstance::getRBFControlCount() const {
+    return rbfControlCount;
+}
+
 }  // namespace rl4
diff --git a/Engine/Plugins/Animation/RigLogic/Source/RigLogicLib/Private/riglogic/controls/instances/StandardControlsInputInstance.h b/Engine/Plugins/Animation/RigLogic/Source/RigLogicLib/Private/riglogic/controls/instances/StandardControlsInputInstance.h
--- a/Engine/Plugins/Animation/RigLogic/Source/RigLogicLib/Private/riglogic/controls/instances/StandardControlsInputInstance.h
+++ b/Engine/Plugins/Animation/RigLogic/Source/RigLogicLib/Private/riglogic/controls/instances/StandardControlsInputInstance.h
@@ -11,22 +11,31 @@ namespace rl4 {
 
 class StandardControlsInputInstance : public ControlsInputInstance {
     public:
-        StandardControlsInputInstance(std::uint16_t guiControlCount,
-                                      std::uint16_t rawControlCount,
-                                      std::uint16_t psdControlCount,
-                                      std::uint16_t mlControlCount,
-                                      std::uint16_t rbfControlCount,
+        StandardControlsInputInstance(std::uint16_t guiControlCount_,
+                                      std::uint16_t rawControlCount_,
+                                      std::uint16_t psdControlCount_,
+                                      std::uint16_t mlControlCount_,
+                                      std::uint16_t rbfControlCount_,
                                       MemoryResource* memRes);
 
         ArrayView<float> getGUIControlBuffer() override;
         ArrayView<float> getInputBuffer() override;
         ConstArrayView<float> getGUIControlBuffer() const override;
         ConstArrayView<float> getInputBuffer() const override;
+        std::uint16_t getGUIControlCount() const override;
+        std::uint16_t getRawControlCount() const override;
+        std::uint16_t getPSDControlCount() const override;
+        std::uint16_t getMLControlCount() const override;
+        std::uint16_t getRBFControlCount() const override;
 
     private:
         Vector<float> guiControlBuffer;
         Vector<float> inputBuffer;
-
+        std::uint16_t guiControlCount;
+        std::uint16_t rawControlCount;
+        std::uint16_t psdControlCount;
+        std::uint16_t mlControlCount;
+        std::uint16_t rbfControlCount;
 };
 
 }  // namespace rl4
diff --git a/Engine/Plugins/Animation/RigLogic/Source/RigLogicLib/Private/riglogic/joints/JointBehaviorFilter.cpp b/Engine/Plugins/Animation/RigLogic/Source/RigLogicLib/Private/riglogic/joints/JointBehaviorFilter.cpp
--- a/Engine/Plugins/Animation/RigLogic/Source/RigLogicLib/Private/riglogic/joints/JointBehaviorFilter.cpp
+++ b/Engine/Plugins/Animation/RigLogic/Source/RigLogicLib/Private/riglogic/joints/JointBehaviorFilter.cpp
@@ -3,7 +3,7 @@
 #include "riglogic/joints/JointBehaviorFilter.h"
 
 #include <cstdint>
-#if defined(_MSC_VER) && !defined(__clang__) && (_MSC_VER >= 1900) && (__cplusplus >= 202002L)
+#if defined(_MSC_VER) && !defined(__clang__) && (_MSC_VER < 1938) && (_MSC_VER >= 1900) && (__cplusplus >= 202002L)
     #include <span>
 #endif
 
@@ -147,7 +147,7 @@ std::uint16_t JointBehaviorFilter::getJointGroupCount() const {
 
 void JointBehaviorFilter::copyInputIndices(std::uint16_t jointGroupIndex, ArrayView<std::uint16_t> dest) const {
     const auto inputIndices = reader->getJointGroupInputIndices(jointGroupIndex);
-    #if defined(_MSC_VER) && !defined(__clang__)
+    #if defined(_MSC_VER) && !defined(__clang__) && (_MSC_VER < 1938)
         #if (_MSC_VER >= 1900) && (__cplusplus >= 202002L)
             std::copy(inputIndices.begin(),
                       inputIndices.end(),
@@ -181,7 +181,7 @@ void JointBehaviorFilter::copyValues(std::uint16_t jointGroupIndex, ArrayView<fl
     for (std::size_t row = {}; row < rowCount; ++row) {
         if (isAttributeEnabled(outputIndices[row])) {
             const auto rowOfValues = values.subview(row * colCount, colCount);
-            #if defined(_MSC_VER) && !defined(__clang__)
+            #if defined(_MSC_VER) && !defined(__clang__) && (_MSC_VER < 1938)
                 #if (_MSC_VER >= 1900) && (__cplusplus >= 202002L)
                     std::copy(rowOfValues.begin(),
                               rowOfValues.end(),
diff --git a/Engine/Plugins/Animation/RigLogic/Source/RigLogicLib/Private/riglogic/joints/cpu/quaternions/CalculationStrategy.h b/Engine/Plugins/Animation/RigLogic/Source/RigLogicLib/Private/riglogic/joints/cpu/quaternions/CalculationStrategy.h
--- a/Engine/Plugins/Animation/RigLogic/Source/RigLogicLib/Private/riglogic/joints/cpu/quaternions/CalculationStrategy.h
+++ b/Engine/Plugins/Animation/RigLogic/Source/RigLogicLib/Private/riglogic/joints/cpu/quaternions/CalculationStrategy.h
@@ -103,7 +103,7 @@ namespace rl4 {
  *     LERPED_G = [LERPED_X[3], LERPED_Y[3], LERPED_Z[3], LERPED_W[3]]
  */
 template<typename TFVec>
-static FORCE_INLINE void fastlerpWithIdentity(TFVec& qxACEG, TFVec& qyACEG, TFVec& qzACEG, TFVec& qwACEG, TFVec weight) {
+static FORCE_INLINE void fastlerpWithIdentity(TFVec& qxACEG, TFVec& qyACEG, TFVec& qzACEG, TFVec& qwACEG, const TFVec& weight) {
     const TFVec negOne{-1.0f};
     const TFVec one{1.0f};
     const TFVec oneMinusWeight = one - weight;
diff --git a/Engine/Plugins/Animation/RigLogic/Source/RigLogicLib/Private/riglogic/rbf/cpu/CPURBFBehaviorEvaluator.h b/Engine/Plugins/Animation/RigLogic/Source/RigLogicLib/Private/riglogic/rbf/cpu/CPURBFBehaviorEvaluator.h
--- a/Engine/Plugins/Animation/RigLogic/Source/RigLogicLib/Private/riglogic/rbf/cpu/CPURBFBehaviorEvaluator.h
+++ b/Engine/Plugins/Animation/RigLogic/Source/RigLogicLib/Private/riglogic/rbf/cpu/CPURBFBehaviorEvaluator.h
@@ -92,7 +92,7 @@ class Evaluator : public RBFBehaviorEvaluator {
                 }
 
                 for (std::uint16_t ci = 0u; ci < outputControlIndices.size(); ci++) {
-                    rawControls[outputControlIndices[ci]] = outputControlWeights[ci] * outputWeightsBuffer[i] * inputWeight;
+                    rawControls[outputControlIndices[ci]] += outputControlWeights[ci] * outputWeightsBuffer[i] * inputWeight;
                 }
             }
         }
@@ -102,10 +102,13 @@ class Evaluator : public RBFBehaviorEvaluator {
             assert(lod < lods.indicesPerLOD.size());
             const auto& solverIndices = lods.indicesPerLOD[lod];
             auto rawControls = inputs->getInputBuffer();
+            const auto offset = static_cast<std::size_t>(inputs->getRawControlCount()) +
+                static_cast<std::size_t>(inputs->getPSDControlCount()) +
+                static_cast<std::size_t>(inputs->getMLControlCount());
+            std::fill_n(rawControls.data() + offset, inputs->getRBFControlCount(), 0.0f);
             auto inputBuffer = static_cast<OutputInstance*>(intermediateOutputs)->getInputBuffer();
             auto intermediateWeightsBuffer = static_cast<OutputInstance*>(intermediateOutputs)->getIntermediateWeightsBuffer();
             auto outputWeightsBuffer = static_cast<OutputInstance*>(intermediateOutputs)->getOutputWeightsBuffer();
-
             for (const auto solverIndex : solverIndices) {
                 assert(solverIndex < solvers.size());
                 calculate(solverIndex, rawControls, inputBuffer, intermediateWeightsBuffer, outputWeightsBuffer);
@@ -122,10 +125,13 @@ class Evaluator : public RBFBehaviorEvaluator {
             assert(solverIndex < solvers.size());
 
             auto rawControls = inputs->getInputBuffer();
+            const auto offset = static_cast<std::size_t>(inputs->getRawControlCount()) +
+                static_cast<std::size_t>(inputs->getPSDControlCount()) +
+                static_cast<std::size_t>(inputs->getMLControlCount());
+            std::fill_n(rawControls.data() + offset, inputs->getRBFControlCount(), 0.0f);
             auto inputBuffer = static_cast<OutputInstance*>(intermediateOutputs)->getInputBuffer();
             auto intermediateWeightsBuffer = static_cast<OutputInstance*>(intermediateOutputs)->getIntermediateWeightsBuffer();
             auto outputWeightsBuffer = static_cast<OutputInstance*>(intermediateOutputs)->getOutputWeightsBuffer();
-
             calculate(solverIndex, rawControls, inputBuffer, intermediateWeightsBuffer, outputWeightsBuffer);
         }
 
diff --git a/Engine/Plugins/Animation/RigLogic/Source/RigLogicLib/Private/riglogic/riglogic/RigInstanceImpl.cpp b/Engine/Plugins/Animation/RigLogic/Source/RigLogicLib/Private/riglogic/riglogic/RigInstanceImpl.cpp
--- a/Engine/Plugins/Animation/RigLogic/Source/RigLogicLib/Private/riglogic/riglogic/RigInstanceImpl.cpp
+++ b/Engine/Plugins/Animation/RigLogic/Source/RigLogicLib/Private/riglogic/riglogic/RigInstanceImpl.cpp
@@ -20,7 +20,7 @@
 #include <cassert>
 #include <cstdint>
 #include <numeric>
-#if defined(_MSC_VER) && !defined(__clang__) && (_MSC_VER >= 1900) && (__cplusplus >= 202002L)
+#if defined(_MSC_VER) && !defined(__clang__) && (_MSC_VER < 1938) && (_MSC_VER >= 1900) && (__cplusplus >= 202002L)
     #include <span>
 #endif
 #ifdef _MSC_VER
@@ -90,7 +90,7 @@ ConstArrayView<float> RigInstanceImpl::getGUIControlValues() const {
 
 void RigInstanceImpl::setGUIControlValues(const float* values) {
     auto guiControlBuffer = controlsInstance->getGUIControlBuffer();
-    #if defined(_MSC_VER) && !defined(__clang__)
+    #if defined(_MSC_VER) && !defined(__clang__) && (_MSC_VER < 1938)
         #if (_MSC_VER >= 1900) && (__cplusplus >= 202002L)
             std::copy(values,
                       values + guiControlCount,
@@ -124,7 +124,7 @@ ConstArrayView<float> RigInstanceImpl::getRawControlValues() const {
 
 void RigInstanceImpl::setRawControlValues(const float* values) {
     auto inputBuffer = controlsInstance->getInputBuffer();
-    #if defined(_MSC_VER) && !defined(__clang__)
+    #if defined(_MSC_VER) && !defined(__clang__) && (_MSC_VER < 1938)
         #if (_MSC_VER >= 1900) && (__cplusplus >= 202002L)
             std::copy(values,
                       values + rawControlCount,
diff --git a/Engine/Plugins/Animation/RigLogic/Source/RigLogicLib/Public/dna/version/Version.h b/Engine/Plugins/Animation/RigLogic/Source/RigLogicLib/Public/dna/version/Version.h
--- a/Engine/Plugins/Animation/RigLogic/Source/RigLogicLib/Public/dna/version/Version.h
+++ b/Engine/Plugins/Animation/RigLogic/Source/RigLogicLib/Public/dna/version/Version.h
@@ -4,5 +4,5 @@
 
 #define DNA_MAJOR_VERSION 9
 #define DNA_MINOR_VERSION 4
-#define DNA_PATCH_VERSION 1
-#define DNA_VERSION_STRING "9.4.1"
+#define DNA_PATCH_VERSION 3
+#define DNA_VERSION_STRING "9.4.3"
diff --git a/Engine/Plugins/Animation/RigLogic/Source/RigLogicLib/Public/pma/version/Version.h b/Engine/Plugins/Animation/RigLogic/Source/RigLogicLib/Public/pma/version/Version.h
--- a/Engine/Plugins/Animation/RigLogic/Source/RigLogicLib/Public/pma/version/Version.h
+++ b/Engine/Plugins/Animation/RigLogic/Source/RigLogicLib/Public/pma/version/Version.h
@@ -4,5 +4,5 @@
 
 #define PMA_MAJOR_VERSION 1
 #define PMA_MINOR_VERSION 3
-#define PMA_PATCH_VERSION 11
-#define PMA_VERSION_STRING "1.3.11"
+#define PMA_PATCH_VERSION 12
+#define PMA_VERSION_STRING "1.3.12"
diff --git a/Engine/Plugins/Animation/RigLogic/Source/RigLogicLib/Public/riglogic/version/Version.h b/Engine/Plugins/Animation/RigLogic/Source/RigLogicLib/Public/riglogic/version/Version.h
--- a/Engine/Plugins/Animation/RigLogic/Source/RigLogicLib/Public/riglogic/version/Version.h
+++ b/Engine/Plugins/Animation/RigLogic/Source/RigLogicLib/Public/riglogic/version/Version.h
@@ -4,5 +4,5 @@
 
 #define RL_MAJOR_VERSION 10
 #define RL_MINOR_VERSION 2
-#define RL_PATCH_VERSION 1
-#define RL_VERSION_STRING "10.2.1"
+#define RL_PATCH_VERSION 5
+#define RL_VERSION_STRING "10.2.5"
diff --git a/Engine/Plugins/Animation/RigLogic/Source/RigLogicLib/Public/status/version/Version.h b/Engine/Plugins/Animation/RigLogic/Source/RigLogicLib/Public/status/version/Version.h
--- a/Engine/Plugins/Animation/RigLogic/Source/RigLogicLib/Public/status/version/Version.h
+++ b/Engine/Plugins/Animation/RigLogic/Source/RigLogicLib/Public/status/version/Version.h
@@ -4,5 +4,5 @@
 
 #define SC_MAJOR_VERSION 1
 #define SC_MINOR_VERSION 2
-#define SC_PATCH_VERSION 5
-#define SC_VERSION_STRING "1.2.5"
+#define SC_PATCH_VERSION 6
+#define SC_VERSION_STRING "1.2.6"
diff --git a/Engine/Plugins/Animation/RigLogic/Source/RigLogicLib/Public/terse/types/DynArray.h b/Engine/Plugins/Animation/RigLogic/Source/RigLogicLib/Public/terse/types/DynArray.h
--- a/Engine/Plugins/Animation/RigLogic/Source/RigLogicLib/Public/terse/types/DynArray.h
+++ b/Engine/Plugins/Animation/RigLogic/Source/RigLogicLib/Public/terse/types/DynArray.h
@@ -13,7 +13,7 @@
 #include <functional>
 #include <iterator>
 #include <memory>
-#if defined(_MSC_VER) && !defined(__clang__) && (_MSC_VER >= 1900) && (__cplusplus >= 202002L)
+#if defined(_MSC_VER) && !defined(__clang__) && (_MSC_VER < 1938) && (_MSC_VER >= 1900) && (__cplusplus >= 202002L)
     #include <span>
 #endif
 #include <type_traits>
@@ -71,7 +71,7 @@ class DynArray {
         DynArray(TIterator start, TIterator end, const allocator_type& allocator = allocator_type{}) :
             DynArray{static_cast<std::size_t>(std::distance(start, end)), allocator} {
 
-            #if defined(_MSC_VER) && !defined(__clang__)
+            #if defined(_MSC_VER) && !defined(__clang__) && (_MSC_VER < 1938)
                 if (size() != 0ul) {
                     #if (_MSC_VER >= 1900) && (__cplusplus >= 202002L)
                         std::copy(start, end, std::span{data(), size()}.begin());
@@ -212,7 +212,7 @@ class DynArray {
         template<typename TIterator>
         void assign(TIterator start, TIterator end) {
             resize_uninitialized(static_cast<std::size_t>(std::distance(start, end)));
-            #if defined(_MSC_VER) && !defined(__clang__)
+            #if defined(_MSC_VER) && !defined(__clang__) && (_MSC_VER < 1938)
                 if (size() != 0ul) {
                     #if (_MSC_VER >= 1900) && (__cplusplus >= 202002L)
                         std::copy(start, end, std::span{data(), size()}.begin());
diff --git a/Engine/Plugins/Animation/RigLogic/Source/RigLogicLib/Public/terse/version/Version.h b/Engine/Plugins/Animation/RigLogic/Source/RigLogicLib/Public/terse/version/Version.h
--- a/Engine/Plugins/Animation/RigLogic/Source/RigLogicLib/Public/terse/version/Version.h
+++ b/Engine/Plugins/Animation/RigLogic/Source/RigLogicLib/Public/terse/version/Version.h
@@ -4,5 +4,5 @@
 
 #define TERSE_MAJOR_VERSION 3
 #define TERSE_MINOR_VERSION 1
-#define TERSE_PATCH_VERSION 9
-#define TERSE_VERSION_STRING "3.1.9"
+#define TERSE_PATCH_VERSION 10
+#define TERSE_VERSION_STRING "3.1.10"
diff --git a/Engine/Plugins/Animation/RigLogic/Source/RigLogicLib/Public/trio/version/Version.h b/Engine/Plugins/Animation/RigLogic/Source/RigLogicLib/Public/trio/version/Version.h
--- a/Engine/Plugins/Animation/RigLogic/Source/RigLogicLib/Public/trio/version/Version.h
+++ b/Engine/Plugins/Animation/RigLogic/Source/RigLogicLib/Public/trio/version/Version.h
@@ -4,5 +4,5 @@
 
 #define TRIO_MAJOR_VERSION 4
 #define TRIO_MINOR_VERSION 0
-#define TRIO_PATCH_VERSION 15
-#define TRIO_VERSION_STRING "4.0.15"
+#define TRIO_PATCH_VERSION 16
+#define TRIO_VERSION_STRING "4.0.16"
