Enable NEON intrinsics for RigLogic on supported platforms

diff --git a/Engine/Platforms/Android/Plugins/RigLogic/RigLogic_Android.uplugin b/Engine/Platforms/Android/Plugins/RigLogic/RigLogic_Android.uplugin
new file mode 100644
--- /dev/null
+++ b/Engine/Platforms/Android/Plugins/RigLogic/RigLogic_Android.uplugin
@@ -0,0 +1,5 @@
+{
+	"FileVersion": 3,
+	"bIsPluginExtension": true
+}
+
diff --git a/Engine/Platforms/Android/Plugins/RigLogic/Source/RigLogicLib/RigLogicLib_Android.Build.cs b/Engine/Platforms/Android/Plugins/RigLogic/Source/RigLogicLib/RigLogicLib_Android.Build.cs
new file mode 100644
--- /dev/null
+++ b/Engine/Platforms/Android/Plugins/RigLogic/Source/RigLogicLib/RigLogicLib_Android.Build.cs
@@ -0,0 +1,13 @@
+// Copyright Epic Games, Inc. All Rights Reserved.
+namespace UnrealBuildTool.Rules
+{
+	public class RigLogicLib_Android : RigLogicLib
+	{
+		public RigLogicLib_Android(ReadOnlyTargetRules Target) : base(Target)
+		{
+			PrivateDefinitions.Add("RL_AUTODETECT_SSE=1");
+			PrivateDefinitions.Add("RL_AUTODETECT_NEON=1");
+			PrivateDefinitions.Add("RL_AUTODETECT_HALF_FLOATS=1");
+		}
+	}
+}
\ No newline at end of file
diff --git a/Engine/Platforms/IOS/Plugins/RigLogic/RigLogic_IOS.uplugin b/Engine/Platforms/IOS/Plugins/RigLogic/RigLogic_IOS.uplugin
new file mode 100644
--- /dev/null
+++ b/Engine/Platforms/IOS/Plugins/RigLogic/RigLogic_IOS.uplugin
@@ -0,0 +1,5 @@
+{
+	"FileVersion": 3,
+	"bIsPluginExtension": true
+}
+
diff --git a/Engine/Platforms/IOS/Plugins/RigLogic/Source/RigLogicLib/RigLogicLib_IOS.Build.cs b/Engine/Platforms/IOS/Plugins/RigLogic/Source/RigLogicLib/RigLogicLib_IOS.Build.cs
new file mode 100644
--- /dev/null
+++ b/Engine/Platforms/IOS/Plugins/RigLogic/Source/RigLogicLib/RigLogicLib_IOS.Build.cs
@@ -0,0 +1,12 @@
+// Copyright Epic Games, Inc. All Rights Reserved.
+namespace UnrealBuildTool.Rules
+{
+	public class RigLogicLib_IOS : RigLogicLib
+	{
+		public RigLogicLib_IOS(ReadOnlyTargetRules Target) : base(Target)
+		{
+			PrivateDefinitions.Add("RL_AUTODETECT_NEON=1");
+			PrivateDefinitions.Add("RL_AUTODETECT_HALF_FLOATS=1");
+		}
+	}
+}
\ No newline at end of file
diff --git a/Engine/Plugins/Animation/RigLogic/RigLogic.uplugin b/Engine/Plugins/Animation/RigLogic/RigLogic.uplugin
--- a/Engine/Plugins/Animation/RigLogic/RigLogic.uplugin
+++ b/Engine/Plugins/Animation/RigLogic/RigLogic.uplugin
@@ -1,9 +1,9 @@
 {
 	"FileVersion": 3,
 	"Version": 10.2,
-	"VersionName": "10.2.5",
-	"FriendlyName": "RigLogic Plugin v10.2.5",
-	"Description": "3Lateral RigLogic Plugin for Facial Animation v10.2.5",
+	"VersionName": "10.2.6",
+	"FriendlyName": "RigLogic Plugin v10.2.6",
+	"Description": "3Lateral RigLogic Plugin for Facial Animation v10.2.6",
 	"Category": "Animation",
 	"CreatedBy": "3Lateral",
 	"CreatedByURL": "https://www.3lateral.com/",
diff --git a/Engine/Plugins/Animation/RigLogic/Source/RigLogicLib/Private/riglogic/system/simd/Detect.h b/Engine/Plugins/Animation/RigLogic/Source/RigLogicLib/Private/riglogic/system/simd/Detect.h
--- a/Engine/Plugins/Animation/RigLogic/Source/RigLogicLib/Private/riglogic/system/simd/Detect.h
+++ b/Engine/Plugins/Animation/RigLogic/Source/RigLogicLib/Private/riglogic/system/simd/Detect.h
@@ -8,5 +8,21 @@
     defined(_X86_) || defined(__THW_INTEL__) || defined(__I86__) || defined(__amd64__) || defined(__amd64) || \
     defined(__x86_64__) || defined(__x86_64) || defined(_M_X64) || defined(_M_AMD64)
         #define RL_BUILD_WITH_SSE 1
+        #if defined(RL_AUTODETECT_HALF_FLOATS) && !defined(RL_USE_HALF_FLOATS)
+            #if defined(__F16C__) || (defined(_MSC_VER) && (_MSC_VER >= 1700) && defined(__AVX__)) || (defined(__INTEL_COMPILER) && defined(__AVX__))
+                #define RL_USE_HALF_FLOATS 1
+            #endif
+        #endif  // RL_AUTODETECT_HALF_FLOATS
     #endif
 #endif  // RL_AUTODETECT_SSE
+
+#if defined(RL_AUTODETECT_NEON) && !defined(RL_BUILD_WITH_NEON)
+    #if defined(__ARM_NEON)
+        #define RL_BUILD_WITH_NEON 1
+        #if defined(RL_AUTODETECT_HALF_FLOATS) && !defined(RL_USE_HALF_FLOATS)
+            #if (defined(__GNUC__) && (defined(__arm__) || defined(__aarch64__))) || (defined(_MSC_VER) && defined(_M_ARM64))
+                #define RL_USE_HALF_FLOATS 1
+            #endif
+        #endif  // RL_AUTODETECT_HALF_FLOATS
+    #endif
+#endif  // RL_AUTODETECT_NEON
diff --git a/Engine/Plugins/Animation/RigLogic/Source/RigLogicLib/Public/riglogic/version/Version.h b/Engine/Plugins/Animation/RigLogic/Source/RigLogicLib/Public/riglogic/version/Version.h
--- a/Engine/Plugins/Animation/RigLogic/Source/RigLogicLib/Public/riglogic/version/Version.h
+++ b/Engine/Plugins/Animation/RigLogic/Source/RigLogicLib/Public/riglogic/version/Version.h
@@ -4,5 +4,5 @@
 
 #define RL_MAJOR_VERSION 10
 #define RL_MINOR_VERSION 2
-#define RL_PATCH_VERSION 5
-#define RL_VERSION_STRING "10.2.5"
+#define RL_PATCH_VERSION 6
+#define RL_VERSION_STRING "10.2.6"
