Make sure RigLogic raw controls that are quaternions are default initialized to the identity quaternion, instead of just being zero initialized and fix zero-initialization of rbf output controls so individual solver invocations don't clear each other

diff --git a/Engine/Plugins/Animation/RigLogic/RigLogic.uplugin b/Engine/Plugins/Animation/RigLogic/RigLogic.uplugin
--- a/Engine/Plugins/Animation/RigLogic/RigLogic.uplugin
+++ b/Engine/Plugins/Animation/RigLogic/RigLogic.uplugin
@@ -1,9 +1,9 @@
 {
 	"FileVersion": 3,
 	"Version": 10.2,
-	"VersionName": "10.2.6",
-	"FriendlyName": "RigLogic Plugin v10.2.6",
-	"Description": "3Lateral RigLogic Plugin for Facial Animation v10.2.6",
+	"VersionName": "10.2.7",
+	"FriendlyName": "RigLogic Plugin v10.2.7",
+	"Description": "3Lateral RigLogic Plugin for Facial Animation v10.2.7",
 	"Category": "Animation",
 	"CreatedBy": "3Lateral",
 	"CreatedByURL": "https://www.3lateral.com/",
diff --git a/Engine/Plugins/Animation/RigLogic/Source/RigLogicLib/Private/riglogic/controls/Controls.cpp b/Engine/Plugins/Animation/RigLogic/Source/RigLogicLib/Private/riglogic/controls/Controls.cpp
--- a/Engine/Plugins/Animation/RigLogic/Source/RigLogicLib/Private/riglogic/controls/Controls.cpp
+++ b/Engine/Plugins/Animation/RigLogic/Source/RigLogicLib/Private/riglogic/controls/Controls.cpp
@@ -11,14 +11,18 @@
 
 namespace rl4 {
 
-Controls::Controls(ConditionalTable&& guiToRawMapping_, PSDMatrix&& psds_, ControlsInputInstance::Factory instanceFactory_) :
+Controls::Controls(ConditionalTable&& guiToRawMapping_,
+                   PSDMatrix&& psds_,
+                   Vector<float>&& initialValues_,
+                   ControlsInputInstance::Factory instanceFactory_) :
     guiToRawMapping{std::move(guiToRawMapping_)},
     psds{std::move(psds_)},
+    initialValues{std::move(initialValues_)},
     instanceFactory{instanceFactory_} {
 }
 
 ControlsInputInstance::Pointer Controls::createInstance(MemoryResource* instanceMemRes) const {
-    return instanceFactory(instanceMemRes);
+    return instanceFactory(initialValues, instanceMemRes);
 }
 
 void Controls::mapGUIToRaw(ControlsInputInstance* instance) const {
diff --git a/Engine/Plugins/Animation/RigLogic/Source/RigLogicLib/Private/riglogic/controls/Controls.h b/Engine/Plugins/Animation/RigLogic/Source/RigLogicLib/Private/riglogic/controls/Controls.h
--- a/Engine/Plugins/Animation/RigLogic/Source/RigLogicLib/Private/riglogic/controls/Controls.h
+++ b/Engine/Plugins/Animation/RigLogic/Source/RigLogicLib/Private/riglogic/controls/Controls.h
@@ -16,7 +16,10 @@ class Controls {
         using Pointer = UniqueInstance<Controls>::PointerType;
 
     public:
-        Controls(ConditionalTable&& guiToRawMapping_, PSDMatrix&& psds_, ControlsInputInstance::Factory instanceFactory_);
+        Controls(ConditionalTable&& guiToRawMapping_,
+                 PSDMatrix&& psds_,
+                 Vector<float>&& initialValues_,
+                 ControlsInputInstance::Factory instanceFactory_);
 
         ControlsInputInstance::Pointer createInstance(MemoryResource* instanceMemRes) const;
         void mapGUIToRaw(ControlsInputInstance* instance) const;
@@ -25,12 +28,13 @@ class Controls {
 
         template<class Archive>
         void serialize(Archive& archive) {
-            archive(guiToRawMapping, psds);
+            archive(guiToRawMapping, psds, initialValues);
         }
 
     private:
         ConditionalTable guiToRawMapping;
         PSDMatrix psds;
+        Vector<float> initialValues;
         ControlsInputInstance::Factory instanceFactory;
 
 };
diff --git a/Engine/Plugins/Animation/RigLogic/Source/RigLogicLib/Private/riglogic/controls/ControlsFactory.cpp b/Engine/Plugins/Animation/RigLogic/Source/RigLogicLib/Private/riglogic/controls/ControlsFactory.cpp
--- a/Engine/Plugins/Animation/RigLogic/Source/RigLogicLib/Private/riglogic/controls/ControlsFactory.cpp
+++ b/Engine/Plugins/Animation/RigLogic/Source/RigLogicLib/Private/riglogic/controls/ControlsFactory.cpp
@@ -21,9 +21,15 @@ static ControlsInputInstance::Factory createInstanceFactory(const Configuration&
                                                             std::uint16_t psdControlCount,
                                                             std::uint16_t mlControlCount,
                                                             std::uint16_t rbfControlCount) {
-    return [ = ](MemoryResource* memRes) {
+    return [ = ](ConstArrayView<float> initialValues, MemoryResource* memRes) {
                auto factory = UniqueInstance<StandardControlsInputInstance, ControlsInputInstance>::with(memRes);
-               return factory.create(guiControlCount, rawControlCount, psdControlCount, mlControlCount, rbfControlCount, memRes);
+               return factory.create(guiControlCount,
+                                     rawControlCount,
+                                     psdControlCount,
+                                     mlControlCount,
+                                     rbfControlCount,
+                                     initialValues,
+                                     memRes);
     };
 }
 
@@ -34,7 +40,10 @@ Controls::Pointer ControlsFactory::create(const Configuration& config, const Rig
                                                  metrics.psdControlCount,
                                                  metrics.mlControlCount,
                                                  metrics.rbfControlCount);
-    return UniqueInstance<Controls>::with(memRes).create(ConditionalTable{memRes}, PSDMatrix{memRes}, instanceFactory);
+    return UniqueInstance<Controls>::with(memRes).create(ConditionalTable{memRes},
+                                                         PSDMatrix{memRes},
+                                                         Vector<float>{memRes},
+                                                         instanceFactory);
 }
 
 Controls::Pointer ControlsFactory::create(const Configuration& config, const dna::Reader* reader, MemoryResource* memRes) {
@@ -59,14 +68,20 @@ Controls::Pointer ControlsFactory::create(const Configuration& config, const dna
         }
     }
 
+    const auto guiControlCount = reader->getGUIControlCount();
+    const auto rawControlCount = reader->getRawControlCount();
+    const auto psdControlCount = reader->getPSDCount();
+    const auto mlControlCount = reader->getMLControlCount();
+    const auto rbfControlCount = reader->getRBFPoseControlCount();
+
     ConditionalTable conditionals{std::move(inputIndices),
                                   std::move(outputIndices),
                                   std::move(fromValues),
                                   std::move(toValues),
                                   std::move(slopeValues),
                                   std::move(cutValues),
-                                  reader->getGUIControlCount(),
-                                  reader->getRawControlCount(),
+                                  guiControlCount,
+                                  rawControlCount,
                                   memRes};
 
     Vector<std::uint16_t> psdRows{memRes};
@@ -75,19 +90,53 @@ Controls::Pointer ControlsFactory::create(const Configuration& config, const dna
     extd::copy(reader->getPSDRowIndices(), psdRows);
     extd::copy(reader->getPSDColumnIndices(), psdColumns);
     extd::copy(reader->getPSDValues(), psdValues);
-    PSDMatrix psds{reader->getPSDCount(),
+    PSDMatrix psds{psdControlCount,
                    std::move(psdRows),
                    std::move(psdColumns),
                    std::move(psdValues)};
 
+    auto endsWith = [](StringView str, StringView suffix) {
+            #if __cplusplus >= 201402L || (defined(_MSC_VER) && _MSC_VER >= 1900)
+                return (suffix.size() <= str.size()) && std::equal(extd::advanced(str.begin(), str.size() - suffix.size()),
+                                                                   str.end(),
+                                                                   suffix.begin(),
+                                                                   suffix.end());
+            #else
+                return (suffix.size() <= str.size()) && std::equal(extd::advanced(str.begin(), str.size() - suffix.size()),
+                                                                   str.end(),
+                                                                   suffix.begin());
+            #endif
+        };
+    auto makeSuffix = [](const char* str) {
+            return StringView{str, std::strlen(str)};
+        };
+    const auto qwSuffix = makeSuffix(".w");
+    const auto controlCount = (static_cast<std::size_t>(rawControlCount) +
+                               static_cast<std::size_t>(psdControlCount) +
+                               static_cast<std::size_t>(mlControlCount) +
+                               static_cast<std::size_t>(rbfControlCount));
+    Vector<float> initialValues{controlCount, {}, memRes};
+    for (std::uint16_t solverIndex = {}; solverIndex < reader->getRBFSolverCount(); ++solverIndex) {
+        const auto solverRawControlIndices = reader->getRBFSolverRawControlIndices(solverIndex);
+        for (const auto rawControlIndex : solverRawControlIndices) {
+            const auto rawControlName = reader->getRawControlName(rawControlIndex);
+            if (endsWith(rawControlName, qwSuffix)) {
+                initialValues[rawControlIndex] = 1.0f;
+            }
+        }
+    }
+
     auto instanceFactory = createInstanceFactory(config,
-                                                 conditionals.getInputCount(),
-                                                 conditionals.getOutputCount(),
-                                                 psds.getDistinctPSDCount(),
-                                                 reader->getMLControlCount(),
-                                                 reader->getRBFPoseControlCount());
+                                                 guiControlCount,
+                                                 rawControlCount,
+                                                 psdControlCount,
+                                                 mlControlCount,
+                                                 rbfControlCount);
 
-    return UniqueInstance<Controls>::with(memRes).create(std::move(conditionals), std::move(psds), instanceFactory);
+    return UniqueInstance<Controls>::with(memRes).create(std::move(conditionals),
+                                                         std::move(psds),
+                                                         std::move(initialValues),
+                                                         instanceFactory);
 }
 
 }  // namespace rl4
diff --git a/Engine/Plugins/Animation/RigLogic/Source/RigLogicLib/Private/riglogic/controls/ControlsInputInstance.h b/Engine/Plugins/Animation/RigLogic/Source/RigLogicLib/Private/riglogic/controls/ControlsInputInstance.h
--- a/Engine/Plugins/Animation/RigLogic/Source/RigLogicLib/Private/riglogic/controls/ControlsInputInstance.h
+++ b/Engine/Plugins/Animation/RigLogic/Source/RigLogicLib/Private/riglogic/controls/ControlsInputInstance.h
@@ -18,7 +18,7 @@ namespace rl4 {
 class ControlsInputInstance {
     public:
         using Pointer = UniqueInstance<ControlsInputInstance>::PointerType;
-        using Factory = std::function<Pointer(MemoryResource*)>;
+        using Factory = std::function<Pointer(ConstArrayView<float>, MemoryResource*)>;
 
     public:
         virtual ~ControlsInputInstance();
diff --git a/Engine/Plugins/Animation/RigLogic/Source/RigLogicLib/Private/riglogic/controls/instances/StandardControlsInputInstance.cpp b/Engine/Plugins/Animation/RigLogic/Source/RigLogicLib/Private/riglogic/controls/instances/StandardControlsInputInstance.cpp
--- a/Engine/Plugins/Animation/RigLogic/Source/RigLogicLib/Private/riglogic/controls/instances/StandardControlsInputInstance.cpp
+++ b/Engine/Plugins/Animation/RigLogic/Source/RigLogicLib/Private/riglogic/controls/instances/StandardControlsInputInstance.cpp
@@ -12,9 +12,10 @@ StandardControlsInputInstance::StandardControlsInputInstance(std::uint16_t guiCo
                                                              std::uint16_t psdControlCount_,
                                                              std::uint16_t mlControlCount_,
                                                              std::uint16_t rbfControlCount_,
+                                                             ConstArrayView<float> initialValues,
                                                              MemoryResource* memRes) :
     guiControlBuffer{guiControlCount_, {}, memRes},
-    inputBuffer{static_cast<std::size_t>(rawControlCount_ + psdControlCount_ + mlControlCount_ + rbfControlCount_), {}, memRes},
+    inputBuffer{initialValues.begin(), initialValues.end(), memRes},
     guiControlCount{guiControlCount_},
     rawControlCount{rawControlCount_},
     psdControlCount{psdControlCount_},
diff --git a/Engine/Plugins/Animation/RigLogic/Source/RigLogicLib/Private/riglogic/controls/instances/StandardControlsInputInstance.h b/Engine/Plugins/Animation/RigLogic/Source/RigLogicLib/Private/riglogic/controls/instances/StandardControlsInputInstance.h
--- a/Engine/Plugins/Animation/RigLogic/Source/RigLogicLib/Private/riglogic/controls/instances/StandardControlsInputInstance.h
+++ b/Engine/Plugins/Animation/RigLogic/Source/RigLogicLib/Private/riglogic/controls/instances/StandardControlsInputInstance.h
@@ -16,6 +16,7 @@ class StandardControlsInputInstance : public ControlsInputInstance {
                                       std::uint16_t psdControlCount_,
                                       std::uint16_t mlControlCount_,
                                       std::uint16_t rbfControlCount_,
+                                      ConstArrayView<float> initialValues,
                                       MemoryResource* memRes);
 
         ArrayView<float> getGUIControlBuffer() override;
diff --git a/Engine/Plugins/Animation/RigLogic/Source/RigLogicLib/Private/riglogic/rbf/cpu/CPURBFBehaviorEvaluator.h b/Engine/Plugins/Animation/RigLogic/Source/RigLogicLib/Private/riglogic/rbf/cpu/CPURBFBehaviorEvaluator.h
--- a/Engine/Plugins/Animation/RigLogic/Source/RigLogicLib/Private/riglogic/rbf/cpu/CPURBFBehaviorEvaluator.h
+++ b/Engine/Plugins/Animation/RigLogic/Source/RigLogicLib/Private/riglogic/rbf/cpu/CPURBFBehaviorEvaluator.h
@@ -33,6 +33,7 @@ class Evaluator : public RBFBehaviorEvaluator {
         Evaluator(LODSpec<std::uint16_t>&& lods_,
                   SolverVectorType&& solvers_,
                   Matrix<std::uint16_t>&& solverRawControlInputIndices_,
+                  Matrix<std::uint16_t>&& solverRawControlOutputIndices_,
                   Matrix<std::uint16_t>&& solverPoseIndices_,
                   Matrix<std::uint16_t>&& poseInputControlIndices_,
                   Matrix<std::uint16_t>&& poseOutputControlIndices_,
@@ -43,6 +44,7 @@ class Evaluator : public RBFBehaviorEvaluator {
             lods{std::move(lods_)},
             solvers{std::move(solvers_)},
             solverRawControlInputIndices{std::move(solverRawControlInputIndices_)},
+            solverRawControlOutputIndices{std::move(solverRawControlOutputIndices_)},
             solverPoseIndices{std::move(solverPoseIndices_)},
             poseInputControlIndices{std::move(poseInputControlIndices_)},
             poseOutputControlIndices{std::move(poseOutputControlIndices_)},
@@ -64,6 +66,7 @@ class Evaluator : public RBFBehaviorEvaluator {
 
             assert(solverIndex < solvers.size());
             const auto& rawControlInputIndices = solverRawControlInputIndices[solverIndex];
+            const auto& rawControlOutputIndices = solverRawControlOutputIndices[solverIndex];
             const auto rawControlInputCount = rawControlInputIndices.size();
             inputBuffer = inputBuffer.subview(0u, rawControlInputCount);
 
@@ -72,6 +75,10 @@ class Evaluator : public RBFBehaviorEvaluator {
             intermediateWeightsBuffer = intermediateWeightsBuffer.subview(0u, poseCount);
             outputWeightsBuffer = outputWeightsBuffer.subview(0u, poseCount);
 
+            for (const auto outputIndex : rawControlOutputIndices) {
+                rawControls[outputIndex] = 0.0f;
+            }
+
             for (std::uint16_t i = {}; i < rawControlInputCount; ++i) {
                 const std::uint16_t ri = rawControlInputIndices[i];
                 inputBuffer[i] = rawControls[ri];
@@ -102,10 +109,6 @@ class Evaluator : public RBFBehaviorEvaluator {
             assert(lod < lods.indicesPerLOD.size());
             const auto& solverIndices = lods.indicesPerLOD[lod];
             auto rawControls = inputs->getInputBuffer();
-            const auto offset = static_cast<std::size_t>(inputs->getRawControlCount()) +
-                static_cast<std::size_t>(inputs->getPSDControlCount()) +
-                static_cast<std::size_t>(inputs->getMLControlCount());
-            std::fill_n(rawControls.data() + offset, inputs->getRBFControlCount(), 0.0f);
             auto inputBuffer = static_cast<OutputInstance*>(intermediateOutputs)->getInputBuffer();
             auto intermediateWeightsBuffer = static_cast<OutputInstance*>(intermediateOutputs)->getIntermediateWeightsBuffer();
             auto outputWeightsBuffer = static_cast<OutputInstance*>(intermediateOutputs)->getOutputWeightsBuffer();
@@ -121,14 +124,8 @@ class Evaluator : public RBFBehaviorEvaluator {
                        std::uint16_t solverIndex) const override {
             assert(lod < lods.indicesPerLOD.size());
             static_cast<void>(lod);
-
             assert(solverIndex < solvers.size());
-
             auto rawControls = inputs->getInputBuffer();
-            const auto offset = static_cast<std::size_t>(inputs->getRawControlCount()) +
-                static_cast<std::size_t>(inputs->getPSDControlCount()) +
-                static_cast<std::size_t>(inputs->getMLControlCount());
-            std::fill_n(rawControls.data() + offset, inputs->getRBFControlCount(), 0.0f);
             auto inputBuffer = static_cast<OutputInstance*>(intermediateOutputs)->getInputBuffer();
             auto intermediateWeightsBuffer = static_cast<OutputInstance*>(intermediateOutputs)->getIntermediateWeightsBuffer();
             auto outputWeightsBuffer = static_cast<OutputInstance*>(intermediateOutputs)->getOutputWeightsBuffer();
@@ -154,6 +151,7 @@ class Evaluator : public RBFBehaviorEvaluator {
                 solvers[i]->load(archive);
             }
             archive(solverRawControlInputIndices);
+            archive(solverRawControlOutputIndices);
             archive(solverPoseIndices);
             archive(poseInputControlIndices);
             archive(poseOutputControlIndices);
@@ -172,6 +170,7 @@ class Evaluator : public RBFBehaviorEvaluator {
                 solver->save(archive);
             }
             archive(solverRawControlInputIndices);
+            archive(solverRawControlOutputIndices);
             archive(solverPoseIndices);
             archive(poseInputControlIndices);
             archive(poseOutputControlIndices);
@@ -184,6 +183,7 @@ class Evaluator : public RBFBehaviorEvaluator {
         LODSpec<std::uint16_t> lods;
         SolverVectorType solvers;
         Matrix<std::uint16_t> solverRawControlInputIndices;
+        Matrix<std::uint16_t> solverRawControlOutputIndices;
         Matrix<std::uint16_t> solverPoseIndices;
         Matrix<std::uint16_t> poseInputControlIndices;
         Matrix<std::uint16_t> poseOutputControlIndices;
diff --git a/Engine/Plugins/Animation/RigLogic/Source/RigLogicLib/Private/riglogic/rbf/cpu/CPURBFBehaviorFactory.h b/Engine/Plugins/Animation/RigLogic/Source/RigLogicLib/Private/riglogic/rbf/cpu/CPURBFBehaviorFactory.h
--- a/Engine/Plugins/Animation/RigLogic/Source/RigLogicLib/Private/riglogic/rbf/cpu/CPURBFBehaviorFactory.h
+++ b/Engine/Plugins/Animation/RigLogic/Source/RigLogicLib/Private/riglogic/rbf/cpu/CPURBFBehaviorFactory.h
@@ -28,6 +28,7 @@ class Factory {
             std::uint16_t maximumInputCount{};
             std::uint16_t maximumTargetCount{};
             Matrix<std::uint16_t> solverRawControlInputIndices{memRes};
+            Matrix<std::uint16_t> solverRawControlOutputIndices{memRes};
             Matrix<std::uint16_t> solverPoseIndices{memRes};
             Matrix<std::uint16_t> poseInputControlIndices{memRes};
             Matrix<std::uint16_t> poseOutputControlIndices{memRes};
@@ -43,6 +44,7 @@ class Factory {
                 return factory.create(LODSpec<std::uint16_t>{memRes},
                                       std::move(solvers),
                                       std::move(solverRawControlInputIndices),
+                                      std::move(solverRawControlOutputIndices),
                                       std::move(solverPoseIndices),
                                       std::move(poseInputControlIndices),
                                       std::move(poseOutputControlIndices),
@@ -55,6 +57,7 @@ class Factory {
             auto lods = computeLODs(reader, memRes);
             solvers.reserve(lods.count);
             solverRawControlInputIndices.resize(lods.count);
+            solverRawControlOutputIndices.resize(lods.count);
             solverPoseIndices.resize(lods.count);
 
             RBFSolverRecipe recipe{};
@@ -91,6 +94,16 @@ class Factory {
                 auto poseIndices = reader->getRBFSolverPoseIndices(solverIndex);
                 solverPoseIndices[solverIndex].assign(poseIndices.begin(), poseIndices.end());
 
+                auto& solverOutputIndices = solverRawControlOutputIndices[solverIndex];
+                solverOutputIndices.reserve(poseIndices.size());
+                for (const auto poseIndex : poseIndices) {
+                    const auto poseOutputIndices = reader->getRBFPoseOutputControlIndices(poseIndex);
+                    solverOutputIndices.insert(solverOutputIndices.end(), poseOutputIndices.begin(), poseOutputIndices.end());
+                }
+                std::sort(solverOutputIndices.begin(), solverOutputIndices.end());
+                solverOutputIndices.erase(std::unique(solverOutputIndices.begin(), solverOutputIndices.end()),
+                                          solverOutputIndices.end());
+
                 const auto targetCount = static_cast<std::uint16_t>(poseIndices.size());
                 if (targetCount > maximumTargetCount) {
                     maximumTargetCount = targetCount;
@@ -123,6 +136,7 @@ class Factory {
             return factory.create(std::move(lods),
                                   std::move(solvers),
                                   std::move(solverRawControlInputIndices),
+                                  std::move(solverRawControlOutputIndices),
                                   std::move(solverPoseIndices),
                                   std::move(poseInputControlIndices),
                                   std::move(poseOutputControlIndices),
diff --git a/Engine/Plugins/Animation/RigLogic/Source/RigLogicLib/Private/riglogic/system/simd/Detect.h b/Engine/Plugins/Animation/RigLogic/Source/RigLogicLib/Private/riglogic/system/simd/Detect.h
--- a/Engine/Plugins/Animation/RigLogic/Source/RigLogicLib/Private/riglogic/system/simd/Detect.h
+++ b/Engine/Plugins/Animation/RigLogic/Source/RigLogicLib/Private/riglogic/system/simd/Detect.h
@@ -9,7 +9,8 @@
     defined(__x86_64__) || defined(__x86_64) || defined(_M_X64) || defined(_M_AMD64)
         #define RL_BUILD_WITH_SSE 1
         #if defined(RL_AUTODETECT_HALF_FLOATS) && !defined(RL_USE_HALF_FLOATS)
-            #if defined(__F16C__) || (defined(_MSC_VER) && (_MSC_VER >= 1700) && defined(__AVX__)) || (defined(__INTEL_COMPILER) && defined(__AVX__))
+            #if defined(__F16C__) || (defined(_MSC_VER) && (_MSC_VER >= 1700) && defined(__AVX__)) || \
+            (defined(__INTEL_COMPILER) && defined(__AVX__))
                 #define RL_USE_HALF_FLOATS 1
             #endif
         #endif  // RL_AUTODETECT_HALF_FLOATS
diff --git a/Engine/Plugins/Animation/RigLogic/Source/RigLogicLib/Public/riglogic/version/Version.h b/Engine/Plugins/Animation/RigLogic/Source/RigLogicLib/Public/riglogic/version/Version.h
--- a/Engine/Plugins/Animation/RigLogic/Source/RigLogicLib/Public/riglogic/version/Version.h
+++ b/Engine/Plugins/Animation/RigLogic/Source/RigLogicLib/Public/riglogic/version/Version.h
@@ -4,5 +4,5 @@
 
 #define RL_MAJOR_VERSION 10
 #define RL_MINOR_VERSION 2
-#define RL_PATCH_VERSION 6
-#define RL_VERSION_STRING "10.2.6"
+#define RL_PATCH_VERSION 7
+#define RL_VERSION_STRING "10.2.7"
diff --git a/Engine/Plugins/Animation/RigLogic/Source/RigLogicLibTest/Private/rltests/animatedmaps/TestAnimatedMaps.cpp b/Engine/Plugins/Animation/RigLogic/Source/RigLogicLibTest/Private/rltests/animatedmaps/TestAnimatedMaps.cpp
--- a/Engine/Plugins/Animation/RigLogic/Source/RigLogicLibTest/Private/rltests/animatedmaps/TestAnimatedMaps.cpp
+++ b/Engine/Plugins/Animation/RigLogic/Source/RigLogicLibTest/Private/rltests/animatedmaps/TestAnimatedMaps.cpp
@@ -41,7 +41,8 @@ TEST_P(AnimatedMapsTest, LODLimitsCondTableSize) {
         {0.3f, 0.6f},  // LOD0
         {0.3f, 0.0f}  // LOD1
     };
-    auto inputInstance = inputInstanceFactory(&amr);
+    rl4::Vector<float> initialValues(conditionalTableInputs.size());
+    auto inputInstance = inputInstanceFactory(initialValues, &amr);
     auto inputBuffer = inputInstance->getInputBuffer();
     std::copy(conditionalTableInputs.begin(), conditionalTableInputs.end(), inputBuffer.begin());
 
diff --git a/Engine/Plugins/Animation/RigLogic/Source/RigLogicLibTest/Private/rltests/blendshapes/TestBlendShapes.cpp b/Engine/Plugins/Animation/RigLogic/Source/RigLogicLibTest/Private/rltests/blendshapes/TestBlendShapes.cpp
--- a/Engine/Plugins/Animation/RigLogic/Source/RigLogicLibTest/Private/rltests/blendshapes/TestBlendShapes.cpp
+++ b/Engine/Plugins/Animation/RigLogic/Source/RigLogicLibTest/Private/rltests/blendshapes/TestBlendShapes.cpp
@@ -37,7 +37,8 @@ TEST_P(BlendShapesTest, Calculate) {
                                                                     0,
                                                                     0,
                                                                     0);
-    auto inputInstance = inputInstanceFactory(&amr);
+    rl4::Vector<float> initialValues(blendShapeInputs.size());
+    auto inputInstance = inputInstanceFactory(initialValues, &amr);
     auto inputBuffer = inputInstance->getInputBuffer();
     std::copy(blendShapeInputs.begin(), blendShapeInputs.end(), inputBuffer.begin());
 
diff --git a/Engine/Plugins/Animation/RigLogic/Source/RigLogicLibTest/Private/rltests/controls/ControlFixtures.cpp b/Engine/Plugins/Animation/RigLogic/Source/RigLogicLibTest/Private/rltests/controls/ControlFixtures.cpp
--- a/Engine/Plugins/Animation/RigLogic/Source/RigLogicLibTest/Private/rltests/controls/ControlFixtures.cpp
+++ b/Engine/Plugins/Animation/RigLogic/Source/RigLogicLibTest/Private/rltests/controls/ControlFixtures.cpp
@@ -11,13 +11,14 @@ rl4::ControlsInputInstance::Factory ControlsFactory::getInstanceFactory(std::uin
                                                                         std::uint16_t psdControlCount,
                                                                         std::uint16_t mlControlCount,
                                                                         std::uint16_t rbfControlCount) {
-    return [ = ](rl4::MemoryResource* memRes) {
+    return [ = ](rl4::ConstArrayView<float> initialValues, rl4::MemoryResource* memRes) {
                return pma::UniqueInstance<rl4::StandardControlsInputInstance, rl4::ControlsInputInstance>::with(memRes).create(
                    guiControlCount,
                    rawControlCount,
                    psdControlCount,
                    mlControlCount,
                    rbfControlCount,
+                   initialValues,
                    memRes);
     };
 }
diff --git a/Engine/Plugins/Animation/RigLogic/Source/RigLogicLibTest/Private/rltests/controls/TestControls.cpp b/Engine/Plugins/Animation/RigLogic/Source/RigLogicLibTest/Private/rltests/controls/TestControls.cpp
--- a/Engine/Plugins/Animation/RigLogic/Source/RigLogicLibTest/Private/rltests/controls/TestControls.cpp
+++ b/Engine/Plugins/Animation/RigLogic/Source/RigLogicLibTest/Private/rltests/controls/TestControls.cpp
@@ -25,7 +25,9 @@ TEST(ControlsTest, GUIToRawMapping) {
     const std::uint16_t rawControlCount = conditionals.getOutputCount();
     const std::uint16_t psdControlCount = psds.getDistinctPSDCount();
     auto instanceFactory = ControlsFactory::getInstanceFactory(guiControlCount, rawControlCount, psdControlCount, 0u, 0u);
-    rl4::Controls controls{std::move(conditionals), std::move(psds), instanceFactory};
+
+    rl4::Vector<float> initialValues(static_cast<std::size_t>(rawControlCount) + static_cast<std::size_t>(psdControlCount));
+    rl4::Controls controls{std::move(conditionals), std::move(psds), std::move(initialValues), instanceFactory};
 
     const rl4::Vector<float> guiControls{0.1f, 0.2f};
     const rl4::Vector<float> expected{0.3f, 0.6f};
@@ -56,7 +58,8 @@ TEST(ControlsTest, PSDsAppendToOutput) {
     const std::uint16_t rawControlCount = conditionals.getOutputCount();
     const std::uint16_t psdControlCount = psds.getDistinctPSDCount();
     auto instanceFactory = ControlsFactory::getInstanceFactory(guiControlCount, rawControlCount, psdControlCount, 0u, 0u);
-    rl4::Controls controls{std::move(conditionals), std::move(psds), instanceFactory};
+    rl4::Vector<float> initialValues(static_cast<std::size_t>(rawControlCount) + static_cast<std::size_t>(psdControlCount));
+    rl4::Controls controls{std::move(conditionals), std::move(psds), std::move(initialValues), instanceFactory};
     auto instance = controls.createInstance(&amr);
     auto buffer = instance->getInputBuffer();
     std::copy(rawControls.begin(), rawControls.end(), buffer.begin());
diff --git a/Engine/Plugins/Animation/RigLogic/Source/RigLogicLibTest/Private/rltests/joints/bpcm/TestBlock4CalculationStrategy.cpp b/Engine/Plugins/Animation/RigLogic/Source/RigLogicLibTest/Private/rltests/joints/bpcm/TestBlock4CalculationStrategy.cpp
--- a/Engine/Plugins/Animation/RigLogic/Source/RigLogicLibTest/Private/rltests/joints/bpcm/TestBlock4CalculationStrategy.cpp
+++ b/Engine/Plugins/Animation/RigLogic/Source/RigLogicLibTest/Private/rltests/joints/bpcm/TestBlock4CalculationStrategy.cpp
@@ -56,7 +56,8 @@ class Block4JointCalculationStrategyTest : public ::testing::TestWithParam<Strat
 
             auto inputInstanceFactory =
                 ControlsFactory::getInstanceFactory(0, static_cast<std::uint16_t>(block4::input::values.size()), 0, 0, 0);
-            auto inputInstance = inputInstanceFactory(&memRes);
+            rl4::Vector<float> initialValues(block4::input::values.size());
+            auto inputInstance = inputInstanceFactory(initialValues, &memRes);
             auto inputBuffer = inputInstance->getInputBuffer();
             std::copy(block4::input::values.begin(), block4::input::values.end(), inputBuffer.begin());
 
diff --git a/Engine/Plugins/Animation/RigLogic/Source/RigLogicLibTest/Private/rltests/joints/bpcm/TestBlock8CalculationStrategy.cpp b/Engine/Plugins/Animation/RigLogic/Source/RigLogicLibTest/Private/rltests/joints/bpcm/TestBlock8CalculationStrategy.cpp
--- a/Engine/Plugins/Animation/RigLogic/Source/RigLogicLibTest/Private/rltests/joints/bpcm/TestBlock8CalculationStrategy.cpp
+++ b/Engine/Plugins/Animation/RigLogic/Source/RigLogicLibTest/Private/rltests/joints/bpcm/TestBlock8CalculationStrategy.cpp
@@ -55,7 +55,8 @@ class Block8JointCalculationStrategyTest : public ::testing::TestWithParam<Strat
 
             auto inputInstanceFactory =
                 ControlsFactory::getInstanceFactory(0, static_cast<std::uint16_t>(block8::input::values.size()), 0, 0, 0);
-            auto inputInstance = inputInstanceFactory(&memRes);
+            rl4::Vector<float> initialValues(block8::input::values.size());
+            auto inputInstance = inputInstanceFactory(initialValues, &memRes);
             auto inputBuffer = inputInstance->getInputBuffer();
             std::copy(block8::input::values.begin(), block8::input::values.end(), inputBuffer.begin());
 
diff --git a/Engine/Plugins/Animation/RigLogic/Source/RigLogicLibTest/Private/rltests/joints/quaternions/TestQuaternionEvaluator.cpp b/Engine/Plugins/Animation/RigLogic/Source/RigLogicLibTest/Private/rltests/joints/quaternions/TestQuaternionEvaluator.cpp
--- a/Engine/Plugins/Animation/RigLogic/Source/RigLogicLibTest/Private/rltests/joints/quaternions/TestQuaternionEvaluator.cpp
+++ b/Engine/Plugins/Animation/RigLogic/Source/RigLogicLibTest/Private/rltests/joints/quaternions/TestQuaternionEvaluator.cpp
@@ -150,7 +150,8 @@ TYPED_TEST(QuaternionEvaluatorTest, EvaluatePerLOD) {
                                                                     0,
                                                                     0,
                                                                     0);
-    auto inputInstance = inputInstanceFactory(&this->memRes);
+    rl4::Vector<float> initialValues(rltests::qs::input::values.size());
+    auto inputInstance = inputInstanceFactory(initialValues, &this->memRes);
     auto inputBuffer = inputInstance->getInputBuffer();
     std::copy(rltests::qs::input::values.begin(), rltests::qs::input::values.end(), inputBuffer.begin());
 
diff --git a/Engine/Plugins/Animation/RigLogic/Source/RigLogicLibTest/Private/rltests/joints/twistswing/TestTwistSwingEvaluator.cpp b/Engine/Plugins/Animation/RigLogic/Source/RigLogicLibTest/Private/rltests/joints/twistswing/TestTwistSwingEvaluator.cpp
--- a/Engine/Plugins/Animation/RigLogic/Source/RigLogicLibTest/Private/rltests/joints/twistswing/TestTwistSwingEvaluator.cpp
+++ b/Engine/Plugins/Animation/RigLogic/Source/RigLogicLibTest/Private/rltests/joints/twistswing/TestTwistSwingEvaluator.cpp
@@ -143,7 +143,8 @@ TYPED_TEST(TwistSwingEvaluatorTest, EvaluatePerLOD) {
                                                                     0,
                                                                     0,
                                                                     0);
-    auto inputInstance = inputInstanceFactory(&this->memRes);
+    rl4::Vector<float> initialValues(rltests::twsw::input::values.size());
+    auto inputInstance = inputInstanceFactory(initialValues, &this->memRes);
     auto inputBuffer = inputInstance->getInputBuffer();
     std::copy(rltests::twsw::input::values.begin(), rltests::twsw::input::values.end(), inputBuffer.begin());
 
diff --git a/Engine/Plugins/Animation/RigLogic/Source/RigLogicLibTest/Private/rltests/ml/cpu/TestInference.cpp b/Engine/Plugins/Animation/RigLogic/Source/RigLogicLibTest/Private/rltests/ml/cpu/TestInference.cpp
--- a/Engine/Plugins/Animation/RigLogic/Source/RigLogicLibTest/Private/rltests/ml/cpu/TestInference.cpp
+++ b/Engine/Plugins/Animation/RigLogic/Source/RigLogicLibTest/Private/rltests/ml/cpu/TestInference.cpp
@@ -96,7 +96,9 @@ TYPED_TEST(MLBInferenceTest, InferencePerLOD) {
                                                                     0,
                                                                     rltests::ml::block4::unoptimized::mlControlCount,
                                                                     0);
-    auto inputInstance = inputInstanceFactory(&this->memRes);
+    rl4::Vector<float> initialValues(static_cast<std::size_t>(rltests::ml::block4::unoptimized::rawControlCount) +
+                                     static_cast<std::size_t>(rltests::ml::block4::unoptimized::mlControlCount));
+    auto inputInstance = inputInstanceFactory(initialValues, &this->memRes);
     auto inputBuffer = inputInstance->getInputBuffer();
     auto outputBuffer = inputBuffer.subview(rltests::ml::block4::unoptimized::rawControlCount,
                                             rltests::ml::block4::unoptimized::mlControlCount);
diff --git a/Engine/Plugins/Animation/RigLogic/Source/RigLogicLibTest/Private/rltests/rbf/cpu/RBFFixtures.cpp b/Engine/Plugins/Animation/RigLogic/Source/RigLogicLibTest/Private/rltests/rbf/cpu/RBFFixtures.cpp
--- a/Engine/Plugins/Animation/RigLogic/Source/RigLogicLibTest/Private/rltests/rbf/cpu/RBFFixtures.cpp
+++ b/Engine/Plugins/Animation/RigLogic/Source/RigLogicLibTest/Private/rltests/rbf/cpu/RBFFixtures.cpp
@@ -82,10 +82,11 @@ const Matrix<std::uint16_t> poseInputControlIndices =
 {{}, {}, {}, {}, {}, {}, {}, {}, {}, {}, {}, {}, {}, {}, {1u}};
 
 const Matrix<std::uint16_t> poseOutputControlIndices =
-{{8u}, {9u}, {10u}, {11u}, {12u}, {13u}, {14u}, {15u}, {16u}, {17u}, {18u}, {19u}, {20u}, {21u}, {22u, 23u}};
+{{8u}, {9u}, {10u}, {11u}, {12u}, {13u}, {14u}, {15u}, {16u}, {17u}, {18u}, {19u}, {20u}, {21u, 22u}, {22u, 23u}};
 
 const Matrix<float> poseOutputControlWeights =
-{{1.0f}, {1.0f}, {1.0f}, {1.0f}, {1.0f}, {1.0f}, {1.0f}, {1.0f}, {1.0f}, {1.0f}, {1.0f}, {1.0f}, {1.0f}, {1.0f}, {0.5f, 0.5f}};
+{{1.0f}, {1.0f}, {1.0f}, {1.0f}, {1.0f}, {1.0f}, {1.0f}, {1.0f}, {1.0f}, {1.0f}, {1.0f}, {1.0f}, {1.0f}, {1.0f, 1.0f},
+    {0.5f, 0.5f}};
 
 const pma::Matrix<float> solverRawControlValues = {
     {
@@ -244,7 +245,7 @@ const Matrix<float> valuesPerLOD = {
         0.0198993f,
         0.0f,
         0.453696f,
-        0.09189290998017785f,
+        0.5455889099f,
         0.09189290998017785f,
     },
     {
diff --git a/Engine/Plugins/Animation/RigLogic/Source/RigLogicLibTest/Private/rltests/rbf/cpu/TestRBFEvaluator.cpp b/Engine/Plugins/Animation/RigLogic/Source/RigLogicLibTest/Private/rltests/rbf/cpu/TestRBFEvaluator.cpp
--- a/Engine/Plugins/Animation/RigLogic/Source/RigLogicLibTest/Private/rltests/rbf/cpu/TestRBFEvaluator.cpp
+++ b/Engine/Plugins/Animation/RigLogic/Source/RigLogicLibTest/Private/rltests/rbf/cpu/TestRBFEvaluator.cpp
@@ -94,7 +94,9 @@ TYPED_TEST(RBFBehaviorTest, SolverPerLOD) {
                                                                     0,
                                                                     0,
                                                                     rltests::rbf::unoptimized::poseControlCount);
-    auto inputInstance = inputInstanceFactory(&this->memRes);
+    rl4::Vector<float> initialValues(static_cast<std::size_t>(rltests::rbf::unoptimized::rawControlCount) +
+                                     static_cast<std::size_t>(rltests::rbf::unoptimized::poseControlCount));
+    auto inputInstance = inputInstanceFactory(initialValues, &this->memRes);
     auto inputBuffer = inputInstance->getInputBuffer();
     auto outputBuffer =
         inputBuffer.subview(rltests::rbf::unoptimized::rawControlCount, rltests::rbf::unoptimized::poseControlCount);
@@ -108,6 +110,8 @@ TYPED_TEST(RBFBehaviorTest, SolverPerLOD) {
     for (std::uint16_t lod = 0u; lod < rltests::rbf::unoptimized::lodCount; ++lod) {
         std::fill(outputBuffer.begin(), outputBuffer.end(), 0.0f);
         this->evaluator->calculate(inputInstance.get(), intermediateOutputs.get(), lod);
+        // Call twice to make sure the output control values are not accumulating between calls
+        this->evaluator->calculate(inputInstance.get(), intermediateOutputs.get(), lod);
         const auto& expected = rltests::rbf::output::valuesPerLOD[lod];
         #ifdef RL_USE_HALF_FLOATS
             static constexpr float threshold = 0.05f;
