[MH Plugin] Remove SkelMesh asset from plugin

diff --git a/Engine/Plugins/Animation/RigLogic/Source/RigLogicModule/Private/DNAUtils.cpp b/Engine/Plugins/Animation/RigLogic/Source/RigLogicModule/Private/DNAUtils.cpp
--- a/Engine/Plugins/Animation/RigLogic/Source/RigLogicModule/Private/DNAUtils.cpp
+++ b/Engine/Plugins/Animation/RigLogic/Source/RigLogicModule/Private/DNAUtils.cpp
@@ -3,8 +3,10 @@
 #include "DNAUtils.h"
 
 #include "CoreMinimal.h"
+#include "Misc/FileHelper.h"
 #include "HAL/LowLevelMemTracker.h"
 #include "DNACommon.h"
+#include "DNAAsset.h"
 #include "DNAReaderAdapter.h"
 #include "FMemoryResource.h"
 #include "RigLogicMemoryStream.h"
@@ -86,3 +88,20 @@ void WriteDNAToFile(const IDNAReader* Reader, EDNADataLayer Layer, const FString
 	DNAStreamWriter->setFrom(Reader->Unwrap(), CalculateDNADataLayerBitmask(Layer), dna::UnknownLayerPolicy::Preserve, FMemoryResource::Instance());
 	DNAStreamWriter->write();
 }
+
+TObjectPtr<class UDNAAsset> GetDNAAssetFromFile(const FString& InFilePath, UObject* InOuter, EDNADataLayer InLayer)
+{
+	UDNAAsset* DNAAsset = nullptr;
+	TArray<uint8> DNADataAsBuffer;
+	if (FFileHelper::LoadFileToArray(DNADataAsBuffer, *InFilePath))
+	{
+		if (TSharedPtr<IDNAReader> DNAReader = ReadDNAFromBuffer(&DNADataAsBuffer, InLayer))
+		{
+			DNAAsset = NewObject<UDNAAsset>(InOuter);
+			DNAAsset->SetBehaviorReader(DNAReader);
+			DNAAsset->SetGeometryReader(DNAReader);
+		}
+	}
+
+	return DNAAsset;
+}
diff --git a/Engine/Plugins/Animation/RigLogic/Source/RigLogicModule/Public/DNAUtils.h b/Engine/Plugins/Animation/RigLogic/Source/RigLogicModule/Public/DNAUtils.h
--- a/Engine/Plugins/Animation/RigLogic/Source/RigLogicModule/Public/DNAUtils.h
+++ b/Engine/Plugins/Animation/RigLogic/Source/RigLogicModule/Public/DNAUtils.h
@@ -22,3 +22,6 @@ RIGLOGICMODULE_API TSharedPtr<IDNAReader> ReadDNAFromBuffer(TArray<uint8>* DNABu
 RIGLOGICMODULE_API TArray<uint8> ReadStreamFromDNA(const IDNAReader* Reader, EDNADataLayer Layer);
 RIGLOGICMODULE_API void WriteDNAToFile(const IDNAReader* Reader, EDNADataLayer Layer, const FString& Path);
 RIGLOGICMODULE_API dna::DataLayer CalculateDNADataLayerBitmask(EDNADataLayer Layer);
+
+/** Creates a new DNA asset with behavior/geometry data set. FileHelper is used to load dna data from specified file */
+RIGLOGICMODULE_API TObjectPtr<class UDNAAsset> GetDNAAssetFromFile(const FString& InFilePath, UObject* InOuter, EDNADataLayer InLayer = EDNADataLayer::All);
\ No newline at end of file
