Enable runtime CPU feature detection in RigLogic

diff --git a/Engine/Platforms/Android/Plugins/RigLogic/Source/RigLogicLib/RigLogicLib_Android.Build.cs b/Engine/Platforms/Android/Plugins/RigLogic/Source/RigLogicLib/RigLogicLib_Android.Build.cs
--- a/Engine/Platforms/Android/Plugins/RigLogic/Source/RigLogicLib/RigLogicLib_Android.Build.cs
+++ b/Engine/Platforms/Android/Plugins/RigLogic/Source/RigLogicLib/RigLogicLib_Android.Build.cs
@@ -5,9 +5,9 @@ namespace UnrealBuildTool.Rules
 	{
 		public RigLogicLib_Android(ReadOnlyTargetRules Target) : base(Target)
 		{
-			PrivateDefinitions.Add("RL_AUTODETECT_SSE=1");
-			PrivateDefinitions.Add("RL_AUTODETECT_NEON=1");
-			PrivateDefinitions.Add("RL_AUTODETECT_HALF_FLOATS=1");
+			PrivateDefinitions.Add("RL_AUTODETECT_SSE");
+			PrivateDefinitions.Add("RL_AUTODETECT_NEON");
+			PrivateDefinitions.Add("RL_AUTODETECT_HALF_FLOATS");
 		}
 	}
 }
\ No newline at end of file
diff --git a/Engine/Platforms/IOS/Plugins/RigLogic/Source/RigLogicLib/RigLogicLib_IOS.Build.cs b/Engine/Platforms/IOS/Plugins/RigLogic/Source/RigLogicLib/RigLogicLib_IOS.Build.cs
--- a/Engine/Platforms/IOS/Plugins/RigLogic/Source/RigLogicLib/RigLogicLib_IOS.Build.cs
+++ b/Engine/Platforms/IOS/Plugins/RigLogic/Source/RigLogicLib/RigLogicLib_IOS.Build.cs
@@ -5,8 +5,8 @@ namespace UnrealBuildTool.Rules
 	{
 		public RigLogicLib_IOS(ReadOnlyTargetRules Target) : base(Target)
 		{
-			PrivateDefinitions.Add("RL_AUTODETECT_NEON=1");
-			PrivateDefinitions.Add("RL_AUTODETECT_HALF_FLOATS=1");
+			PrivateDefinitions.Add("RL_AUTODETECT_NEON");
+			PrivateDefinitions.Add("RL_AUTODETECT_HALF_FLOATS");
 		}
 	}
 }
\ No newline at end of file
diff --git a/Engine/Plugins/Animation/RigLogic/Source/RigLogicLib/Private/riglogic/joints/cpu/bpcm/BPCMJointsBuilderFactory.cpp b/Engine/Plugins/Animation/RigLogic/Source/RigLogicLib/Private/riglogic/joints/cpu/bpcm/BPCMJointsBuilderFactory.cpp
--- a/Engine/Plugins/Animation/RigLogic/Source/RigLogicLib/Private/riglogic/joints/cpu/bpcm/BPCMJointsBuilderFactory.cpp
+++ b/Engine/Plugins/Animation/RigLogic/Source/RigLogicLib/Private/riglogic/joints/cpu/bpcm/BPCMJointsBuilderFactory.cpp
@@ -8,41 +8,62 @@
 namespace rl4 {
 
 UniqueInstance<JointsBuilder>::PointerType BPCMJointsBuilderFactory::create(const Configuration& config, MemoryResource* memRes) {
-    #if defined(__clang__)
-        #pragma clang diagnostic push
-        #pragma clang diagnostic ignored "-Wunused-local-typedef"
-    #endif
-    #if !defined(__clang__) && defined(__GNUC__)
-        #pragma GCC diagnostic push
-        #pragma GCC diagnostic ignored "-Wunused-local-typedefs"
-    #endif
-    #ifdef RL_USE_HALF_FLOATS
-        using StorageType = std::uint16_t;
-    #else
-        using StorageType = float;
-    #endif  // RL_USE_HALF_FLOATS
-    #if !defined(__clang__) && defined(__GNUC__)
-        #pragma GCC diagnostic pop
-    #endif
-    #if defined(__clang__)
-        #pragma clang diagnostic pop
-    #endif
-
+    auto features = trimd::getCPUFeatures();
+    RL_UNUSED(features);
     #ifdef RL_BUILD_WITH_SSE
-        if ((config.calculationType == CalculationType::SSE) || (config.calculationType == CalculationType::AnyVector)) {
-            using SSEBPCMJointsBuilder = bpcm::BPCMJointsBuilder<StorageType, trimd::sse::F128>;
+        #ifdef RL_DISABLE_RUNTIME_FEATURE_DETECTION
+            features.SSE2 = true;
+        #endif  // RL_DISABLE_RUNTIME_FEATURE_DETECTION
+        if (features.SSE2 &&
+            ((config.calculationType == CalculationType::SSE) || (config.calculationType == CalculationType::AnyVector))) {
+            #ifdef RL_BUILD_WITH_HALF_FLOATS
+                #ifdef RL_DISABLE_RUNTIME_FEATURE_DETECTION
+                    features.F16C = true;
+                #endif  // RL_DISABLE_RUNTIME_FEATURE_DETECTION
+                if (features.F16C) {
+                    using SSEBPCMJointsBuilder = bpcm::BPCMJointsBuilder<std::uint16_t, trimd::sse::F128>;
+                    return UniqueInstance<SSEBPCMJointsBuilder, JointsBuilder>::with(memRes).create(config, memRes);
+                }
+            #endif  // RL_BUILD_WITH_HALF_FLOATS
+            using SSEBPCMJointsBuilder = bpcm::BPCMJointsBuilder<float, trimd::sse::F128>;
             return UniqueInstance<SSEBPCMJointsBuilder, JointsBuilder>::with(memRes).create(config, memRes);
         }
     #endif  // RL_BUILD_WITH_SSE
     #ifdef RL_BUILD_WITH_AVX
-        if ((config.calculationType == CalculationType::AVX) || (config.calculationType == CalculationType::AnyVector)) {
-            using AVXBPCMJointsBuilder = bpcm::BPCMJointsBuilder<StorageType, trimd::avx::F256>;
+        #ifdef RL_DISABLE_RUNTIME_FEATURE_DETECTION
+            features.AVX = true;
+        #endif  // RL_DISABLE_RUNTIME_FEATURE_DETECTION
+        if (features.AVX &&
+            ((config.calculationType == CalculationType::AVX) || (config.calculationType == CalculationType::AnyVector))) {
+            #ifdef RL_BUILD_WITH_HALF_FLOATS
+                #ifdef RL_DISABLE_RUNTIME_FEATURE_DETECTION
+                    features.F16C = true;
+                #endif  // RL_DISABLE_RUNTIME_FEATURE_DETECTION
+                if (features.F16C) {
+                    using AVXBPCMJointsBuilder = bpcm::BPCMJointsBuilder<std::uint16_t, trimd::avx::F256>;
+                    return UniqueInstance<AVXBPCMJointsBuilder, JointsBuilder>::with(memRes).create(config, memRes);
+                }
+            #endif  // RL_BUILD_WITH_HALF_FLOATS
+            using AVXBPCMJointsBuilder = bpcm::BPCMJointsBuilder<float, trimd::avx::F256>;
             return UniqueInstance<AVXBPCMJointsBuilder, JointsBuilder>::with(memRes).create(config, memRes);
         }
     #endif  // RL_BUILD_WITH_AVX
     #ifdef RL_BUILD_WITH_NEON
-        if ((config.calculationType == CalculationType::NEON) || (config.calculationType == CalculationType::AnyVector)) {
-            using NEONBPCMJointsBuilder = bpcm::BPCMJointsBuilder<StorageType, trimd::neon::F128>;
+        #ifdef RL_DISABLE_RUNTIME_FEATURE_DETECTION
+            features.NEON = true;
+        #endif  // RL_DISABLE_RUNTIME_FEATURE_DETECTION
+        if (features.NEON &&
+            ((config.calculationType == CalculationType::NEON) || (config.calculationType == CalculationType::AnyVector))) {
+            #ifdef RL_BUILD_WITH_HALF_FLOATS
+                #ifdef RL_DISABLE_RUNTIME_FEATURE_DETECTION
+                    features.FP16 = true;
+                #endif  // RL_DISABLE_RUNTIME_FEATURE_DETECTION
+                if (features.FP16) {
+                    using NEONBPCMJointsBuilder = bpcm::BPCMJointsBuilder<std::uint16_t, trimd::neon::F128>;
+                    return UniqueInstance<NEONBPCMJointsBuilder, JointsBuilder>::with(memRes).create(config, memRes);
+                }
+            #endif  // RL_BUILD_WITH_HALF_FLOATS
+            using NEONBPCMJointsBuilder = bpcm::BPCMJointsBuilder<float, trimd::neon::F128>;
             return UniqueInstance<NEONBPCMJointsBuilder, JointsBuilder>::with(memRes).create(config, memRes);
         }
     #endif  // RL_BUILD_WITH_NEON
diff --git a/Engine/Plugins/Animation/RigLogic/Source/RigLogicLib/Private/riglogic/joints/cpu/quaternions/QuaternionJointsBuilderFactory.cpp b/Engine/Plugins/Animation/RigLogic/Source/RigLogicLib/Private/riglogic/joints/cpu/quaternions/QuaternionJointsBuilderFactory.cpp
--- a/Engine/Plugins/Animation/RigLogic/Source/RigLogicLib/Private/riglogic/joints/cpu/quaternions/QuaternionJointsBuilderFactory.cpp
+++ b/Engine/Plugins/Animation/RigLogic/Source/RigLogicLib/Private/riglogic/joints/cpu/quaternions/QuaternionJointsBuilderFactory.cpp
@@ -9,41 +9,63 @@ namespace rl4 {
 
 UniqueInstance<JointsBuilder>::PointerType QuaternionJointsBuilderFactory::create(const Configuration& config,
                                                                                   MemoryResource* memRes) {
-    #if defined(__clang__)
-        #pragma clang diagnostic push
-        #pragma clang diagnostic ignored "-Wunused-local-typedef"
-    #endif
-    #if !defined(__clang__) && defined(__GNUC__)
-        #pragma GCC diagnostic push
-        #pragma GCC diagnostic ignored "-Wunused-local-typedefs"
-    #endif
-    #ifdef RL_USE_HALF_FLOATS
-        using StorageType = std::uint16_t;
-    #else
-        using StorageType = float;
-    #endif  // RL_USE_HALF_FLOATS
-    #if !defined(__clang__) && defined(__GNUC__)
-        #pragma GCC diagnostic pop
-    #endif
-    #if defined(__clang__)
-        #pragma clang diagnostic pop
-    #endif
-
+    auto features = trimd::getCPUFeatures();
+    RL_UNUSED(features);
     #ifdef RL_BUILD_WITH_SSE
-        if ((config.calculationType == CalculationType::SSE) || (config.calculationType == CalculationType::AnyVector)) {
-            using SSEQuaternionJointsBuilder = QuaternionJointsBuilder<StorageType, trimd::sse::F256, trimd::sse::F128>;
+        #ifdef RL_DISABLE_RUNTIME_FEATURE_DETECTION
+            features.SSE2 = true;
+        #endif  // RL_DISABLE_RUNTIME_FEATURE_DETECTION
+        if (features.SSE2 &&
+            ((config.calculationType == CalculationType::SSE) || (config.calculationType == CalculationType::AnyVector))) {
+            #ifdef RL_BUILD_WITH_HALF_FLOATS
+                #ifdef RL_DISABLE_RUNTIME_FEATURE_DETECTION
+                    features.F16C = true;
+                #endif  // RL_DISABLE_RUNTIME_FEATURE_DETECTION
+                if (features.F16C) {
+                    using SSEQuaternionJointsBuilder = QuaternionJointsBuilder<std::uint16_t, trimd::sse::F256, trimd::sse::F128>;
+                    return UniqueInstance<SSEQuaternionJointsBuilder, JointsBuilder>::with(memRes).create(config, memRes);
+                }
+            #endif  // RL_BUILD_WITH_HALF_FLOATS
+            using SSEQuaternionJointsBuilder = QuaternionJointsBuilder<float, trimd::sse::F256, trimd::sse::F128>;
             return UniqueInstance<SSEQuaternionJointsBuilder, JointsBuilder>::with(memRes).create(config, memRes);
         }
     #endif  // RL_BUILD_WITH_SSE
     #ifdef RL_BUILD_WITH_AVX
-        if ((config.calculationType == CalculationType::AVX) || (config.calculationType == CalculationType::AnyVector)) {
-            using AVXQuaternionJointsBuilder = QuaternionJointsBuilder<StorageType, trimd::avx::F256, trimd::sse::F128>;
+        #ifdef RL_DISABLE_RUNTIME_FEATURE_DETECTION
+            features.AVX = true;
+        #endif  // RL_DISABLE_RUNTIME_FEATURE_DETECTION
+        if (features.AVX &&
+            ((config.calculationType == CalculationType::AVX) || (config.calculationType == CalculationType::AnyVector))) {
+            #ifdef RL_BUILD_WITH_HALF_FLOATS
+                #ifdef RL_DISABLE_RUNTIME_FEATURE_DETECTION
+                    features.F16C = true;
+                #endif  // RL_DISABLE_RUNTIME_FEATURE_DETECTION
+                if (features.F16C) {
+                    using AVXQuaternionJointsBuilder = QuaternionJointsBuilder<std::uint16_t, trimd::avx::F256, trimd::sse::F128>;
+                    return UniqueInstance<AVXQuaternionJointsBuilder, JointsBuilder>::with(memRes).create(config, memRes);
+                }
+            #endif  // RL_BUILD_WITH_HALF_FLOATS
+            using AVXQuaternionJointsBuilder = QuaternionJointsBuilder<float, trimd::avx::F256, trimd::sse::F128>;
             return UniqueInstance<AVXQuaternionJointsBuilder, JointsBuilder>::with(memRes).create(config, memRes);
         }
     #endif  // RL_BUILD_WITH_AVX
     #ifdef RL_BUILD_WITH_NEON
-        if ((config.calculationType == CalculationType::NEON) || (config.calculationType == CalculationType::AnyVector)) {
-            using NEONQuaternionJointsBuilder = QuaternionJointsBuilder<StorageType, trimd::neon::F256, trimd::neon::F128>;
+        #ifdef RL_DISABLE_RUNTIME_FEATURE_DETECTION
+            features.NEON = true;
+        #endif  // RL_DISABLE_RUNTIME_FEATURE_DETECTION
+        if (features.NEON &&
+            ((config.calculationType == CalculationType::NEON) || (config.calculationType == CalculationType::AnyVector))) {
+            #ifdef RL_BUILD_WITH_HALF_FLOATS
+                #ifdef RL_DISABLE_RUNTIME_FEATURE_DETECTION
+                    features.FP16 = true;
+                #endif  // RL_DISABLE_RUNTIME_FEATURE_DETECTION
+                if (features.FP16) {
+                    using NEONQuaternionJointsBuilder = QuaternionJointsBuilder<std::uint16_t, trimd::neon::F256,
+                                                                                trimd::neon::F128>;
+                    return UniqueInstance<NEONQuaternionJointsBuilder, JointsBuilder>::with(memRes).create(config, memRes);
+                }
+            #endif  // RL_BUILD_WITH_HALF_FLOATS
+            using NEONQuaternionJointsBuilder = QuaternionJointsBuilder<float, trimd::neon::F256, trimd::neon::F128>;
             return UniqueInstance<NEONQuaternionJointsBuilder, JointsBuilder>::with(memRes).create(config, memRes);
         }
     #endif  // RL_BUILD_WITH_NEON
diff --git a/Engine/Plugins/Animation/RigLogic/Source/RigLogicLib/Private/riglogic/joints/cpu/twistswing/TwistSwingJointsBuilderFactory.cpp b/Engine/Plugins/Animation/RigLogic/Source/RigLogicLib/Private/riglogic/joints/cpu/twistswing/TwistSwingJointsBuilderFactory.cpp
--- a/Engine/Plugins/Animation/RigLogic/Source/RigLogicLib/Private/riglogic/joints/cpu/twistswing/TwistSwingJointsBuilderFactory.cpp
+++ b/Engine/Plugins/Animation/RigLogic/Source/RigLogicLib/Private/riglogic/joints/cpu/twistswing/TwistSwingJointsBuilderFactory.cpp
@@ -12,41 +12,59 @@ UniqueInstance<JointsBuilder>::PointerType TwistSwingJointsBuilderFactory::creat
     // At the moment, only a Scalar implementation exists for twist and swing evaluation, so there is no need to
     // compile any other configuration.
     /*
-    #if defined(__clang__)
-        #pragma clang diagnostic push
-        #pragma clang diagnostic ignored "-Wunused-local-typedef"
-    #endif
-    #if !defined(__clang__) && defined(__GNUC__)
-        #pragma GCC diagnostic push
-        #pragma GCC diagnostic ignored "-Wunused-local-typedefs"
-    #endif
-    #ifdef RL_USE_HALF_FLOATS
-        using StorageType = std::uint16_t;
-    #else
-        using StorageType = float;
-    #endif  // RL_USE_HALF_FLOATS
-    #if !defined(__clang__) && defined(__GNUC__)
-        #pragma GCC diagnostic pop
-    #endif
-    #if defined(__clang__)
-        #pragma clang diagnostic pop
-    #endif
-
+    auto features = trimd::getCPUFeatures();
+    RL_UNUSED(features);
     #ifdef RL_BUILD_WITH_SSE
-        if ((config.calculationType == CalculationType::SSE) || (config.calculationType == CalculationType::AnyVector)) {
-            using SSETwistSwingJointsBuilder = TwistSwingJointsBuilder<StorageType, trimd::sse::F256, trimd::sse::F128>;
+        #ifdef RL_DISABLE_RUNTIME_FEATURE_DETECTION
+        features.SSE2 = true;
+        #endif  // RL_DISABLE_RUNTIME_FEATURE_DETECTION
+        if (features.SSE2 && ((config.calculationType == CalculationType::SSE) || (config.calculationType == CalculationType::AnyVector))) {
+            #ifdef RL_BUILD_WITH_HALF_FLOATS
+                #ifdef RL_DISABLE_RUNTIME_FEATURE_DETECTION
+                features.F16C = true;
+                #endif  // RL_DISABLE_RUNTIME_FEATURE_DETECTION
+                if (features.F16C) {
+                    using SSETwistSwingJointsBuilder = TwistSwingJointsBuilder<std::uint16_t, trimd::sse::F256, trimd::sse::F128>;
+                    return UniqueInstance<SSETwistSwingJointsBuilder, JointsBuilder>::with(memRes).create(config, memRes);
+                }
+            #endif  // RL_BUILD_WITH_HALF_FLOATS
+            using SSETwistSwingJointsBuilder = TwistSwingJointsBuilder<float, trimd::sse::F256, trimd::sse::F128>;
             return UniqueInstance<SSETwistSwingJointsBuilder, JointsBuilder>::with(memRes).create(config, memRes);
         }
     #endif  // RL_BUILD_WITH_SSE
     #ifdef RL_BUILD_WITH_AVX
-        if ((config.calculationType == CalculationType::AVX) || (config.calculationType == CalculationType::AnyVector)) {
-            using AVXTwistSwingJointsBuilder = TwistSwingJointsBuilder<StorageType, trimd::avx::F256, trimd::sse::F128>;
+        #ifdef RL_DISABLE_RUNTIME_FEATURE_DETECTION
+        features.AVX = true;
+        #endif  // RL_DISABLE_RUNTIME_FEATURE_DETECTION
+        if (features.AVX && ((config.calculationType == CalculationType::AVX) || (config.calculationType == CalculationType::AnyVector))) {
+            #ifdef RL_BUILD_WITH_HALF_FLOATS
+                #ifdef RL_DISABLE_RUNTIME_FEATURE_DETECTION
+                features.F16C = true;
+                #endif  // RL_DISABLE_RUNTIME_FEATURE_DETECTION
+                if (features.F16C) {
+                    using AVXTwistSwingJointsBuilder = TwistSwingJointsBuilder<std::uint16_t, trimd::avx::F256, trimd::sse::F128>;
+                    return UniqueInstance<AVXTwistSwingJointsBuilder, JointsBuilder>::with(memRes).create(config, memRes);
+                }
+            #endif  // RL_BUILD_WITH_HALF_FLOATS
+            using AVXTwistSwingJointsBuilder = TwistSwingJointsBuilder<float, trimd::avx::F256, trimd::sse::F128>;
             return UniqueInstance<AVXTwistSwingJointsBuilder, JointsBuilder>::with(memRes).create(config, memRes);
         }
     #endif  // RL_BUILD_WITH_AVX
     #ifdef RL_BUILD_WITH_NEON
-        if ((config.calculationType == CalculationType::NEON) || (config.calculationType == CalculationType::AnyVector)) {
-            using NEONTwistSwingJointsBuilder = TwistSwingJointsBuilder<StorageType, trimd::neon::F256, trimd::neon::F128>;
+        #ifdef RL_DISABLE_RUNTIME_FEATURE_DETECTION
+        features.NEON = true;
+        #endif  // RL_DISABLE_RUNTIME_FEATURE_DETECTION
+        if (features.NEON && ((config.calculationType == CalculationType::NEON) || (config.calculationType == CalculationType::AnyVector))) {
+            #ifdef RL_BUILD_WITH_HALF_FLOATS
+                #ifdef RL_DISABLE_RUNTIME_FEATURE_DETECTION
+                features.FP16 = true;
+                #endif  // RL_DISABLE_RUNTIME_FEATURE_DETECTION
+                if (features.FP16) {
+                    using NEONTwistSwingJointsBuilder = TwistSwingJointsBuilder<std::uint16_t, trimd::neon::F256, trimd::neon::F128>;
+                    return UniqueInstance<NEONTwistSwingJointsBuilder, JointsBuilder>::with(memRes).create(config, memRes);
+                }
+            #endif  // RL_BUILD_WITH_HALF_FLOATS
+            using NEONTwistSwingJointsBuilder = TwistSwingJointsBuilder<float, trimd::neon::F256, trimd::neon::F128>;
             return UniqueInstance<NEONTwistSwingJointsBuilder, JointsBuilder>::with(memRes).create(config, memRes);
         }
     #endif  // RL_BUILD_WITH_NEON
diff --git a/Engine/Plugins/Animation/RigLogic/Source/RigLogicLib/Private/riglogic/ml/MachineLearnedBehaviorFactory.cpp b/Engine/Plugins/Animation/RigLogic/Source/RigLogicLib/Private/riglogic/ml/MachineLearnedBehaviorFactory.cpp
--- a/Engine/Plugins/Animation/RigLogic/Source/RigLogicLib/Private/riglogic/ml/MachineLearnedBehaviorFactory.cpp
+++ b/Engine/Plugins/Animation/RigLogic/Source/RigLogicLib/Private/riglogic/ml/MachineLearnedBehaviorFactory.cpp
@@ -18,42 +18,59 @@ namespace rl4 {
 MachineLearnedBehaviorEvaluator::Pointer createMLEvaluator(const Configuration& config,
                                                            const dna::MachineLearnedBehaviorReader* reader,
                                                            MemoryResource* memRes) {
-    #if defined(__clang__)
-        #pragma clang diagnostic push
-        #pragma clang diagnostic ignored "-Wunused-local-typedef"
-    #endif
-    #if !defined(__clang__) && defined(__GNUC__)
-        #pragma GCC diagnostic push
-        #pragma GCC diagnostic ignored "-Wunused-local-typedefs"
-    #endif
-    #ifdef RL_USE_HALF_FLOATS
-        using StorageValueType = std::uint16_t;
-    #else
-        using StorageValueType = float;
-    #endif  // RL_USE_HALF_FLOATS
-    #if !defined(__clang__) && defined(__GNUC__)
-        #pragma GCC diagnostic pop
-    #endif
-    #if defined(__clang__)
-        #pragma clang diagnostic pop
-    #endif
-
-    // Work around unused parameter warning when building without SSE and AVX
-    static_cast<void>(config);
+    auto features = trimd::getCPUFeatures();
+    RL_UNUSED(features);
+    RL_UNUSED(config);
     #ifdef RL_BUILD_WITH_SSE
-        if ((config.calculationType == CalculationType::SSE) || (config.calculationType == CalculationType::AnyVector)) {
-            return ml::cpu::Factory<StorageValueType, trimd::sse::F256, trimd::sse::F128>::create(reader, memRes);
+        #ifdef RL_DISABLE_RUNTIME_FEATURE_DETECTION
+            features.SSE2 = true;
+        #endif  // RL_DISABLE_RUNTIME_FEATURE_DETECTION
+        if (features.SSE2 &&
+            ((config.calculationType == CalculationType::SSE) || (config.calculationType == CalculationType::AnyVector))) {
+            #ifdef RL_BUILD_WITH_HALF_FLOATS
+                #ifdef RL_DISABLE_RUNTIME_FEATURE_DETECTION
+                    features.F16C = true;
+                #endif  // RL_DISABLE_RUNTIME_FEATURE_DETECTION
+                if (features.F16C) {
+                    return ml::cpu::Factory<std::uint16_t, trimd::sse::F256, trimd::sse::F128>::create(reader, memRes);
+                }
+            #endif  // RL_BUILD_WITH_HALF_FLOATS
+            return ml::cpu::Factory<float, trimd::sse::F256, trimd::sse::F128>::create(reader, memRes);
         }
     #endif  // RL_BUILD_WITH_SSE
     #ifdef RL_BUILD_WITH_AVX
-        if ((config.calculationType == CalculationType::AVX) || (config.calculationType == CalculationType::AnyVector)) {
+        #ifdef RL_DISABLE_RUNTIME_FEATURE_DETECTION
+            features.AVX = true;
+        #endif  // RL_DISABLE_RUNTIME_FEATURE_DETECTION
+        if (features.AVX &&
+            ((config.calculationType == CalculationType::AVX) || (config.calculationType == CalculationType::AnyVector))) {
             // Use 256-bit AVX registers and whatever 128-bit width type is available
-            return ml::cpu::Factory<StorageValueType, trimd::avx::F256, trimd::sse::F128>::create(reader, memRes);
+            #ifdef RL_BUILD_WITH_HALF_FLOATS
+                #ifdef RL_DISABLE_RUNTIME_FEATURE_DETECTION
+                    features.F16C = true;
+                #endif  // RL_DISABLE_RUNTIME_FEATURE_DETECTION
+                if (features.F16C) {
+                    return ml::cpu::Factory<std::uint16_t, trimd::avx::F256, trimd::sse::F128>::create(reader, memRes);
+                }
+            #endif  // RL_BUILD_WITH_HALF_FLOATS
+            return ml::cpu::Factory<float, trimd::avx::F256, trimd::sse::F128>::create(reader, memRes);
         }
     #endif  // RL_BUILD_WITH_AVX
     #ifdef RL_BUILD_WITH_NEON
-        if ((config.calculationType == CalculationType::NEON) || (config.calculationType == CalculationType::AnyVector)) {
-            return ml::cpu::Factory<StorageValueType, trimd::neon::F256, trimd::neon::F128>::create(reader, memRes);
+        #ifdef RL_DISABLE_RUNTIME_FEATURE_DETECTION
+            features.NEON = true;
+        #endif  // RL_DISABLE_RUNTIME_FEATURE_DETECTION
+        if (features.NEON &&
+            ((config.calculationType == CalculationType::NEON) || (config.calculationType == CalculationType::AnyVector))) {
+            #ifdef RL_BUILD_WITH_HALF_FLOATS
+                #ifdef RL_DISABLE_RUNTIME_FEATURE_DETECTION
+                    features.FP16 = true;
+                #endif  // RL_DISABLE_RUNTIME_FEATURE_DETECTION
+                if (features.FP16) {
+                    return ml::cpu::Factory<std::uint16_t, trimd::neon::F256, trimd::neon::F128>::create(reader, memRes);
+                }
+            #endif  // RL_BUILD_WITH_HALF_FLOATS
+            return ml::cpu::Factory<float, trimd::neon::F256, trimd::neon::F128>::create(reader, memRes);
         }
     #endif  // RL_BUILD_WITH_NEON
     return ml::cpu::Factory<float, trimd::scalar::F256, trimd::scalar::F128>::create(reader, memRes);
diff --git a/Engine/Plugins/Animation/RigLogic/Source/RigLogicLib/Private/riglogic/system/simd/Detect.h b/Engine/Plugins/Animation/RigLogic/Source/RigLogicLib/Private/riglogic/system/simd/Detect.h
--- a/Engine/Plugins/Animation/RigLogic/Source/RigLogicLib/Private/riglogic/system/simd/Detect.h
+++ b/Engine/Plugins/Animation/RigLogic/Source/RigLogicLib/Private/riglogic/system/simd/Detect.h
@@ -2,28 +2,38 @@
 
 #pragma once
 
+#include "trimd/Platform.h"
+
 #if defined(RL_AUTODETECT_SSE) && !defined(RL_BUILD_WITH_SSE)
     // Assume any x86 architecture we try to build on will support SSE
-    #if defined(i386) || defined(__i386) || defined(__i386__) || defined(__IA32__) || defined(_M_IX86) || defined(__X86__) || \
-    defined(_X86_) || defined(__THW_INTEL__) || defined(__I86__) || defined(__amd64__) || defined(__amd64) || \
-    defined(__x86_64__) || defined(__x86_64) || defined(_M_X64) || defined(_M_AMD64)
+    #if defined(TRIMD_PLATFORM_X86)
         #define RL_BUILD_WITH_SSE 1
-        #if defined(RL_AUTODETECT_HALF_FLOATS) && !defined(RL_USE_HALF_FLOATS)
+        #if defined(RL_AUTODETECT_HALF_FLOATS) && !defined(RL_BUILD_WITH_HALF_FLOATS)
             #if defined(__F16C__) || (defined(_MSC_VER) && (_MSC_VER >= 1700) && defined(__AVX__)) || \
             (defined(__INTEL_COMPILER) && defined(__AVX__))
-                #define RL_USE_HALF_FLOATS 1
+                #define RL_BUILD_WITH_HALF_FLOATS 1
             #endif
         #endif  // RL_AUTODETECT_HALF_FLOATS
     #endif
 #endif  // RL_AUTODETECT_SSE
 
+#if defined(RL_AUTODETECT_AVX) && !defined(RL_BUILD_WITH_AVX)
+    #if defined(TRIMD_PLATFORM_X86) && defined(__AVX__)
+        #define RL_BUILD_WITH_AVX 1
+        #if defined(RL_AUTODETECT_HALF_FLOATS) && !defined(RL_BUILD_WITH_HALF_FLOATS)
+            #if defined(__F16C__) || (defined(_MSC_VER) && (_MSC_VER >= 1700)) || defined(__INTEL_COMPILER)
+                #define RL_BUILD_WITH_HALF_FLOATS 1
+            #endif
+        #endif  // RL_AUTODETECT_HALF_FLOATS
+    #endif
+#endif  // RL_AUTODETECT_AVX
+
 #if defined(RL_AUTODETECT_NEON) && !defined(RL_BUILD_WITH_NEON)
-    #if defined(__ARM_NEON)
+    #if defined(TRIMD_PLATFORM_ARM) && defined(__ARM_NEON)
         #define RL_BUILD_WITH_NEON 1
-        #if defined(RL_AUTODETECT_HALF_FLOATS) && !defined(RL_USE_HALF_FLOATS)
-            #if (defined(__GNUC__) && (defined(__arm__) || defined(__aarch64__))) || (defined(_MSC_VER) && defined(_M_ARM64))
-                #define RL_USE_HALF_FLOATS 1
-            #endif
+        #if defined(RL_AUTODETECT_HALF_FLOATS) && !defined(RL_BUILD_WITH_HALF_FLOATS)
+            // Half-float support will compile if NEON is available, and the actual feature check will be performed at runtime
+            #define RL_BUILD_WITH_HALF_FLOATS 1
         #endif  // RL_AUTODETECT_HALF_FLOATS
     #endif
 #endif  // RL_AUTODETECT_NEON
diff --git a/Engine/Plugins/Animation/RigLogic/Source/RigLogicLib/Private/riglogic/system/simd/SIMD.h b/Engine/Plugins/Animation/RigLogic/Source/RigLogicLib/Private/riglogic/system/simd/SIMD.h
--- a/Engine/Plugins/Animation/RigLogic/Source/RigLogicLib/Private/riglogic/system/simd/SIMD.h
+++ b/Engine/Plugins/Animation/RigLogic/Source/RigLogicLib/Private/riglogic/system/simd/SIMD.h
@@ -2,12 +2,18 @@
 
 #pragma once
 
+#ifndef RL_DISABLE_RUNTIME_FEATURE_DETECTION
+    #ifndef TRIMD_ENABLE_RUNTIME_FEATURE_DETECTION
+        #define TRIMD_ENABLE_RUNTIME_FEATURE_DETECTION
+    #endif  // TRIMD_ENABLE_RUNTIME_FEATURE_DETECTION
+#endif  // RL_DISABLE_RUNTIME_FEATURE_DETECTION
+
 #include "riglogic/system/simd/Detect.h"
 
 #if defined(RL_BUILD_WITH_AVX)
-    #if defined(RL_USE_HALF_FLOATS) && !defined(TRIMD_ENABLE_F16C)
+    #if defined(RL_BUILD_WITH_HALF_FLOATS) && !defined(TRIMD_ENABLE_F16C)
         #define TRIMD_ENABLE_F16C
-    #endif  // RL_USE_HALF_FLOATS
+    #endif  // RL_BUILD_WITH_HALF_FLOATS
     #if !defined(TRIMD_ENABLE_AVX)
         #define TRIMD_ENABLE_AVX
     #endif
@@ -17,16 +23,16 @@
 #endif  // RL_BUILD_WITH_AVX
 
 #if defined(RL_BUILD_WITH_SSE) && !defined(TRIMD_ENABLE_SSE)
-    #if defined(RL_USE_HALF_FLOATS) && !defined(TRIMD_ENABLE_F16C)
+    #if defined(RL_BUILD_WITH_HALF_FLOATS) && !defined(TRIMD_ENABLE_F16C)
         #define TRIMD_ENABLE_F16C
-    #endif  // RL_USE_HALF_FLOATS
+    #endif  // RL_BUILD_WITH_HALF_FLOATS
     #define TRIMD_ENABLE_SSE
 #endif  // RL_BUILD_WITH_SSE
 
 #if defined(RL_BUILD_WITH_NEON) && !defined(TRIMD_ENABLE_NEON)
-    #if defined(RL_USE_HALF_FLOATS) && !defined(TRIMD_ENABLE_NEON_FP16)
+    #if defined(RL_BUILD_WITH_HALF_FLOATS) && !defined(TRIMD_ENABLE_NEON_FP16)
         #define TRIMD_ENABLE_NEON_FP16
-    #endif  // RL_USE_HALF_FLOATS
+    #endif  // RL_BUILD_WITH_HALF_FLOATS
     #define TRIMD_ENABLE_NEON
 #endif  // RL_BUILD_WITH_NEON
 
diff --git a/Engine/Plugins/Animation/RigLogic/Source/RigLogicLib/Private/riglogic/utils/Macros.h b/Engine/Plugins/Animation/RigLogic/Source/RigLogicLib/Private/riglogic/utils/Macros.h
--- a/Engine/Plugins/Animation/RigLogic/Source/RigLogicLib/Private/riglogic/utils/Macros.h
+++ b/Engine/Plugins/Animation/RigLogic/Source/RigLogicLib/Private/riglogic/utils/Macros.h
@@ -18,6 +18,6 @@
     #endif
 #endif
 
-#if !defined(UNUSED)
-    #define UNUSED(x) static_cast<void>(x)
+#if !defined(RL_UNUSED)
+    #define RL_UNUSED(x) static_cast<void>(x)
 #endif
diff --git a/Engine/Plugins/Animation/RigLogic/Source/RigLogicLib/Public/riglogic/types/Aliases.h b/Engine/Plugins/Animation/RigLogic/Source/RigLogicLib/Public/riglogic/types/Aliases.h
--- a/Engine/Plugins/Animation/RigLogic/Source/RigLogicLib/Public/riglogic/types/Aliases.h
+++ b/Engine/Plugins/Animation/RigLogic/Source/RigLogicLib/Public/riglogic/types/Aliases.h
@@ -13,6 +13,7 @@
 #include <pma/resources/ArenaMemoryResource.h>
 #include <pma/resources/DefaultMemoryResource.h>
 #include <status/Status.h>
+#include <tdm/TDM.h>
 #include <trio/Stream.h>
 #include <trio/streams/FileStream.h>
 #include <trio/streams/MemoryMappedFileStream.h>
@@ -31,6 +32,8 @@ using dna::BinaryStreamWriter;
 using dna::StringView;
 using dna::UnknownLayerPolicy;
 using dna::Vector3;
+using tdm::fvec3;
+using tdm::fquat;
 
 template<typename T>
 using ArrayView = dna::ArrayView<T>;
diff --git a/Engine/Plugins/Animation/RigLogic/Source/RigLogicLib/Public/riglogic/version/Version.h b/Engine/Plugins/Animation/RigLogic/Source/RigLogicLib/Public/riglogic/version/Version.h
--- a/Engine/Plugins/Animation/RigLogic/Source/RigLogicLib/Public/riglogic/version/Version.h
+++ b/Engine/Plugins/Animation/RigLogic/Source/RigLogicLib/Public/riglogic/version/Version.h
@@ -3,6 +3,6 @@
 #pragma once
 
 #define RL_MAJOR_VERSION 10
-#define RL_MINOR_VERSION 2
-#define RL_PATCH_VERSION 7
-#define RL_VERSION_STRING "10.2.7"
+#define RL_MINOR_VERSION 3
+#define RL_PATCH_VERSION 0
+#define RL_VERSION_STRING "10.3.0"
diff --git a/Engine/Plugins/Animation/RigLogic/Source/RigLogicLib/Public/trimd/Platform.h b/Engine/Plugins/Animation/RigLogic/Source/RigLogicLib/Public/trimd/Platform.h
new file mode 100644
--- /dev/null
+++ b/Engine/Plugins/Animation/RigLogic/Source/RigLogicLib/Public/trimd/Platform.h
@@ -0,0 +1,172 @@
+// Copyright Epic Games, Inc. All Rights Reserved.
+
+#pragma once
+
+// *INDENT-OFF*
+#if defined(i386) || defined(__i386) || defined(__i386__) || defined(__IA32__) || defined(_M_IX86) || defined(__X86__) || \
+    defined(_X86_) || defined(__THW_INTEL__) || defined(__I86__) || defined(__amd64__) || defined(__amd64) || \
+    defined(__x86_64__) || defined(__x86_64) || defined(_M_X64) || defined(_M_AMD64)
+    #define TRIMD_PLATFORM_X86 1
+#endif
+
+#if defined(__arm__) || defined(__aarch64__) || defined(_M_ARM64) || defined(_M_ARM64EC)
+    #define TRIMD_PLATFORM_ARM 1
+#endif
+
+namespace trimd {
+
+struct CPUFeatures {
+    // ARM
+    bool NEON;
+    bool FP16;
+    // X86
+    bool SSE;
+    bool SSE2;
+    bool SSE3;
+    bool SSSE3;
+    bool SSE41;
+    bool SSE42;
+    bool AVX;
+    bool F16C;
+};
+
+}  // namespace trimd
+
+#ifdef TRIMD_ENABLE_RUNTIME_FEATURE_DETECTION
+
+#ifdef TRIMD_PLATFORM_X86
+    #ifdef _MSC_VER
+        #include <intrin.h>
+    #endif  // _MSC_VER
+#endif  // TRIMD_PLATFORM_X86
+
+#ifdef TRIMD_PLATFORM_ARM
+    #if defined(__APPLE__)
+        #include <mach/machine.h>
+        #include <sys/sysctl.h>
+        #include <sys/types.h>
+    #elif defined(WIN32)
+        #include "trimd/PlatformWindows.h"
+    #else
+        #include <sys/auxv.h>
+    #endif
+#endif  // TRIMD_PLATFORM_ARM
+
+namespace trimd {
+
+#ifdef TRIMD_PLATFORM_X86
+    #ifdef _MSC_VER
+        static void cpuidex(int info[4], int functionid, int subfunctionid) {
+            __cpuidex(info, functionid, subfunctionid);
+        }
+    #else
+        // Based on [REDACTED_PROFILE_URL]/ispc/blob/main/builtins/dispatch.c
+        static void cpuidex(int info[4], int functionid, int subfunctionid) {
+            __asm__ __volatile__ ("cpuid" : "=a" (info[0]), "=b" (info[1]), "=c" (info[2]), "=d" (info[3]) : "0" (functionid), "2" (subfunctionid));
+        }
+    #endif
+
+    inline CPUFeatures getCPUFeatures() {
+        CPUFeatures flags = {};
+        int info[4];
+        cpuidex(info, 1, 0);
+        flags.SSE = (info[3] & (1 << 25)) != 0;
+        flags.SSE2 = (info[3] & (1 << 26)) != 0;
+        flags.SSE3 = (info[2] & (1 << 0)) != 0;
+        flags.SSSE3 = (info[2] & (1 << 9)) != 0;
+        flags.SSE41 = (info[2] & (1 << 19)) != 0;
+        flags.SSE42 = (info[2] & (1 << 20)) != 0;
+        flags.AVX = (info[2] & (1 << 28)) != 0;
+        flags.F16C = (info[2] & (1 << 29)) != 0;
+        return flags;
+    }
+#endif  // TRIMD_PLATFORM_X86
+
+#ifdef TRIMD_PLATFORM_ARM
+    #if defined(__APPLE__)
+        static bool getsysctlflag(const char* name) {
+            unsigned int enabled = {};
+            std::size_t enabledSize = sizeof(enabled);
+            return sysctlbyname(name, &enabled, &enabledSize, nullptr, 0) == 0 && enabled;
+        }
+
+        inline CPUFeatures getCPUFeatures() {
+            CPUFeatures flags = {};
+            flags.NEON = getsysctlflag("hw.optional.AdvSIMD") || getsysctlflag("hw.optional.neon");
+            flags.FP16 = (getsysctlflag("hw.optional.arm.FEAT_FP16") || getsysctlflag("hw.optional.neon_fp16")) &&
+                (getsysctlflag("hw.optional.AdvSIMD_HPFPCvt") || getsysctlflag("hw.optional.neon_hpfp"));
+            return flags;
+        }
+    #elif defined(WIN32)
+        #ifndef PF_ARM_NEON_INSTRUCTIONS_AVAILABLE
+            #define PF_ARM_NEON_INSTRUCTIONS_AVAILABLE 19
+        #endif  // PF_ARM_NEON_INSTRUCTIONS_AVAILABLE
+
+        inline CPUFeatures getCPUFeatures() {
+            CPUFeatures flags = {};
+            flags.NEON = IsProcessorFeaturePresent(PF_ARM_NEON_INSTRUCTIONS_AVAILABLE) != 0;
+            // Heuristic, guessing that FP16 must be available if the below conditions hold true
+            flags.FP16 = (IsProcessorFeaturePresent(PF_FLOATING_POINT_EMULATED) == 0) &&
+                (IsProcessorFeaturePresent(PF_ARM_FMAC_INSTRUCTIONS_AVAILABLE) != 0);
+            return flags;
+        }
+    #else
+        #ifdef __arm__
+            // ARM-32bit
+            #ifndef HWCAP_NEON
+                #define HWCAP_NEON (1 << 12)
+            #endif  // HWCAP_NEON
+
+            #ifndef HWCAP_FPHP
+                #define HWCAP_FPHP (1 << 22)
+            #endif  // HWCAP_FPHP
+
+            #ifndef HWCAP_ASIMDHP
+                #define HWCAP_ASIMDHP (1 << 23)
+            #endif  // HWCAP_ASIMDHP
+        #else
+            // ARM-64bit
+            #ifndef HWCAP_ASIMD
+                #define HWCAP_ASIMD (1 << 1)
+            #endif  // HWCAP_ASIMD
+
+            #ifndef HWCAP_FPHP
+                #define HWCAP_FPHP (1 << 9)
+            #endif  // HWCAP_FPHP
+
+            #ifndef HWCAP_ASIMDHP
+                #define HWCAP_ASIMDHP (1 << 10)
+            #endif  // HWCAP_ASIMDHP
+        #endif
+
+        inline CPUFeatures getCPUFeatures() {
+            CPUFeatures flags = {};
+            const unsigned long hwcaps = getauxval(AT_HWCAP);
+            #ifdef __arm__
+                // ARM-32bit
+                flags.NEON = (hwcaps & HWCAP_NEON) != 0;
+            #else
+                // ARM-64bit
+                flags.NEON = (hwcaps & HWCAP_ASIMD) != 0;
+            #endif
+
+            flags.FP16 = ((hwcaps & HWCAP_ASIMDHP) != 0) && ((hwcaps & HWCAP_FPHP) != 0);
+            return flags;
+        }
+    #endif  // __linux__
+#endif  // TRIMD_PLATFORM_ARM
+
+}  // namespace trimd
+
+#else
+
+namespace trimd {
+
+inline CPUFeatures getCPUFeatures() {
+    return CPUFeatures{};
+}
+
+}  // namespace trimd
+
+#endif  // TRIMD_ENABLE_RUNTIME_FEATURE_DETECTION
+// *INDENT-ON*
diff --git a/Engine/Plugins/Animation/RigLogic/Source/RigLogicLib/Public/trimd/PlatformWindows.h b/Engine/Plugins/Animation/RigLogic/Source/RigLogicLib/Public/trimd/PlatformWindows.h
new file mode 100644
--- /dev/null
+++ b/Engine/Plugins/Animation/RigLogic/Source/RigLogicLib/Public/trimd/PlatformWindows.h
@@ -0,0 +1,60 @@
+// Copyright Epic Games, Inc. All Rights Reserved.
+
+#pragma once
+
+// HEADER_UNIT_SKIP - Not included directly
+
+#ifdef WIN32
+    #ifdef _MSC_VER
+        #pragma warning(push)
+        #pragma warning(disable : 4005 4365 4987)
+    #endif
+
+    #ifdef TRIMD_CUSTOM_WINDOWS_H
+        #include TRIMD_CUSTOM_WINDOWS_H
+    #else
+        #define WIN32_LEAN_AND_MEAN
+        #define NOGDICAPMASKS
+        #define NOVIRTUALKEYCODES
+        #define NOWINMESSAGES
+        #define NOWINSTYLES
+        #define NOSYSMETRICS
+        #define NOMENUS
+        #define NOICONS
+        #define NOKEYSTATES
+        #define NOSYSCOMMANDS
+        #define NORASTEROPS
+        #define NOSHOWWINDOW
+        #define NOATOM
+        #define NOCLIPBOARD
+        #define NOCOLOR
+        #define NOCTLMGR
+        #define NODRAWTEXT
+        #define NOGDI
+        #define NOKERNEL
+        #define NOUSER
+        #define NOMB
+        #define NOMEMMGR
+        #define NOMETAFILE
+        #define NOMINMAX
+        #define NOMSG
+        #define NOOPENFILE
+        #define NOSCROLL
+        #define NOSERVICE
+        #define NOSOUND
+        #define NOTEXTMETRIC
+        #define NOWH
+        #define NOWINOFFSETS
+        #define NOCOMM
+        #define NOKANJI
+        #define NOHELP
+        #define NOPROFILER
+        #define NODEFERWINDOWPOS
+        #define NOMCX
+        #include <windows.h>
+    #endif  // TRIMD_CUSTOM_WINDOWS_H
+
+    #ifdef _MSC_VER
+        #pragma warning(pop)
+    #endif
+#endif  // WIN32
diff --git a/Engine/Plugins/Animation/RigLogic/Source/RigLogicLib/Public/trimd/TRiMD.h b/Engine/Plugins/Animation/RigLogic/Source/RigLogicLib/Public/trimd/TRiMD.h
--- a/Engine/Plugins/Animation/RigLogic/Source/RigLogicLib/Public/trimd/TRiMD.h
+++ b/Engine/Plugins/Animation/RigLogic/Source/RigLogicLib/Public/trimd/TRiMD.h
@@ -2,6 +2,8 @@
 
 #pragma once
 
+#include "trimd/Platform.h"
+// Includes that go after platform detection macros
 #include "trimd/AVX.h"
 #include "trimd/NEON.h"
 #include "trimd/SSE.h"
diff --git a/Engine/Plugins/Animation/RigLogic/Source/RigLogicLib/Public/trimd/version/Version.h b/Engine/Plugins/Animation/RigLogic/Source/RigLogicLib/Public/trimd/version/Version.h
--- a/Engine/Plugins/Animation/RigLogic/Source/RigLogicLib/Public/trimd/version/Version.h
+++ b/Engine/Plugins/Animation/RigLogic/Source/RigLogicLib/Public/trimd/version/Version.h
@@ -3,6 +3,6 @@
 #pragma once
 
 #define TRIMD_MAJOR_VERSION 2
-#define TRIMD_MINOR_VERSION 3
-#define TRIMD_PATCH_VERSION 1
-#define TRIMD_VERSION_STRING "2.3.1"
+#define TRIMD_MINOR_VERSION 4
+#define TRIMD_PATCH_VERSION 0
+#define TRIMD_VERSION_STRING "2.4.0"
diff --git a/Engine/Plugins/Animation/RigLogic/Source/RigLogicLib/RigLogicLib.Build.cs b/Engine/Plugins/Animation/RigLogic/Source/RigLogicLib/RigLogicLib.Build.cs
--- a/Engine/Plugins/Animation/RigLogic/Source/RigLogicLib/RigLogicLib.Build.cs
+++ b/Engine/Plugins/Animation/RigLogic/Source/RigLogicLib/RigLogicLib.Build.cs
@@ -36,8 +36,9 @@ namespace UnrealBuildTool.Rules
 
 			if (Target.Platform == UnrealTargetPlatform.Win64)
 			{
-				PrivateDefinitions.Add("TRIO_WINDOWS_FILE_MAPPING_AVAILABLE=1");
+				PrivateDefinitions.Add("TRIMD_CUSTOM_WINDOWS_H=\"WindowsPlatformUE.h\"");
 				PrivateDefinitions.Add("TRIO_CUSTOM_WINDOWS_H=\"WindowsPlatformUE.h\"");
+				PrivateDefinitions.Add("TRIO_WINDOWS_FILE_MAPPING_AVAILABLE=1");
 			}
 
 			if (Target.Platform == UnrealTargetPlatform.Linux || Target.Platform == UnrealTargetPlatform.LinuxArm64)
@@ -53,6 +54,7 @@ namespace UnrealBuildTool.Rules
 			}
 
 			PrivateDefinitions.Add("RL_AUTODETECT_SSE=1");
+			PrivateDefinitions.Add("RL_AUTODETECT_HALF_FLOATS=1");
 
 			bDisableAutoRTFMInstrumentation = true;
 		}
diff --git a/Engine/Plugins/Animation/RigLogic/Source/RigLogicLibTest/Private/rltests/StorageValueType.h b/Engine/Plugins/Animation/RigLogic/Source/RigLogicLibTest/Private/rltests/StorageValueType.h
--- a/Engine/Plugins/Animation/RigLogic/Source/RigLogicLibTest/Private/rltests/StorageValueType.h
+++ b/Engine/Plugins/Animation/RigLogic/Source/RigLogicLibTest/Private/rltests/StorageValueType.h
@@ -4,7 +4,7 @@
 
 #include <cstdint>
 
-#ifdef RL_USE_HALF_FLOATS
+#ifdef RL_BUILD_WITH_HALF_FLOATS
     using StorageValueType = std::uint16_t;
 #else
     using StorageValueType = float;
diff --git a/Engine/Plugins/Animation/RigLogic/Source/RigLogicLibTest/Private/rltests/joints/bpcm/TestBlock4CalculationStrategy.cpp b/Engine/Plugins/Animation/RigLogic/Source/RigLogicLibTest/Private/rltests/joints/bpcm/TestBlock4CalculationStrategy.cpp
--- a/Engine/Plugins/Animation/RigLogic/Source/RigLogicLibTest/Private/rltests/joints/bpcm/TestBlock4CalculationStrategy.cpp
+++ b/Engine/Plugins/Animation/RigLogic/Source/RigLogicLibTest/Private/rltests/joints/bpcm/TestBlock4CalculationStrategy.cpp
@@ -107,7 +107,7 @@ using Block4JointCalculationTypeList = ::testing::Types<
         std::tuple<StorageValueType, trimd::neon::F128, TStrategyTestParams<3>, rl4::bpcm::EulerAnglesToQuaternions<tdm::fdeg,
                                                                                                                     tdm::rot_seq::xyz> >,
 #endif  // RL_BUILD_WITH_NEON
-#if !defined(RL_USE_HALF_FLOATS)
+#if !defined(RL_BUILD_WITH_HALF_FLOATS)
         std::tuple<StorageValueType, trimd::scalar::F128, TStrategyTestParams<0>, rl4::bpcm::NoopAdapter>,
         std::tuple<StorageValueType, trimd::scalar::F128, TStrategyTestParams<0>,
                    rl4::bpcm::EulerAnglesToQuaternions<tdm::fdeg, tdm::rot_seq::xyz> >,
@@ -120,7 +120,7 @@ using Block4JointCalculationTypeList = ::testing::Types<
         std::tuple<StorageValueType, trimd::scalar::F128, TStrategyTestParams<3>, rl4::bpcm::NoopAdapter>,
         std::tuple<StorageValueType, trimd::scalar::F128, TStrategyTestParams<3>,
                    rl4::bpcm::EulerAnglesToQuaternions<tdm::fdeg, tdm::rot_seq::xyz> >,
-#endif  // RL_USE_HALF_FLOATS
+#endif  // RL_BUILD_WITH_HALF_FLOATS
     std::tuple<>
     >;
 
diff --git a/Engine/Plugins/Animation/RigLogic/Source/RigLogicLibTest/Private/rltests/joints/bpcm/TestBlock4StorageBuilder.cpp b/Engine/Plugins/Animation/RigLogic/Source/RigLogicLibTest/Private/rltests/joints/bpcm/TestBlock4StorageBuilder.cpp
--- a/Engine/Plugins/Animation/RigLogic/Source/RigLogicLibTest/Private/rltests/joints/bpcm/TestBlock4StorageBuilder.cpp
+++ b/Engine/Plugins/Animation/RigLogic/Source/RigLogicLibTest/Private/rltests/joints/bpcm/TestBlock4StorageBuilder.cpp
@@ -85,11 +85,11 @@ using Block4StorageValueTypeList = ::testing::Types<
         std::tuple<StorageValueType, trimd::neon::F128, TCalculationType<rl4::CalculationType::NEON>,
                    rl4::bpcm::EulerAnglesToQuaternions<tdm::fdeg, tdm::rot_seq::xyz> >,
 #endif  // RL_BUILD_WITH_NEON
-#if !defined(RL_USE_HALF_FLOATS)
+#if !defined(RL_BUILD_WITH_HALF_FLOATS)
         std::tuple<StorageValueType, trimd::scalar::F128, TCalculationType<rl4::CalculationType::Scalar>, rl4::bpcm::NoopAdapter>,
         std::tuple<StorageValueType, trimd::scalar::F128, TCalculationType<rl4::CalculationType::Scalar>,
                    rl4::bpcm::EulerAnglesToQuaternions<tdm::fdeg, tdm::rot_seq::xyz> >,
-#endif  // RL_USE_HALF_FLOATS
+#endif  // RL_BUILD_WITH_HALF_FLOATS
     std::tuple<>
     >;
 
diff --git a/Engine/Plugins/Animation/RigLogic/Source/RigLogicLibTest/Private/rltests/joints/bpcm/TestBlock8CalculationStrategy.cpp b/Engine/Plugins/Animation/RigLogic/Source/RigLogicLibTest/Private/rltests/joints/bpcm/TestBlock8CalculationStrategy.cpp
--- a/Engine/Plugins/Animation/RigLogic/Source/RigLogicLibTest/Private/rltests/joints/bpcm/TestBlock8CalculationStrategy.cpp
+++ b/Engine/Plugins/Animation/RigLogic/Source/RigLogicLibTest/Private/rltests/joints/bpcm/TestBlock8CalculationStrategy.cpp
@@ -120,7 +120,7 @@ using Block8JointCalculationTypeList = ::testing::Types<
         std::tuple<StorageValueType, trimd::neon::F256, TStrategyTestParams<3>, rl4::bpcm::EulerAnglesToQuaternions<tdm::fdeg,
                                                                                                                     tdm::rot_seq::xyz> >,
 #endif  // RL_BUILD_WITH_NEON
-#if !defined(RL_USE_HALF_FLOATS)
+#if !defined(RL_BUILD_WITH_HALF_FLOATS)
         std::tuple<StorageValueType, trimd::scalar::F256, TStrategyTestParams<0>, rl4::bpcm::NoopAdapter>,
         std::tuple<StorageValueType, trimd::scalar::F256, TStrategyTestParams<0>,
                    rl4::bpcm::EulerAnglesToQuaternions<tdm::fdeg, tdm::rot_seq::xyz> >,
@@ -133,7 +133,7 @@ using Block8JointCalculationTypeList = ::testing::Types<
         std::tuple<StorageValueType, trimd::scalar::F256, TStrategyTestParams<3>, rl4::bpcm::NoopAdapter>,
         std::tuple<StorageValueType, trimd::scalar::F256, TStrategyTestParams<3>,
                    rl4::bpcm::EulerAnglesToQuaternions<tdm::fdeg, tdm::rot_seq::xyz> >,
-#endif  // RL_USE_HALF_FLOATS
+#endif  // RL_BUILD_WITH_HALF_FLOATS
     std::tuple<>
     >;
 
diff --git a/Engine/Plugins/Animation/RigLogic/Source/RigLogicLibTest/Private/rltests/joints/quaternions/TestQuaternionEvaluator.cpp b/Engine/Plugins/Animation/RigLogic/Source/RigLogicLibTest/Private/rltests/joints/quaternions/TestQuaternionEvaluator.cpp
--- a/Engine/Plugins/Animation/RigLogic/Source/RigLogicLibTest/Private/rltests/joints/quaternions/TestQuaternionEvaluator.cpp
+++ b/Engine/Plugins/Animation/RigLogic/Source/RigLogicLibTest/Private/rltests/joints/quaternions/TestQuaternionEvaluator.cpp
@@ -73,7 +73,7 @@ class QuaternionEvaluatorTest : public ::testing::Test {
 
 #if defined(RL_BUILD_WITH_AVX) && defined(RL_BUILD_WITH_SSE)
     using QuaternionEvaluatorTypeList = ::testing::Types<
-    #ifdef RL_USE_HALF_FLOATS
+    #ifdef RL_BUILD_WITH_HALF_FLOATS
             std::tuple<StorageValueType, trimd::avx::F256, trimd::sse::F128, rl4::QuaternionsToEulerAngles<tdm::frad,
                                                                                                            tdm::rot_seq::xyz> >,
             std::tuple<StorageValueType, trimd::avx::F256, trimd::sse::F128, rl4::PassthroughAdapter>,
@@ -90,11 +90,11 @@ class QuaternionEvaluatorTest : public ::testing::Test {
             std::tuple<StorageValueType, trimd::scalar::F256, trimd::scalar::F128, rl4::QuaternionsToEulerAngles<tdm::frad,
                                                                                                                  tdm::rot_seq::xyz> >,
             std::tuple<StorageValueType, trimd::scalar::F256, trimd::scalar::F128, rl4::PassthroughAdapter>
-    #endif  // RL_USE_HALF_FLOATS
+    #endif  // RL_BUILD_WITH_HALF_FLOATS
         >;
 #elif defined(RL_BUILD_WITH_AVX)
     using QuaternionEvaluatorTypeList = ::testing::Types<
-    #ifdef RL_USE_HALF_FLOATS
+    #ifdef RL_BUILD_WITH_HALF_FLOATS
             std::tuple<StorageValueType, trimd::avx::F256, trimd::sse::F128, rl4::QuaternionsToEulerAngles<tdm::frad,
                                                                                                            tdm::rot_seq::xyz> >,
             std::tuple<StorageValueType, trimd::avx::F256, trimd::sse::F128, rl4::PassthroughAdapter>
@@ -105,11 +105,11 @@ class QuaternionEvaluatorTest : public ::testing::Test {
             std::tuple<StorageValueType, trimd::scalar::F256, trimd::scalar::F128, rl4::QuaternionsToEulerAngles<tdm::frad,
                                                                                                                  tdm::rot_seq::xyz> >,
             std::tuple<StorageValueType, trimd::scalar::F256, trimd::scalar::F128, rl4::PassthroughAdapter>
-    #endif  // RL_USE_HALF_FLOATS
+    #endif  // RL_BUILD_WITH_HALF_FLOATS
         >;
 #elif defined(RL_BUILD_WITH_SSE)
     using QuaternionEvaluatorTypeList = ::testing::Types<
-    #ifdef RL_USE_HALF_FLOATS
+    #ifdef RL_BUILD_WITH_HALF_FLOATS
             std::tuple<StorageValueType, trimd::sse::F256, trimd::sse::F128, rl4::QuaternionsToEulerAngles<tdm::frad,
                                                                                                            tdm::rot_seq::xyz> >,
             std::tuple<StorageValueType, trimd::sse::F256, trimd::sse::F128, rl4::PassthroughAdapter>
@@ -120,10 +120,10 @@ class QuaternionEvaluatorTest : public ::testing::Test {
             std::tuple<StorageValueType, trimd::scalar::F256, trimd::scalar::F128, rl4::QuaternionsToEulerAngles<tdm::frad,
                                                                                                                  tdm::rot_seq::xyz> >,
             std::tuple<StorageValueType, trimd::scalar::F256, trimd::scalar::F128, rl4::PassthroughAdapter>
-    #endif  // RL_USE_HALF_FLOATS
+    #endif  // RL_BUILD_WITH_HALF_FLOATS
         >;
 #else
-    #ifndef RL_USE_HALF_FLOATS
+    #ifndef RL_BUILD_WITH_HALF_FLOATS
         using QuaternionEvaluatorTypeList = ::testing::Types<
             std::tuple<StorageValueType, trimd::scalar::F256, trimd::scalar::F128, rl4::QuaternionsToEulerAngles<tdm::frad,
                                                                                                                  tdm::rot_seq::xyz> >,
@@ -131,7 +131,7 @@ class QuaternionEvaluatorTest : public ::testing::Test {
             >;
     #else
         using QuaternionEvaluatorTypeList = ::testing::Types<std::tuple<> >;
-    #endif  // RL_USE_HALF_FLOATS
+    #endif  // RL_BUILD_WITH_HALF_FLOATS
 #endif
 
 TYPED_TEST_SUITE(QuaternionEvaluatorTest, QuaternionEvaluatorTypeList, );
@@ -159,11 +159,11 @@ TYPED_TEST(QuaternionEvaluatorTest, EvaluatePerLOD) {
         std::fill(outputBuffer.begin(), outputBuffer.end(), 0.0f);
         this->evaluator->calculate(inputInstance.get(), &outputInstance, lod);
         const auto& expected = rltests::qs::output::valuesPerLODPerConfig[this->rotationSelectorIndex][lod];
-        #ifdef RL_USE_HALF_FLOATS
+        #ifdef RL_BUILD_WITH_HALF_FLOATS
             static constexpr float threshold = 0.05f;
         #else
             static constexpr float threshold = 0.002f;
-        #endif  // RL_USE_HALF_FLOATS
+        #endif  // RL_BUILD_WITH_HALF_FLOATS
         ASSERT_ELEMENTS_NEAR(outputBuffer, expected, expected.size(), threshold);
     }
 }
diff --git a/Engine/Plugins/Animation/RigLogic/Source/RigLogicLibTest/Private/rltests/joints/quaternions/TestQuaternionStorageBuilder.cpp b/Engine/Plugins/Animation/RigLogic/Source/RigLogicLibTest/Private/rltests/joints/quaternions/TestQuaternionStorageBuilder.cpp
--- a/Engine/Plugins/Animation/RigLogic/Source/RigLogicLibTest/Private/rltests/joints/quaternions/TestQuaternionStorageBuilder.cpp
+++ b/Engine/Plugins/Animation/RigLogic/Source/RigLogicLibTest/Private/rltests/joints/quaternions/TestQuaternionStorageBuilder.cpp
@@ -91,7 +91,7 @@ class QuaternionJointStorageBuilderTest : public ::testing::Test {
 
 #if defined(RL_BUILD_WITH_AVX) && defined(RL_BUILD_WITH_SSE)
     using StorageValueTypeList = ::testing::Types<
-    #ifdef RL_USE_HALF_FLOATS
+    #ifdef RL_BUILD_WITH_HALF_FLOATS
             std::tuple<StorageValueType, trimd::avx::F256, trimd::sse::F128, rl4::PassthroughAdapter>,
             std::tuple<StorageValueType, trimd::avx::F256, trimd::sse::F128, rl4::QuaternionsToEulerAngles<tdm::frad,
                                                                                                            tdm::rot_seq::xyz> >,
@@ -108,11 +108,11 @@ class QuaternionJointStorageBuilderTest : public ::testing::Test {
             std::tuple<StorageValueType, trimd::scalar::F256, trimd::scalar::F128, rl4::PassthroughAdapter>,
             std::tuple<StorageValueType, trimd::scalar::F256, trimd::scalar::F128, rl4::QuaternionsToEulerAngles<tdm::frad,
                                                                                                                  tdm::rot_seq::xyz> >
-    #endif  // RL_USE_HALF_FLOATS
+    #endif  // RL_BUILD_WITH_HALF_FLOATS
         >;
 #elif defined(RL_BUILD_WITH_AVX)
     using StorageValueTypeList = ::testing::Types<
-    #ifdef RL_USE_HALF_FLOATS
+    #ifdef RL_BUILD_WITH_HALF_FLOATS
             std::tuple<StorageValueType, trimd::avx::F256, trimd::sse::F128, rl4::PassthroughAdapter>,
             std::tuple<StorageValueType, trimd::avx::F256, trimd::sse::F128, rl4::QuaternionsToEulerAngles<tdm::frad,
                                                                                                            tdm::rot_seq::xyz> >
@@ -123,11 +123,11 @@ class QuaternionJointStorageBuilderTest : public ::testing::Test {
             std::tuple<StorageValueType, trimd::scalar::F256, trimd::scalar::F128, rl4::PassthroughAdapter>,
             std::tuple<StorageValueType, trimd::scalar::F256, trimd::scalar::F128, rl4::QuaternionsToEulerAngles<tdm::frad,
                                                                                                                  tdm::rot_seq::xyz> >
-    #endif  // RL_USE_HALF_FLOATS
+    #endif  // RL_BUILD_WITH_HALF_FLOATS
         >;
 #elif defined(RL_BUILD_WITH_SSE)
     using StorageValueTypeList = ::testing::Types<
-    #ifdef RL_USE_HALF_FLOATS
+    #ifdef RL_BUILD_WITH_HALF_FLOATS
             std::tuple<StorageValueType, trimd::sse::F256, trimd::sse::F128, rl4::PassthroughAdapter>,
             std::tuple<StorageValueType, trimd::sse::F256, trimd::sse::F128, rl4::QuaternionsToEulerAngles<tdm::frad,
                                                                                                            tdm::rot_seq::xyz> >
@@ -138,10 +138,10 @@ class QuaternionJointStorageBuilderTest : public ::testing::Test {
             std::tuple<StorageValueType, trimd::scalar::F256, trimd::scalar::F128, rl4::PassthroughAdapter>,
             std::tuple<StorageValueType, trimd::scalar::F256, trimd::scalar::F128, rl4::QuaternionsToEulerAngles<tdm::frad,
                                                                                                                  tdm::rot_seq::xyz> >
-    #endif  // RL_USE_HALF_FLOATS
+    #endif  // RL_BUILD_WITH_HALF_FLOATS
         >;
 #else
-    #ifndef RL_USE_HALF_FLOATS
+    #ifndef RL_BUILD_WITH_HALF_FLOATS
         using StorageValueTypeList = ::testing::Types<
             std::tuple<StorageValueType, trimd::scalar::F256, trimd::scalar::F128, rl4::PassthroughAdapter>,
             std::tuple<StorageValueType, trimd::scalar::F256, trimd::scalar::F128, rl4::QuaternionsToEulerAngles<tdm::frad,
@@ -149,7 +149,7 @@ class QuaternionJointStorageBuilderTest : public ::testing::Test {
             >;
     #else
         using StorageValueTypeList = ::testing::Types<std::tuple<> >;
-    #endif  // RL_USE_HALF_FLOATS
+    #endif  // RL_BUILD_WITH_HALF_FLOATS
 #endif
 
 TYPED_TEST_SUITE(QuaternionJointStorageBuilderTest, StorageValueTypeList, );
diff --git a/Engine/Plugins/Animation/RigLogic/Source/RigLogicLibTest/Private/rltests/joints/twistswing/TestTwistSwingEvaluator.cpp b/Engine/Plugins/Animation/RigLogic/Source/RigLogicLibTest/Private/rltests/joints/twistswing/TestTwistSwingEvaluator.cpp
--- a/Engine/Plugins/Animation/RigLogic/Source/RigLogicLibTest/Private/rltests/joints/twistswing/TestTwistSwingEvaluator.cpp
+++ b/Engine/Plugins/Animation/RigLogic/Source/RigLogicLibTest/Private/rltests/joints/twistswing/TestTwistSwingEvaluator.cpp
@@ -74,7 +74,7 @@ class TwistSwingEvaluatorTest : public ::testing::Test {
 
 #if defined(RL_BUILD_WITH_AVX) && defined(RL_BUILD_WITH_SSE)
     using TwistSwingEvaluatorTypeList = ::testing::Types<
-    #ifdef RL_USE_HALF_FLOATS
+    #ifdef RL_BUILD_WITH_HALF_FLOATS
             std::tuple<float, trimd::avx::F256, trimd::sse::F128, rl4::QuaternionsToEulerAngles<tdm::frad, tdm::rot_seq::xyz> >,
             std::tuple<float, trimd::avx::F256, trimd::sse::F128, rl4::PassthroughAdapter>,
             std::tuple<float, trimd::sse::F256, trimd::sse::F128, rl4::QuaternionsToEulerAngles<tdm::frad, tdm::rot_seq::xyz> >,
@@ -87,11 +87,11 @@ class TwistSwingEvaluatorTest : public ::testing::Test {
             std::tuple<float, trimd::scalar::F256, trimd::scalar::F128, rl4::QuaternionsToEulerAngles<tdm::frad,
                                                                                                       tdm::rot_seq::xyz> >,
             std::tuple<float, trimd::scalar::F256, trimd::scalar::F128, rl4::PassthroughAdapter>
-    #endif  // RL_USE_HALF_FLOATS
+    #endif  // RL_BUILD_WITH_HALF_FLOATS
         >;
 #elif defined(RL_BUILD_WITH_AVX)
     using TwistSwingEvaluatorTypeList = ::testing::Types<
-    #ifdef RL_USE_HALF_FLOATS
+    #ifdef RL_BUILD_WITH_HALF_FLOATS
             std::tuple<float, trimd::avx::F256, trimd::sse::F128, rl4::QuaternionsToEulerAngles<tdm::frad, tdm::rot_seq::xyz> >,
             std::tuple<float, trimd::avx::F256, trimd::sse::F128, rl4::PassthroughAdapter>
     #else
@@ -100,11 +100,11 @@ class TwistSwingEvaluatorTest : public ::testing::Test {
             std::tuple<float, trimd::scalar::F256, trimd::scalar::F128, rl4::QuaternionsToEulerAngles<tdm::frad,
                                                                                                       tdm::rot_seq::xyz> >,
             std::tuple<float, trimd::scalar::F256, trimd::scalar::F128, rl4::PassthroughAdapter>
-    #endif  // RL_USE_HALF_FLOATS
+    #endif  // RL_BUILD_WITH_HALF_FLOATS
         >;
 #elif defined(RL_BUILD_WITH_SSE)
     using TwistSwingEvaluatorTypeList = ::testing::Types<
-    #ifdef RL_USE_HALF_FLOATS
+    #ifdef RL_BUILD_WITH_HALF_FLOATS
             std::tuple<float, trimd::sse::F256, trimd::sse::F128, rl4::QuaternionsToEulerAngles<tdm::frad, tdm::rot_seq::xyz> >,
             std::tuple<float, trimd::sse::F256, trimd::sse::F128, rl4::PassthroughAdapter>
     #else
@@ -113,10 +113,10 @@ class TwistSwingEvaluatorTest : public ::testing::Test {
             std::tuple<float, trimd::scalar::F256, trimd::scalar::F128, rl4::QuaternionsToEulerAngles<tdm::frad,
                                                                                                       tdm::rot_seq::xyz> >,
             std::tuple<float, trimd::scalar::F256, trimd::scalar::F128, rl4::PassthroughAdapter>
-    #endif  // RL_USE_HALF_FLOATS
+    #endif  // RL_BUILD_WITH_HALF_FLOATS
         >;
 #else
-    #ifndef RL_USE_HALF_FLOATS
+    #ifndef RL_BUILD_WITH_HALF_FLOATS
         using TwistSwingEvaluatorTypeList = ::testing::Types<
             std::tuple<float, trimd::scalar::F256, trimd::scalar::F128, rl4::QuaternionsToEulerAngles<tdm::frad,
                                                                                                       tdm::rot_seq::xyz> >,
@@ -124,7 +124,7 @@ class TwistSwingEvaluatorTest : public ::testing::Test {
             >;
     #else
         using TwistSwingEvaluatorTypeList = ::testing::Types<std::tuple<> >;
-    #endif  // RL_USE_HALF_FLOATS
+    #endif  // RL_BUILD_WITH_HALF_FLOATS
 #endif
 
 TYPED_TEST_SUITE(TwistSwingEvaluatorTest, TwistSwingEvaluatorTypeList, );
@@ -152,11 +152,11 @@ TYPED_TEST(TwistSwingEvaluatorTest, EvaluatePerLOD) {
         std::fill(outputBuffer.begin(), outputBuffer.end(), 0.0f);
         this->evaluator->calculate(inputInstance.get(), &outputInstance, lod);
         const auto& expected = rltests::twsw::output::valuesPerLODPerConfig[this->rotationSelectorIndex][lod];
-        #ifdef RL_USE_HALF_FLOATS
+        #ifdef RL_BUILD_WITH_HALF_FLOATS
             static constexpr float threshold = 0.05f;
         #else
             static constexpr float threshold = 0.002f;
-        #endif  // RL_USE_HALF_FLOATS
+        #endif  // RL_BUILD_WITH_HALF_FLOATS
         ASSERT_ELEMENTS_NEAR(outputBuffer, expected, expected.size(), threshold);
     }
 }
diff --git a/Engine/Plugins/Animation/RigLogic/Source/RigLogicLibTest/Private/rltests/joints/twistswing/TestTwistSwingStorageBuilder.cpp b/Engine/Plugins/Animation/RigLogic/Source/RigLogicLibTest/Private/rltests/joints/twistswing/TestTwistSwingStorageBuilder.cpp
--- a/Engine/Plugins/Animation/RigLogic/Source/RigLogicLibTest/Private/rltests/joints/twistswing/TestTwistSwingStorageBuilder.cpp
+++ b/Engine/Plugins/Animation/RigLogic/Source/RigLogicLibTest/Private/rltests/joints/twistswing/TestTwistSwingStorageBuilder.cpp
@@ -102,7 +102,7 @@ class TwistSwingJointStorageBuilderTest : public ::testing::Test {
 
 #if defined(RL_BUILD_WITH_AVX) && defined(RL_BUILD_WITH_SSE)
     using StorageValueTypeList = ::testing::Types<
-    #ifdef RL_USE_HALF_FLOATS
+    #ifdef RL_BUILD_WITH_HALF_FLOATS
             std::tuple<float, trimd::avx::F256, trimd::sse::F128, rl4::PassthroughAdapter>,
             std::tuple<float, trimd::avx::F256, trimd::sse::F128, rl4::QuaternionsToEulerAngles<tdm::frad, tdm::rot_seq::xyz> >,
             std::tuple<float, trimd::sse::F256, trimd::sse::F128, rl4::PassthroughAdapter>,
@@ -115,11 +115,11 @@ class TwistSwingJointStorageBuilderTest : public ::testing::Test {
             std::tuple<float, trimd::scalar::F256, trimd::scalar::F128, rl4::PassthroughAdapter>,
             std::tuple<float, trimd::scalar::F256, trimd::scalar::F128, rl4::QuaternionsToEulerAngles<tdm::frad,
                                                                                                       tdm::rot_seq::xyz> >
-    #endif  // RL_USE_HALF_FLOATS
+    #endif  // RL_BUILD_WITH_HALF_FLOATS
         >;
 #elif defined(RL_BUILD_WITH_AVX)
     using StorageValueTypeList = ::testing::Types<
-    #ifdef RL_USE_HALF_FLOATS
+    #ifdef RL_BUILD_WITH_HALF_FLOATS
             std::tuple<float, trimd::avx::F256, trimd::sse::F128, rl4::PassthroughAdapter>,
             std::tuple<float, trimd::avx::F256, trimd::sse::F128, rl4::QuaternionsToEulerAngles<tdm::frad, tdm::rot_seq::xyz> >
     #else
@@ -128,11 +128,11 @@ class TwistSwingJointStorageBuilderTest : public ::testing::Test {
             std::tuple<float, trimd::scalar::F256, trimd::scalar::F128, rl4::PassthroughAdapter>,
             std::tuple<float, trimd::scalar::F256, trimd::scalar::F128, rl4::QuaternionsToEulerAngles<tdm::frad,
                                                                                                       tdm::rot_seq::xyz> >
-    #endif  // RL_USE_HALF_FLOATS
+    #endif  // RL_BUILD_WITH_HALF_FLOATS
         >;
 #elif defined(RL_BUILD_WITH_SSE)
     using StorageValueTypeList = ::testing::Types<
-    #ifdef RL_USE_HALF_FLOATS
+    #ifdef RL_BUILD_WITH_HALF_FLOATS
             std::tuple<float, trimd::sse::F256, trimd::sse::F128, rl4::PassthroughAdapter>,
             std::tuple<float, trimd::sse::F256, trimd::sse::F128, rl4::QuaternionsToEulerAngles<tdm::frad, tdm::rot_seq::xyz> >
     #else
@@ -141,10 +141,10 @@ class TwistSwingJointStorageBuilderTest : public ::testing::Test {
             std::tuple<float, trimd::scalar::F256, trimd::scalar::F128, rl4::PassthroughAdapter>,
             std::tuple<float, trimd::scalar::F256, trimd::scalar::F128, rl4::QuaternionsToEulerAngles<tdm::frad,
                                                                                                       tdm::rot_seq::xyz> >
-    #endif  // RL_USE_HALF_FLOATS
+    #endif  // RL_BUILD_WITH_HALF_FLOATS
         >;
 #else
-    #ifndef RL_USE_HALF_FLOATS
+    #ifndef RL_BUILD_WITH_HALF_FLOATS
         using StorageValueTypeList = ::testing::Types<
             std::tuple<float, trimd::scalar::F256, trimd::scalar::F128, rl4::PassthroughAdapter>,
             std::tuple<float, trimd::scalar::F256, trimd::scalar::F128, rl4::QuaternionsToEulerAngles<tdm::frad,
@@ -152,7 +152,7 @@ class TwistSwingJointStorageBuilderTest : public ::testing::Test {
             >;
     #else
         using StorageValueTypeList = ::testing::Types<std::tuple<> >;
-    #endif  // RL_USE_HALF_FLOATS
+    #endif  // RL_BUILD_WITH_HALF_FLOATS
 #endif
 
 TYPED_TEST_SUITE(TwistSwingJointStorageBuilderTest, StorageValueTypeList, );
diff --git a/Engine/Plugins/Animation/RigLogic/Source/RigLogicLibTest/Private/rltests/ml/cpu/TestInference.cpp b/Engine/Plugins/Animation/RigLogic/Source/RigLogicLibTest/Private/rltests/ml/cpu/TestInference.cpp
--- a/Engine/Plugins/Animation/RigLogic/Source/RigLogicLibTest/Private/rltests/ml/cpu/TestInference.cpp
+++ b/Engine/Plugins/Animation/RigLogic/Source/RigLogicLibTest/Private/rltests/ml/cpu/TestInference.cpp
@@ -49,39 +49,39 @@ class MLBInferenceTest : public ::testing::Test {
 
 #if defined(RL_BUILD_WITH_AVX) && defined(RL_BUILD_WITH_SSE)
     using InferenceTypeList = ::testing::Types<
-    #ifdef RL_USE_HALF_FLOATS
+    #ifdef RL_BUILD_WITH_HALF_FLOATS
             std::tuple<StorageValueType, trimd::avx::F256, trimd::sse::F128>,
             std::tuple<StorageValueType, trimd::sse::F256, trimd::sse::F128>
     #else
             std::tuple<StorageValueType, trimd::avx::F256, trimd::sse::F128>,
             std::tuple<StorageValueType, trimd::sse::F256, trimd::sse::F128>,
             std::tuple<StorageValueType, trimd::scalar::F256, trimd::scalar::F128>
-    #endif  // RL_USE_HALF_FLOATS
+    #endif  // RL_BUILD_WITH_HALF_FLOATS
         >;
 #elif defined(RL_BUILD_WITH_AVX)
     using InferenceTypeList = ::testing::Types<
-    #ifdef RL_USE_HALF_FLOATS
+    #ifdef RL_BUILD_WITH_HALF_FLOATS
             std::tuple<StorageValueType, trimd::avx::F256, trimd::sse::F128>
     #else
             std::tuple<StorageValueType, trimd::avx::F256, trimd::sse::F128>,
             std::tuple<StorageValueType, trimd::scalar::F256, trimd::scalar::F128>
-    #endif  // RL_USE_HALF_FLOATS
+    #endif  // RL_BUILD_WITH_HALF_FLOATS
         >;
 #elif defined(RL_BUILD_WITH_SSE)
     using InferenceTypeList = ::testing::Types<
-    #ifdef RL_USE_HALF_FLOATS
+    #ifdef RL_BUILD_WITH_HALF_FLOATS
             std::tuple<StorageValueType, trimd::sse::F256, trimd::sse::F128>
     #else
             std::tuple<StorageValueType, trimd::sse::F256, trimd::sse::F128>,
             std::tuple<StorageValueType, trimd::scalar::F256, trimd::scalar::F128>
-    #endif  // RL_USE_HALF_FLOATS
+    #endif  // RL_BUILD_WITH_HALF_FLOATS
         >;
 #else
-    #ifndef RL_USE_HALF_FLOATS
+    #ifndef RL_BUILD_WITH_HALF_FLOATS
         using InferenceTypeList = ::testing::Types<std::tuple<StorageValueType, trimd::scalar::F256, trimd::scalar::F128> >;
     #else
         using InferenceTypeList = ::testing::Types<std::tuple<> >;
-    #endif  // RL_USE_HALF_FLOATS
+    #endif  // RL_BUILD_WITH_HALF_FLOATS
 #endif
 
 TYPED_TEST_SUITE(MLBInferenceTest, InferenceTypeList, );
@@ -110,11 +110,11 @@ TYPED_TEST(MLBInferenceTest, InferencePerLOD) {
         std::fill(outputBuffer.begin(), outputBuffer.end(), 0.0f);
         this->evaluator->calculate(inputInstance.get(), intermediateOutputs.get(), lod);
         const auto& expected = rltests::ml::block4::output::valuesPerLOD[lod];
-        #ifdef RL_USE_HALF_FLOATS
+        #ifdef RL_BUILD_WITH_HALF_FLOATS
             static constexpr float threshold = 0.05f;
         #else
             static constexpr float threshold = 0.0001f;
-        #endif  // RL_USE_HALF_FLOATS
+        #endif  // RL_BUILD_WITH_HALF_FLOATS
         ASSERT_ELEMENTS_NEAR(outputBuffer, expected, expected.size(), threshold);
     }
 }
diff --git a/Engine/Plugins/Animation/RigLogic/Source/RigLogicLibTest/Private/rltests/ml/cpu/TestStorageBuilder.cpp b/Engine/Plugins/Animation/RigLogic/Source/RigLogicLibTest/Private/rltests/ml/cpu/TestStorageBuilder.cpp
--- a/Engine/Plugins/Animation/RigLogic/Source/RigLogicLibTest/Private/rltests/ml/cpu/TestStorageBuilder.cpp
+++ b/Engine/Plugins/Animation/RigLogic/Source/RigLogicLibTest/Private/rltests/ml/cpu/TestStorageBuilder.cpp
@@ -89,39 +89,39 @@ class MLBSStorageBuilderTest : public ::testing::Test {
 
 #if defined(RL_BUILD_WITH_AVX) && defined(RL_BUILD_WITH_SSE)
     using StorageValueTypeList = ::testing::Types<
-    #ifdef RL_USE_HALF_FLOATS
+    #ifdef RL_BUILD_WITH_HALF_FLOATS
             std::tuple<StorageValueType, trimd::avx::F256, trimd::sse::F128>,
             std::tuple<StorageValueType, trimd::sse::F256, trimd::sse::F128>
     #else
             std::tuple<StorageValueType, trimd::avx::F256, trimd::sse::F128>,
             std::tuple<StorageValueType, trimd::sse::F256, trimd::sse::F128>,
             std::tuple<StorageValueType, trimd::scalar::F256, trimd::scalar::F128>
-    #endif  // RL_USE_HALF_FLOATS
+    #endif  // RL_BUILD_WITH_HALF_FLOATS
         >;
 #elif defined(RL_BUILD_WITH_AVX)
     using StorageValueTypeList = ::testing::Types<
-    #ifdef RL_USE_HALF_FLOATS
+    #ifdef RL_BUILD_WITH_HALF_FLOATS
             std::tuple<StorageValueType, trimd::avx::F256, trimd::sse::F128>
     #else
             std::tuple<StorageValueType, trimd::avx::F256, trimd::sse::F128>,
             std::tuple<StorageValueType, trimd::scalar::F256, trimd::scalar::F128>
-    #endif  // RL_USE_HALF_FLOATS
+    #endif  // RL_BUILD_WITH_HALF_FLOATS
         >;
 #elif defined(RL_BUILD_WITH_SSE)
     using StorageValueTypeList = ::testing::Types<
-    #ifdef RL_USE_HALF_FLOATS
+    #ifdef RL_BUILD_WITH_HALF_FLOATS
             std::tuple<StorageValueType, trimd::sse::F256, trimd::sse::F128>
     #else
             std::tuple<StorageValueType, trimd::sse::F256, trimd::sse::F128>,
             std::tuple<StorageValueType, trimd::scalar::F256, trimd::scalar::F128>
-    #endif  // RL_USE_HALF_FLOATS
+    #endif  // RL_BUILD_WITH_HALF_FLOATS
         >;
 #else
-    #ifndef RL_USE_HALF_FLOATS
+    #ifndef RL_BUILD_WITH_HALF_FLOATS
         using StorageValueTypeList = ::testing::Types<std::tuple<StorageValueType, trimd::scalar::F256, trimd::scalar::F128> >;
     #else
         using StorageValueTypeList = ::testing::Types<std::tuple<> >;
-    #endif  // RL_USE_HALF_FLOATS
+    #endif  // RL_BUILD_WITH_HALF_FLOATS
 #endif
 
 TYPED_TEST_SUITE(MLBSStorageBuilderTest, StorageValueTypeList, );
diff --git a/Engine/Plugins/Animation/RigLogic/Source/RigLogicLibTest/Private/rltests/rbf/cpu/TestRBFEvaluator.cpp b/Engine/Plugins/Animation/RigLogic/Source/RigLogicLibTest/Private/rltests/rbf/cpu/TestRBFEvaluator.cpp
--- a/Engine/Plugins/Animation/RigLogic/Source/RigLogicLibTest/Private/rltests/rbf/cpu/TestRBFEvaluator.cpp
+++ b/Engine/Plugins/Animation/RigLogic/Source/RigLogicLibTest/Private/rltests/rbf/cpu/TestRBFEvaluator.cpp
@@ -51,39 +51,39 @@ class RBFBehaviorTest : public ::testing::Test {
 
 #if defined(RL_BUILD_WITH_AVX) && defined(RL_BUILD_WITH_SSE)
     using RBFEvaluatorTypeList = ::testing::Types<
-    #ifdef RL_USE_HALF_FLOATS
+    #ifdef RL_BUILD_WITH_HALF_FLOATS
             std::tuple<StorageValueType, trimd::avx::F256, trimd::sse::F128>,
             std::tuple<StorageValueType, trimd::sse::F256, trimd::sse::F128>
     #else
             std::tuple<StorageValueType, trimd::avx::F256, trimd::sse::F128>,
             std::tuple<StorageValueType, trimd::sse::F256, trimd::sse::F128>,
             std::tuple<StorageValueType, trimd::scalar::F256, trimd::scalar::F128>
-    #endif  // RL_USE_HALF_FLOATS
+    #endif  // RL_BUILD_WITH_HALF_FLOATS
         >;
 #elif defined(RL_BUILD_WITH_AVX)
     using RBFEvaluatorTypeList = ::testing::Types<
-    #ifdef RL_USE_HALF_FLOATS
+    #ifdef RL_BUILD_WITH_HALF_FLOATS
             std::tuple<StorageValueType, trimd::avx::F256, trimd::sse::F128>
     #else
             std::tuple<StorageValueType, trimd::avx::F256, trimd::sse::F128>,
             std::tuple<StorageValueType, trimd::scalar::F256, trimd::scalar::F128>
-    #endif  // RL_USE_HALF_FLOATS
+    #endif  // RL_BUILD_WITH_HALF_FLOATS
         >;
 #elif defined(RL_BUILD_WITH_SSE)
     using RBFEvaluatorTypeList = ::testing::Types<
-    #ifdef RL_USE_HALF_FLOATS
+    #ifdef RL_BUILD_WITH_HALF_FLOATS
             std::tuple<StorageValueType, trimd::sse::F256, trimd::sse::F128>
     #else
             std::tuple<StorageValueType, trimd::sse::F256, trimd::sse::F128>,
             std::tuple<StorageValueType, trimd::scalar::F256, trimd::scalar::F128>
-    #endif  // RL_USE_HALF_FLOATS
+    #endif  // RL_BUILD_WITH_HALF_FLOATS
         >;
 #else
-    #ifndef RL_USE_HALF_FLOATS
+    #ifndef RL_BUILD_WITH_HALF_FLOATS
         using RBFEvaluatorTypeList = ::testing::Types<std::tuple<StorageValueType, trimd::scalar::F256, trimd::scalar::F128> >;
     #else
         using RBFEvaluatorTypeList = ::testing::Types<std::tuple<> >;
-    #endif  // RL_USE_HALF_FLOATS
+    #endif  // RL_BUILD_WITH_HALF_FLOATS
 #endif
 
 TYPED_TEST_SUITE(RBFBehaviorTest, RBFEvaluatorTypeList, );
@@ -113,11 +113,11 @@ TYPED_TEST(RBFBehaviorTest, SolverPerLOD) {
         // Call twice to make sure the output control values are not accumulating between calls
         this->evaluator->calculate(inputInstance.get(), intermediateOutputs.get(), lod);
         const auto& expected = rltests::rbf::output::valuesPerLOD[lod];
-        #ifdef RL_USE_HALF_FLOATS
+        #ifdef RL_BUILD_WITH_HALF_FLOATS
             static constexpr float threshold = 0.05f;
         #else
             static constexpr float threshold = 0.0001f;
-        #endif  // RL_USE_HALF_FLOATS
+        #endif  // RL_BUILD_WITH_HALF_FLOATS
         ASSERT_ELEMENTS_NEAR(outputBuffer, expected, expected.size(), threshold);
     }
 }
diff --git a/Engine/Plugins/Animation/RigLogic/Source/RigLogicLibTest/Private/rltests/rbf/cpu/TestRBFStorageBuilder.cpp b/Engine/Plugins/Animation/RigLogic/Source/RigLogicLibTest/Private/rltests/rbf/cpu/TestRBFStorageBuilder.cpp
--- a/Engine/Plugins/Animation/RigLogic/Source/RigLogicLibTest/Private/rltests/rbf/cpu/TestRBFStorageBuilder.cpp
+++ b/Engine/Plugins/Animation/RigLogic/Source/RigLogicLibTest/Private/rltests/rbf/cpu/TestRBFStorageBuilder.cpp
@@ -86,41 +86,41 @@ class RBFStorageBuilderTest : public ::testing::Test {
 
 #if defined(RL_BUILD_WITH_AVX) && defined(RL_BUILD_WITH_SSE)
     using StorageValueTypeList = ::testing::Types<
-    #ifdef RL_USE_HALF_FLOATS
+    #ifdef RL_BUILD_WITH_HALF_FLOATS
             std::tuple<float, trimd::avx::F256, trimd::sse::F128>,
             std::tuple<float, trimd::sse::F256, trimd::sse::F128>
     #else
             std::tuple<float, trimd::avx::F256, trimd::sse::F128>,
             std::tuple<float, trimd::sse::F256, trimd::sse::F128>,
             std::tuple<float, trimd::scalar::F256, trimd::scalar::F128>
-    #endif  // RL_USE_HALF_FLOATS
+    #endif  // RL_BUILD_WITH_HALF_FLOATS
         >;
 #elif defined(RL_BUILD_WITH_AVX)
     using StorageValueTypeList = ::testing::Types<
-    #ifdef RL_USE_HALF_FLOATS
+    #ifdef RL_BUILD_WITH_HALF_FLOATS
             std::tuple<float, trimd::avx::F256, trimd::sse::F128>
     #else
             std::tuple<float, trimd::avx::F256, trimd::sse::F128>,
             std::tuple<float, trimd::scalar::F256, trimd::scalar::F128>
-    #endif  // RL_USE_HALF_FLOATS
+    #endif  // RL_BUILD_WITH_HALF_FLOATS
         >;
 #elif defined(RL_BUILD_WITH_SSE)
     using StorageValueTypeList = ::testing::Types<
-    #ifdef RL_USE_HALF_FLOATS
+    #ifdef RL_BUILD_WITH_HALF_FLOATS
             std::tuple<float, trimd::sse::F256, trimd::sse::F128>
     #else
             std::tuple<float, trimd::sse::F256, trimd::sse::F128>,
             std::tuple<float, trimd::scalar::F256, trimd::scalar::F128>
-    #endif  // RL_USE_HALF_FLOATS
+    #endif  // RL_BUILD_WITH_HALF_FLOATS
         >;
 #else
-    #ifndef RL_USE_HALF_FLOATS
+    #ifndef RL_BUILD_WITH_HALF_FLOATS
         using StorageValueTypeList = ::testing::Types<
             std::tuple<float, trimd::scalar::F256, trimd::scalar::F128>
             >;
     #else
         using StorageValueTypeList = ::testing::Types<std::tuple<> >;
-    #endif  // RL_USE_HALF_FLOATS
+    #endif  // RL_BUILD_WITH_HALF_FLOATS
 #endif
 
 TYPED_TEST_SUITE(RBFStorageBuilderTest, StorageValueTypeList, );
diff --git a/Engine/Plugins/Animation/RigLogic/Source/RigLogicLibTest/Private/trimdtests/TestPlatform.cpp b/Engine/Plugins/Animation/RigLogic/Source/RigLogicLibTest/Private/trimdtests/TestPlatform.cpp
new file mode 100644
--- /dev/null
+++ b/Engine/Plugins/Animation/RigLogic/Source/RigLogicLibTest/Private/trimdtests/TestPlatform.cpp
@@ -0,0 +1,25 @@
+// Copyright Epic Games, Inc. All Rights Reserved.
+
+#include "trimdtests/Defs.h"
+
+#include "trimd/TRiMD.h"
+
+
+TEST(CPUFeatureDetection, Platform) {
+    trimd::CPUFeatures features = trimd::getCPUFeatures();
+    if (features.NEON) {
+        // If NEON is detected, none of the x86 CPU features should be available
+        ASSERT_FALSE(features.SSE);
+        ASSERT_FALSE(features.SSE2);
+        ASSERT_FALSE(features.SSE3);
+        ASSERT_FALSE(features.SSSE3);
+        ASSERT_FALSE(features.SSE41);
+        ASSERT_FALSE(features.SSE42);
+        ASSERT_FALSE(features.AVX);
+        ASSERT_FALSE(features.F16C);
+    } else {
+        // If NEON is not detected, x86 features might be found (depending on hardware)
+        // But FP16 should not be available without NEON
+        ASSERT_FALSE(features.FP16);
+    }
+}
diff --git a/Engine/Plugins/Animation/RigLogic/Source/RigLogicLibTest/Private/trimdtests/TestT256.cpp b/Engine/Plugins/Animation/RigLogic/Source/RigLogicLibTest/Private/trimdtests/TestT256.cpp
--- a/Engine/Plugins/Animation/RigLogic/Source/RigLogicLibTest/Private/trimdtests/TestT256.cpp
+++ b/Engine/Plugins/Animation/RigLogic/Source/RigLogicLibTest/Private/trimdtests/TestT256.cpp
@@ -767,13 +767,13 @@ TEST(T256Test, AndNotScalar) {
 }
 
 TEST(T256Test, RsqrtScalar) {
-    trimd::scalar::F256 v{1.0f, 2.0f, 3.0f, 9.0f, 9.0f, 3.0f, 2.0f, 1.0,};
+    trimd::scalar::F256 v{1.0f, 2.0f, 3.0f, 9.0f, 9.0f, 3.0f, 2.0f, 1.0};
     v = trimd::scalar::rsqrt(v);
     trimd::scalar::F256 e{1.0f, 0.70710678f, 0.57735026f, 0.33333333f, 0.33333333f, 0.57735026f, 0.70710678f, 1.0f};
     #ifdef TRIMD_ENABLE_FAST_INVERSE_SQRT
-    static constexpr float threshold = 0.0004f;
+        static constexpr float threshold = 0.0004f;
     #else
-    static constexpr float threshold = 0.0002f;
+        static constexpr float threshold = 0.0002f;
     #endif  // TRIMD_ENABLE_FAST_INVERSE_SQRT
     ASSERT_TRUE(near(v, e, threshold));
 }
@@ -844,13 +844,13 @@ TEST(T256Test, RsqrtScalar) {
     }
 
     TEST(T256Test, RsqrtAVX) {
-        trimd::avx::F256 v{1.0f, 2.0f, 3.0f, 9.0f, 9.0f, 3.0f, 2.0f, 1.0,};
+        trimd::avx::F256 v{1.0f, 2.0f, 3.0f, 9.0f, 9.0f, 3.0f, 2.0f, 1.0};
         v = trimd::avx::rsqrt(v);
         trimd::avx::F256 e{1.0f, 0.70710678f, 0.57735026f, 0.33333333f, 0.33333333f, 0.57735026f, 0.70710678f, 1.0f};
         #ifdef TRIMD_ENABLE_FAST_INVERSE_SQRT
-        static constexpr float threshold = 0.0004f;
+            static constexpr float threshold = 0.0004f;
         #else
-        static constexpr float threshold = 0.0003f;
+            static constexpr float threshold = 0.0003f;
         #endif  // TRIMD_ENABLE_FAST_INVERSE_SQRT
         ASSERT_TRUE(near(v, e, threshold));
     }
