Attempt fixing compile error caused by previous CL

diff --git a/Engine/Plugins/Animation/RigLogic/Source/RigLogicLib/RigLogicLib.Build.cs b/Engine/Plugins/Animation/RigLogic/Source/RigLogicLib/RigLogicLib.Build.cs
--- a/Engine/Plugins/Animation/RigLogic/Source/RigLogicLib/RigLogicLib.Build.cs
+++ b/Engine/Plugins/Animation/RigLogic/Source/RigLogicLib/RigLogicLib.Build.cs
@@ -54,7 +54,10 @@ namespace UnrealBuildTool.Rules
 			}
 
 			PrivateDefinitions.Add("RL_AUTODETECT_SSE=1");
-			PrivateDefinitions.Add("RL_AUTODETECT_HALF_FLOATS=1");
+			if (!Target.bUseAutoRTFMCompiler)
+			{
+				PrivateDefinitions.Add("RL_AUTODETECT_HALF_FLOATS=1");
+			}
 
 			bDisableAutoRTFMInstrumentation = true;
 		}
