Revise half-float detection logic

diff --git a/Engine/Plugins/Animation/RigLogic/Source/RigLogicLib/Private/riglogic/system/simd/Detect.h b/Engine/Plugins/Animation/RigLogic/Source/RigLogicLib/Private/riglogic/system/simd/Detect.h
--- a/Engine/Plugins/Animation/RigLogic/Source/RigLogicLib/Private/riglogic/system/simd/Detect.h
+++ b/Engine/Plugins/Animation/RigLogic/Source/RigLogicLib/Private/riglogic/system/simd/Detect.h
@@ -9,8 +9,7 @@
     #if defined(TRIMD_PLATFORM_X86)
         #define RL_BUILD_WITH_SSE 1
         #if defined(RL_AUTODETECT_HALF_FLOATS) && !defined(RL_BUILD_WITH_HALF_FLOATS)
-            #if defined(__F16C__) || (defined(_MSC_VER) && (_MSC_VER >= 1700) && defined(__AVX__)) || \
-            (defined(__INTEL_COMPILER) && defined(__AVX__))
+            #if defined(__F16C__)
                 #define RL_BUILD_WITH_HALF_FLOATS 1
             #endif
         #endif  // RL_AUTODETECT_HALF_FLOATS
diff --git a/Engine/Plugins/Animation/RigLogic/Source/RigLogicLib/RigLogicLib.Build.cs b/Engine/Plugins/Animation/RigLogic/Source/RigLogicLib/RigLogicLib.Build.cs
--- a/Engine/Plugins/Animation/RigLogic/Source/RigLogicLib/RigLogicLib.Build.cs
+++ b/Engine/Plugins/Animation/RigLogic/Source/RigLogicLib/RigLogicLib.Build.cs
@@ -54,10 +54,7 @@ namespace UnrealBuildTool.Rules
 			}
 
 			PrivateDefinitions.Add("RL_AUTODETECT_SSE=1");
-			if (!Target.bUseAutoRTFMCompiler)
-			{
-				PrivateDefinitions.Add("RL_AUTODETECT_HALF_FLOATS=1");
-			}
+			PrivateDefinitions.Add("RL_AUTODETECT_HALF_FLOATS=1");
 
 			bDisableAutoRTFMInstrumentation = true;
 		}
