Optimize RigLogic for low LOD evaluation

diff --git a/Engine/Plugins/Animation/RigLogic/RigLogic.uplugin b/Engine/Plugins/Animation/RigLogic/RigLogic.uplugin
--- a/Engine/Plugins/Animation/RigLogic/RigLogic.uplugin
+++ b/Engine/Plugins/Animation/RigLogic/RigLogic.uplugin
@@ -1,9 +1,9 @@
 {
 	"FileVersion": 3,
-	"Version": 10.2,
+	"Version": 10.3,
 	"VersionName": "10.2.7",
-	"FriendlyName": "RigLogic Plugin v10.2.7",
-	"Description": "3Lateral RigLogic Plugin for Facial Animation v10.2.7",
+	"FriendlyName": "RigLogic Plugin v10.3.0",
+	"Description": "3Lateral RigLogic Plugin for Facial Animation v10.3.0",
 	"Category": "Animation",
 	"CreatedBy": "3Lateral",
 	"CreatedByURL": "https://www.3lateral.com/",
diff --git a/Engine/Plugins/Animation/RigLogic/Source/RigLogicLib/Private/riglogic/animatedmaps/AnimatedMaps.h b/Engine/Plugins/Animation/RigLogic/Source/RigLogicLib/Private/riglogic/animatedmaps/AnimatedMaps.h
--- a/Engine/Plugins/Animation/RigLogic/Source/RigLogicLib/Private/riglogic/animatedmaps/AnimatedMaps.h
+++ b/Engine/Plugins/Animation/RigLogic/Source/RigLogicLib/Private/riglogic/animatedmaps/AnimatedMaps.h
@@ -20,6 +20,7 @@ class AnimatedMaps {
 
     public:
         virtual AnimatedMapsOutputInstance::Pointer createInstance(MemoryResource* memRes) const = 0;
+        virtual ConstArrayView<std::uint16_t> getAnimatedMapIndicesForLOD(std::uint16_t  /*unused*/) const = 0;
         virtual void calculate(const ControlsInputInstance* inputs, AnimatedMapsOutputInstance* outputs,
                                std::uint16_t lod) const = 0;
         virtual void load(terse::BinaryInputArchive<BoundedIOStream>& archive) = 0;
diff --git a/Engine/Plugins/Animation/RigLogic/Source/RigLogicLib/Private/riglogic/animatedmaps/AnimatedMapsFactory.cpp b/Engine/Plugins/Animation/RigLogic/Source/RigLogicLib/Private/riglogic/animatedmaps/AnimatedMapsFactory.cpp
--- a/Engine/Plugins/Animation/RigLogic/Source/RigLogicLib/Private/riglogic/animatedmaps/AnimatedMapsFactory.cpp
+++ b/Engine/Plugins/Animation/RigLogic/Source/RigLogicLib/Private/riglogic/animatedmaps/AnimatedMapsFactory.cpp
@@ -8,6 +8,7 @@
 #include "riglogic/animatedmaps/AnimatedMapsNull.h"
 #include "riglogic/animatedmaps/AnimatedMapsOutputInstance.h"
 #include "riglogic/conditionaltable/ConditionalTable.h"
+#include "riglogic/controls/Controls.h"
 #include "riglogic/riglogic/Configuration.h"
 #include "riglogic/riglogic/RigMetrics.h"
 #include "riglogic/utils/Extd.h"
@@ -38,6 +39,7 @@ AnimatedMaps::Pointer AnimatedMapsFactory::create(const Configuration& config, c
 
 AnimatedMaps::Pointer AnimatedMapsFactory::create(const Configuration& config,
                                                   const dna::BehaviorReader* reader,
+                                                  Controls* controls,
                                                   MemoryResource* memRes) {
     if (!config.loadAnimatedMaps || (reader->getAnimatedMapCount() == 0u)) {
         return UniqueInstance<AnimatedMapsNull, AnimatedMaps>::with(memRes).create();
@@ -67,6 +69,11 @@ AnimatedMaps::Pointer AnimatedMapsFactory::create(const Configuration& config,
         }
     }
 
+    for (std::uint16_t lod = {}; lod < static_cast<std::uint16_t>(lods.size()); ++lod) {
+        ConstArrayView<std::uint16_t> inputIndicesForLOD(inputIndices.data(), lods[lod]);
+        controls->registerControls(lod, inputIndicesForLOD);
+    }
+
     const auto inputCount = static_cast<std::uint16_t>(reader->getRawControlCount() + reader->getPSDCount());
     const auto outputCount = reader->getAnimatedMapCount();
     ConditionalTable conditionals{std::move(inputIndices),
diff --git a/Engine/Plugins/Animation/RigLogic/Source/RigLogicLib/Private/riglogic/animatedmaps/AnimatedMapsFactory.h b/Engine/Plugins/Animation/RigLogic/Source/RigLogicLib/Private/riglogic/animatedmaps/AnimatedMapsFactory.h
--- a/Engine/Plugins/Animation/RigLogic/Source/RigLogicLib/Private/riglogic/animatedmaps/AnimatedMapsFactory.h
+++ b/Engine/Plugins/Animation/RigLogic/Source/RigLogicLib/Private/riglogic/animatedmaps/AnimatedMapsFactory.h
@@ -9,9 +9,13 @@ namespace rl4 {
 
 struct Configuration;
 struct RigMetrics;
+class Controls;
 
 struct AnimatedMapsFactory {
-    static AnimatedMaps::Pointer create(const Configuration& config, const dna::BehaviorReader* reader, MemoryResource* memRes);
+    static AnimatedMaps::Pointer create(const Configuration& config,
+                                        const dna::BehaviorReader* reader,
+                                        Controls* controls,
+                                        MemoryResource* memRes);
     static AnimatedMaps::Pointer create(const Configuration& config, const RigMetrics& metrics, MemoryResource* memRes);
 };
 
diff --git a/Engine/Plugins/Animation/RigLogic/Source/RigLogicLib/Private/riglogic/animatedmaps/AnimatedMapsImpl.cpp b/Engine/Plugins/Animation/RigLogic/Source/RigLogicLib/Private/riglogic/animatedmaps/AnimatedMapsImpl.cpp
--- a/Engine/Plugins/Animation/RigLogic/Source/RigLogicLib/Private/riglogic/animatedmaps/AnimatedMapsImpl.cpp
+++ b/Engine/Plugins/Animation/RigLogic/Source/RigLogicLib/Private/riglogic/animatedmaps/AnimatedMapsImpl.cpp
@@ -24,6 +24,12 @@ AnimatedMapsOutputInstance::Pointer AnimatedMapsImpl::createInstance(MemoryResou
     return instanceFactory(instanceMemRes);
 }
 
+ConstArrayView<std::uint16_t> AnimatedMapsImpl::getAnimatedMapIndicesForLOD(std::uint16_t lod) const {
+    assert(lod < lods.size());
+    const auto outputIndices = conditionals.getOutputIndices();
+    return outputIndices.subview(0ul, lods[lod]);
+}
+
 void AnimatedMapsImpl::calculate(const ControlsInputInstance* inputs, AnimatedMapsOutputInstance* outputs,
                                  std::uint16_t lod) const {
     assert(lod < lods.size());
diff --git a/Engine/Plugins/Animation/RigLogic/Source/RigLogicLib/Private/riglogic/animatedmaps/AnimatedMapsImpl.h b/Engine/Plugins/Animation/RigLogic/Source/RigLogicLib/Private/riglogic/animatedmaps/AnimatedMapsImpl.h
--- a/Engine/Plugins/Animation/RigLogic/Source/RigLogicLib/Private/riglogic/animatedmaps/AnimatedMapsImpl.h
+++ b/Engine/Plugins/Animation/RigLogic/Source/RigLogicLib/Private/riglogic/animatedmaps/AnimatedMapsImpl.h
@@ -18,8 +18,8 @@ class AnimatedMapsImpl : public AnimatedMaps {
         AnimatedMapsImpl(Vector<std::uint16_t>&& lods_,
                          ConditionalTable&& conditionals_,
                          AnimatedMapsOutputInstance::Factory instanceFactory_);
-
         AnimatedMapsOutputInstance::Pointer createInstance(MemoryResource* instanceMemRes) const override;
+        ConstArrayView<std::uint16_t> getAnimatedMapIndicesForLOD(std::uint16_t lod) const override;
         void calculate(const ControlsInputInstance* inputs, AnimatedMapsOutputInstance* outputs,
                        std::uint16_t lod) const override;
         void load(terse::BinaryInputArchive<BoundedIOStream>& archive) override;
diff --git a/Engine/Plugins/Animation/RigLogic/Source/RigLogicLib/Private/riglogic/animatedmaps/AnimatedMapsNull.cpp b/Engine/Plugins/Animation/RigLogic/Source/RigLogicLib/Private/riglogic/animatedmaps/AnimatedMapsNull.cpp
--- a/Engine/Plugins/Animation/RigLogic/Source/RigLogicLib/Private/riglogic/animatedmaps/AnimatedMapsNull.cpp
+++ b/Engine/Plugins/Animation/RigLogic/Source/RigLogicLib/Private/riglogic/animatedmaps/AnimatedMapsNull.cpp
@@ -10,6 +10,10 @@ AnimatedMapsOutputInstance::Pointer AnimatedMapsNull::createInstance(MemoryResou
     return UniqueInstance<AnimatedMapsNullOutputInstance, AnimatedMapsOutputInstance>::with(instanceMemRes).create();
 }
 
+ConstArrayView<std::uint16_t> AnimatedMapsNull::getAnimatedMapIndicesForLOD(std::uint16_t  /*unused*/) const {
+    return {};
+}
+
 void AnimatedMapsNull::calculate(const ControlsInputInstance*  /*unused*/, AnimatedMapsOutputInstance*  /*unused*/,
                                  std::uint16_t  /*unused*/) const {
 }
diff --git a/Engine/Plugins/Animation/RigLogic/Source/RigLogicLib/Private/riglogic/animatedmaps/AnimatedMapsNull.h b/Engine/Plugins/Animation/RigLogic/Source/RigLogicLib/Private/riglogic/animatedmaps/AnimatedMapsNull.h
--- a/Engine/Plugins/Animation/RigLogic/Source/RigLogicLib/Private/riglogic/animatedmaps/AnimatedMapsNull.h
+++ b/Engine/Plugins/Animation/RigLogic/Source/RigLogicLib/Private/riglogic/animatedmaps/AnimatedMapsNull.h
@@ -16,6 +16,7 @@ class ControlsInputInstance;
 class AnimatedMapsNull : public AnimatedMaps {
     public:
         AnimatedMapsOutputInstance::Pointer createInstance(MemoryResource* instanceMemRes) const override;
+        ConstArrayView<std::uint16_t> getAnimatedMapIndicesForLOD(std::uint16_t  /*unused*/) const override;
         void calculate(const ControlsInputInstance*  /*unused*/, AnimatedMapsOutputInstance*  /*unused*/,
                        std::uint16_t  /*unused*/) const override;
         void load(terse::BinaryInputArchive<BoundedIOStream>&  /*unused*/) override;
diff --git a/Engine/Plugins/Animation/RigLogic/Source/RigLogicLib/Private/riglogic/blendshapes/BlendShapes.h b/Engine/Plugins/Animation/RigLogic/Source/RigLogicLib/Private/riglogic/blendshapes/BlendShapes.h
--- a/Engine/Plugins/Animation/RigLogic/Source/RigLogicLib/Private/riglogic/blendshapes/BlendShapes.h
+++ b/Engine/Plugins/Animation/RigLogic/Source/RigLogicLib/Private/riglogic/blendshapes/BlendShapes.h
@@ -20,6 +20,7 @@ class BlendShapes {
 
     public:
         virtual BlendShapesOutputInstance::Pointer createInstance(MemoryResource* instanceMemRes) const = 0;
+        virtual ConstArrayView<std::uint16_t> getBlendShapeChannelIndicesForLOD(std::uint16_t lod) const = 0;
         virtual void calculate(const ControlsInputInstance* inputs, BlendShapesOutputInstance* outputs,
                                std::uint16_t lod) const = 0;
         virtual void load(terse::BinaryInputArchive<BoundedIOStream>& archive) = 0;
diff --git a/Engine/Plugins/Animation/RigLogic/Source/RigLogicLib/Private/riglogic/blendshapes/BlendShapesFactory.cpp b/Engine/Plugins/Animation/RigLogic/Source/RigLogicLib/Private/riglogic/blendshapes/BlendShapesFactory.cpp
--- a/Engine/Plugins/Animation/RigLogic/Source/RigLogicLib/Private/riglogic/blendshapes/BlendShapesFactory.cpp
+++ b/Engine/Plugins/Animation/RigLogic/Source/RigLogicLib/Private/riglogic/blendshapes/BlendShapesFactory.cpp
@@ -8,6 +8,7 @@
 #include "riglogic/blendshapes/BlendShapesImplOutputInstance.h"
 #include "riglogic/blendshapes/BlendShapesNull.h"
 #include "riglogic/blendshapes/BlendShapesOutputInstance.h"
+#include "riglogic/controls/Controls.h"
 #include "riglogic/riglogic/Configuration.h"
 #include "riglogic/riglogic/RigMetrics.h"
 
@@ -46,6 +47,7 @@ BlendShapes::Pointer BlendShapesFactory::create(const Configuration& config, con
 
 BlendShapes::Pointer BlendShapesFactory::create(const Configuration& config,
                                                 const dna::BehaviorReader* reader,
+                                                Controls* controls,
                                                 MemoryResource* memRes) {
     if (!config.loadBlendShapes || (reader->getBlendShapeChannelCount() == 0u)) {
         return UniqueInstance<BlendShapesNull, BlendShapes>::with(memRes).create();
@@ -59,6 +61,11 @@ BlendShapes::Pointer BlendShapesFactory::create(const Configuration& config,
     extd::copy(reader->getBlendShapeChannelInputIndices(), inputIndices);
     extd::copy(reader->getBlendShapeChannelOutputIndices(), outputIndices);
 
+    for (std::uint16_t lod = {}; lod < static_cast<std::uint16_t>(lods.size()); ++lod) {
+        ConstArrayView<std::uint16_t> inputIndicesForLOD(inputIndices.data(), lods[lod]);
+        controls->registerControls(lod, inputIndicesForLOD);
+    }
+
     auto instanceFactory = createBlendShapesOutputInstanceFactory(config, reader->getBlendShapeChannelCount());
     auto moduleFactory = UniqueInstance<BlendShapesImpl, BlendShapes>::with(memRes);
     return moduleFactory.create(std::move(lods), std::move(inputIndices), std::move(outputIndices), instanceFactory);
diff --git a/Engine/Plugins/Animation/RigLogic/Source/RigLogicLib/Private/riglogic/blendshapes/BlendShapesFactory.h b/Engine/Plugins/Animation/RigLogic/Source/RigLogicLib/Private/riglogic/blendshapes/BlendShapesFactory.h
--- a/Engine/Plugins/Animation/RigLogic/Source/RigLogicLib/Private/riglogic/blendshapes/BlendShapesFactory.h
+++ b/Engine/Plugins/Animation/RigLogic/Source/RigLogicLib/Private/riglogic/blendshapes/BlendShapesFactory.h
@@ -9,9 +9,13 @@ namespace rl4 {
 
 struct Configuration;
 struct RigMetrics;
+class Controls;
 
 struct BlendShapesFactory {
-    static BlendShapes::Pointer create(const Configuration& config, const dna::BehaviorReader* reader, MemoryResource* memRes);
+    static BlendShapes::Pointer create(const Configuration& config,
+                                       const dna::BehaviorReader* reader,
+                                       Controls* controls,
+                                       MemoryResource* memRes);
     static BlendShapes::Pointer create(const Configuration& config, const RigMetrics& metrics, MemoryResource* memRes);
 };
 
diff --git a/Engine/Plugins/Animation/RigLogic/Source/RigLogicLib/Private/riglogic/blendshapes/BlendShapesImpl.cpp b/Engine/Plugins/Animation/RigLogic/Source/RigLogicLib/Private/riglogic/blendshapes/BlendShapesImpl.cpp
--- a/Engine/Plugins/Animation/RigLogic/Source/RigLogicLib/Private/riglogic/blendshapes/BlendShapesImpl.cpp
+++ b/Engine/Plugins/Animation/RigLogic/Source/RigLogicLib/Private/riglogic/blendshapes/BlendShapesImpl.cpp
@@ -25,6 +25,11 @@ BlendShapesOutputInstance::Pointer BlendShapesImpl::createInstance(MemoryResourc
     return instanceFactory(instanceMemRes);
 }
 
+ConstArrayView<std::uint16_t> BlendShapesImpl::getBlendShapeChannelIndicesForLOD(std::uint16_t lod) const {
+    assert(lod < lods.size());
+    return {outputIndices.data(), lods[lod]};
+}
+
 void BlendShapesImpl::calculate(const ControlsInputInstance* inputs, BlendShapesOutputInstance* outputs,
                                 std::uint16_t lod) const {
     assert(lod < lods.size());
diff --git a/Engine/Plugins/Animation/RigLogic/Source/RigLogicLib/Private/riglogic/blendshapes/BlendShapesImpl.h b/Engine/Plugins/Animation/RigLogic/Source/RigLogicLib/Private/riglogic/blendshapes/BlendShapesImpl.h
--- a/Engine/Plugins/Animation/RigLogic/Source/RigLogicLib/Private/riglogic/blendshapes/BlendShapesImpl.h
+++ b/Engine/Plugins/Animation/RigLogic/Source/RigLogicLib/Private/riglogic/blendshapes/BlendShapesImpl.h
@@ -18,8 +18,8 @@ class BlendShapesImpl : public BlendShapes {
                         Vector<std::uint16_t>&& inputIndices_,
                         Vector<std::uint16_t>&& outputIndices_,
                         BlendShapesOutputInstance::Factory instanceFactory_);
-
         BlendShapesOutputInstance::Pointer createInstance(MemoryResource* instanceMemRes) const override;
+        ConstArrayView<std::uint16_t> getBlendShapeChannelIndicesForLOD(std::uint16_t lod) const override;
         void calculate(const ControlsInputInstance* inputs, BlendShapesOutputInstance* outputs, std::uint16_t lod) const override;
         void load(terse::BinaryInputArchive<BoundedIOStream>& archive) override;
         void save(terse::BinaryOutputArchive<BoundedIOStream>& archive) override;
diff --git a/Engine/Plugins/Animation/RigLogic/Source/RigLogicLib/Private/riglogic/blendshapes/BlendShapesNull.cpp b/Engine/Plugins/Animation/RigLogic/Source/RigLogicLib/Private/riglogic/blendshapes/BlendShapesNull.cpp
--- a/Engine/Plugins/Animation/RigLogic/Source/RigLogicLib/Private/riglogic/blendshapes/BlendShapesNull.cpp
+++ b/Engine/Plugins/Animation/RigLogic/Source/RigLogicLib/Private/riglogic/blendshapes/BlendShapesNull.cpp
@@ -10,6 +10,10 @@ BlendShapesOutputInstance::Pointer BlendShapesNull::createInstance(MemoryResourc
     return UniqueInstance<BlendShapesNullOutputInstance, BlendShapesOutputInstance>::with(instanceMemRes).create();
 }
 
+ConstArrayView<std::uint16_t> BlendShapesNull::getBlendShapeChannelIndicesForLOD(std::uint16_t  /*unused*/) const {
+    return {};
+}
+
 void BlendShapesNull::calculate(const ControlsInputInstance*  /*unused*/, BlendShapesOutputInstance*  /*unused*/,
                                 std::uint16_t  /*unused*/) const {
 }
diff --git a/Engine/Plugins/Animation/RigLogic/Source/RigLogicLib/Private/riglogic/blendshapes/BlendShapesNull.h b/Engine/Plugins/Animation/RigLogic/Source/RigLogicLib/Private/riglogic/blendshapes/BlendShapesNull.h
--- a/Engine/Plugins/Animation/RigLogic/Source/RigLogicLib/Private/riglogic/blendshapes/BlendShapesNull.h
+++ b/Engine/Plugins/Animation/RigLogic/Source/RigLogicLib/Private/riglogic/blendshapes/BlendShapesNull.h
@@ -15,6 +15,7 @@ class ControlsInputInstance;
 class BlendShapesNull : public BlendShapes {
     public:
         BlendShapesOutputInstance::Pointer createInstance(MemoryResource* instanceMemRes) const override;
+        ConstArrayView<std::uint16_t> getBlendShapeChannelIndicesForLOD(std::uint16_t  /*unused*/) const override;
         void calculate(const ControlsInputInstance*  /*unused*/, BlendShapesOutputInstance*  /*unused*/,
                        std::uint16_t  /*unused*/) const override;
         void load(terse::BinaryInputArchive<BoundedIOStream>&  /*unused*/) override;
diff --git a/Engine/Plugins/Animation/RigLogic/Source/RigLogicLib/Private/riglogic/conditionaltable/ConditionalTable.cpp b/Engine/Plugins/Animation/RigLogic/Source/RigLogicLib/Private/riglogic/conditionaltable/ConditionalTable.cpp
--- a/Engine/Plugins/Animation/RigLogic/Source/RigLogicLib/Private/riglogic/conditionaltable/ConditionalTable.cpp
+++ b/Engine/Plugins/Animation/RigLogic/Source/RigLogicLib/Private/riglogic/conditionaltable/ConditionalTable.cpp
@@ -116,6 +116,14 @@ std::uint16_t ConditionalTable::getOutputCount() const {
     return outputCount;
 }
 
+ConstArrayView<std::uint16_t> ConditionalTable::getInputIndices() const {
+    return inputIndices;
+}
+
+ConstArrayView<std::uint16_t> ConditionalTable::getOutputIndices() const {
+    return outputIndices;
+}
+
 void ConditionalTable::calculateForward(const float* inputs, float* outputs, std::uint16_t rowCount) const {
     std::fill_n(outputs, outputCount, 0.0f);
 
diff --git a/Engine/Plugins/Animation/RigLogic/Source/RigLogicLib/Private/riglogic/conditionaltable/ConditionalTable.h b/Engine/Plugins/Animation/RigLogic/Source/RigLogicLib/Private/riglogic/conditionaltable/ConditionalTable.h
--- a/Engine/Plugins/Animation/RigLogic/Source/RigLogicLib/Private/riglogic/conditionaltable/ConditionalTable.h
+++ b/Engine/Plugins/Animation/RigLogic/Source/RigLogicLib/Private/riglogic/conditionaltable/ConditionalTable.h
@@ -75,6 +75,8 @@ class ConditionalTable {
         std::uint16_t getRowCount() const;
         std::uint16_t getInputCount() const;
         std::uint16_t getOutputCount() const;
+        ConstArrayView<std::uint16_t> getInputIndices() const;
+        ConstArrayView<std::uint16_t> getOutputIndices() const;
         void calculateForward(const float* inputs, float* outputs) const;
         void calculateForward(const float* inputs, float* outputs, std::uint16_t rowCount) const;
         void calculateReverse(float* inputs, const float* outputs) const;
diff --git a/Engine/Plugins/Animation/RigLogic/Source/RigLogicLib/Private/riglogic/controls/Controls.cpp b/Engine/Plugins/Animation/RigLogic/Source/RigLogicLib/Private/riglogic/controls/Controls.cpp
--- a/Engine/Plugins/Animation/RigLogic/Source/RigLogicLib/Private/riglogic/controls/Controls.cpp
+++ b/Engine/Plugins/Animation/RigLogic/Source/RigLogicLib/Private/riglogic/controls/Controls.cpp
@@ -4,7 +4,7 @@
 
 #include "riglogic/conditionaltable/ConditionalTable.h"
 #include "riglogic/controls/ControlsInputInstance.h"
-#include "riglogic/psdmatrix/PSDMatrix.h"
+#include "riglogic/controls/psdnet/PSDNet.h"
 
 #include <cassert>
 #include <cstdint>
@@ -12,8 +12,8 @@
 namespace rl4 {
 
 Controls::Controls(ConditionalTable&& guiToRawMapping_,
-                   PSDMatrix&& psds_,
-                   Vector<float>&& initialValues_,
+                   PSDNet&& psds_,
+                   Vector<ControlInitializer>&& initialValues_,
                    ControlsInputInstance::Factory instanceFactory_) :
     guiToRawMapping{std::move(guiToRawMapping_)},
     psds{std::move(psds_)},
@@ -25,6 +25,14 @@ ControlsInputInstance::Pointer Controls::createInstance(MemoryResource* instance
     return instanceFactory(initialValues, instanceMemRes);
 }
 
+void Controls::registerControls(std::uint16_t lod, ConstArrayView<std::uint16_t> controlIndices) {
+    psds.registerControls(lod, controlIndices);
+}
+
+ConstArrayView<std::uint16_t> Controls::getPSDIndicesForLOD(std::uint16_t lod) const {
+    return psds.getPSDOutputIndicesForLOD(lod);
+}
+
 void Controls::mapGUIToRaw(ControlsInputInstance* instance) const {
     auto guiControlBuffer = instance->getGUIControlBuffer();
     assert(guiControlBuffer.size() == guiToRawMapping.getInputCount());
@@ -39,8 +47,8 @@ void Controls::mapRawToGUI(ControlsInputInstance* instance) const {
     guiToRawMapping.calculateReverse(guiControlBuffer.data(), inputBuffer.data());
 }
 
-void Controls::calculate(ControlsInputInstance* instance) const {
-    psds.calculate(instance->getInputBuffer().data(), instance->getRawControlCount());
+void Controls::calculate(ControlsInputInstance* instance, std::uint16_t lod) const {
+    psds.calculate(instance->getInputBuffer(), instance->getClampBuffer(), lod);
 }
 
 }  // namespace rl4
diff --git a/Engine/Plugins/Animation/RigLogic/Source/RigLogicLib/Private/riglogic/controls/Controls.h b/Engine/Plugins/Animation/RigLogic/Source/RigLogicLib/Private/riglogic/controls/Controls.h
--- a/Engine/Plugins/Animation/RigLogic/Source/RigLogicLib/Private/riglogic/controls/Controls.h
+++ b/Engine/Plugins/Animation/RigLogic/Source/RigLogicLib/Private/riglogic/controls/Controls.h
@@ -5,7 +5,7 @@
 #include "riglogic/TypeDefs.h"
 #include "riglogic/conditionaltable/ConditionalTable.h"
 #include "riglogic/controls/ControlsInputInstance.h"
-#include "riglogic/psdmatrix/PSDMatrix.h"
+#include "riglogic/controls/psdnet/PSDNet.h"
 
 #include <cstdint>
 
@@ -17,14 +17,16 @@ class Controls {
 
     public:
         Controls(ConditionalTable&& guiToRawMapping_,
-                 PSDMatrix&& psds_,
-                 Vector<float>&& initialValues_,
+                 PSDNet&& psds_,
+                 Vector<ControlInitializer>&& initialValues_,
                  ControlsInputInstance::Factory instanceFactory_);
 
         ControlsInputInstance::Pointer createInstance(MemoryResource* instanceMemRes) const;
+        void registerControls(std::uint16_t lod, ConstArrayView<std::uint16_t> controlIndices);
+        ConstArrayView<std::uint16_t> getPSDIndicesForLOD(std::uint16_t lod) const;
         void mapGUIToRaw(ControlsInputInstance* instance) const;
         void mapRawToGUI(ControlsInputInstance* instance) const;
-        void calculate(ControlsInputInstance* instance) const;
+        void calculate(ControlsInputInstance* instance, std::uint16_t lod) const;
 
         template<class Archive>
         void serialize(Archive& archive) {
@@ -33,8 +35,8 @@ class Controls {
 
     private:
         ConditionalTable guiToRawMapping;
-        PSDMatrix psds;
-        Vector<float> initialValues;
+        PSDNet psds;
+        Vector<ControlInitializer> initialValues;
         ControlsInputInstance::Factory instanceFactory;
 
 };
diff --git a/Engine/Plugins/Animation/RigLogic/Source/RigLogicLib/Private/riglogic/controls/ControlsFactory.cpp b/Engine/Plugins/Animation/RigLogic/Source/RigLogicLib/Private/riglogic/controls/ControlsFactory.cpp
--- a/Engine/Plugins/Animation/RigLogic/Source/RigLogicLib/Private/riglogic/controls/ControlsFactory.cpp
+++ b/Engine/Plugins/Animation/RigLogic/Source/RigLogicLib/Private/riglogic/controls/ControlsFactory.cpp
@@ -6,7 +6,7 @@
 #include "riglogic/conditionaltable/ConditionalTable.h"
 #include "riglogic/controls/Controls.h"
 #include "riglogic/controls/instances/StandardControlsInputInstance.h"
-#include "riglogic/psdmatrix/PSDMatrix.h"
+#include "riglogic/controls/psdnet/PSDNet.h"
 #include "riglogic/riglogic/Configuration.h"
 #include "riglogic/riglogic/RigMetrics.h"
 #include "riglogic/utils/Extd.h"
@@ -15,38 +15,7 @@
 
 namespace rl4 {
 
-static ControlsInputInstance::Factory createInstanceFactory(const Configuration&  /*unused*/,
-                                                            std::uint16_t guiControlCount,
-                                                            std::uint16_t rawControlCount,
-                                                            std::uint16_t psdControlCount,
-                                                            std::uint16_t mlControlCount,
-                                                            std::uint16_t rbfControlCount) {
-    return [ = ](ConstArrayView<float> initialValues, MemoryResource* memRes) {
-               auto factory = UniqueInstance<StandardControlsInputInstance, ControlsInputInstance>::with(memRes);
-               return factory.create(guiControlCount,
-                                     rawControlCount,
-                                     psdControlCount,
-                                     mlControlCount,
-                                     rbfControlCount,
-                                     initialValues,
-                                     memRes);
-    };
-}
-
-Controls::Pointer ControlsFactory::create(const Configuration& config, const RigMetrics& metrics, MemoryResource* memRes) {
-    auto instanceFactory = createInstanceFactory(config,
-                                                 metrics.guiControlCount,
-                                                 metrics.rawControlCount,
-                                                 metrics.psdControlCount,
-                                                 metrics.mlControlCount,
-                                                 metrics.rbfControlCount);
-    return UniqueInstance<Controls>::with(memRes).create(ConditionalTable{memRes},
-                                                         PSDMatrix{memRes},
-                                                         Vector<float>{memRes},
-                                                         instanceFactory);
-}
-
-Controls::Pointer ControlsFactory::create(const Configuration& config, const dna::Reader* reader, MemoryResource* memRes) {
+static ConditionalTable createConditionalTable(const dna::Reader* reader, MemoryResource* memRes) {
     Vector<std::uint16_t> inputIndices{memRes};
     Vector<std::uint16_t> outputIndices{memRes};
     Vector<float> fromValues{memRes};
@@ -70,31 +39,73 @@
 
     const auto guiControlCount = reader->getGUIControlCount();
     const auto rawControlCount = reader->getRawControlCount();
-    const auto psdControlCount = reader->getPSDCount();
-    const auto mlControlCount = reader->getMLControlCount();
-    const auto rbfControlCount = reader->getRBFPoseControlCount();
 
-    ConditionalTable conditionals{std::move(inputIndices),
-                                  std::move(outputIndices),
-                                  std::move(fromValues),
-                                  std::move(toValues),
-                                  std::move(slopeValues),
-                                  std::move(cutValues),
-                                  guiControlCount,
-                                  rawControlCount,
-                                  memRes};
+    return ConditionalTable{std::move(inputIndices),
+                            std::move(outputIndices),
+                            std::move(fromValues),
+                            std::move(toValues),
+                            std::move(slopeValues),
+                            std::move(cutValues),
+                            guiControlCount,
+                            rawControlCount,
+                            memRes};
+}
 
-    Vector<std::uint16_t> psdRows{memRes};
-    Vector<std::uint16_t> psdColumns{memRes};
-    Vector<float> psdValues{memRes};
-    extd::copy(reader->getPSDRowIndices(), psdRows);
-    extd::copy(reader->getPSDColumnIndices(), psdColumns);
-    extd::copy(reader->getPSDValues(), psdValues);
-    PSDMatrix psds{psdControlCount,
-                   std::move(psdRows),
-                   std::move(psdColumns),
-                   std::move(psdValues)};
+static PSDNet createPSDNet(const dna::Reader* reader, MemoryResource* memRes) {
+    Matrix<std::uint16_t> inputLODs{memRes};
+    Matrix<std::uint16_t> outputLODs{memRes};
+    Vector<PSD> psds{memRes};
+    Vector<std::uint16_t> inputIndicesPerPSD{memRes};
 
+    const auto lodCount = reader->getLODCount();
+    const auto rawControlCount = reader->getRawControlCount();
+    const auto psdCount = reader->getPSDCount();
+    const auto psdRows = reader->getPSDRowIndices();
+    const auto psdCols = reader->getPSDColumnIndices();
+    const auto psdWeights = reader->getPSDValues();
+    auto minPSD = rawControlCount;
+    auto maxPSD = static_cast<std::uint16_t>(0);
+    if (psdCount > 0u) {
+        maxPSD = static_cast<std::uint16_t>(static_cast<std::uint32_t>(minPSD) + static_cast<std::uint32_t>(psdCount) - 1u);
+    }
+
+    inputLODs.resize(lodCount);
+    outputLODs.resize(lodCount);
+    for (auto& indices : inputLODs) {
+        indices.reserve(rawControlCount);
+    }
+    for (auto& indices : outputLODs) {
+        indices.reserve(psdCount);
+    }
+    psds.resize(psdCount);
+    inputIndicesPerPSD.reserve(psdCols.size());
+    for (std::size_t start = {}; start < psdRows.size(); ++start) {
+        const auto psdOutputIndex = psdRows[start];
+        PSD& psd = psds[static_cast<std::size_t>(psdOutputIndex) - static_cast<std::size_t>(rawControlCount)];
+        if (psd.size != 0ul) {
+            continue;
+        }
+
+        psd.weight = 1.0f;
+        psd.offset = inputIndicesPerPSD.size();
+        for (std::size_t i = start; i < psdRows.size(); ++i) {
+            if (psdRows[i] == psdOutputIndex) {
+                inputIndicesPerPSD.push_back(psdCols[i]);
+                psd.weight *= psdWeights[i];
+            }
+        }
+        psd.size = inputIndicesPerPSD.size() - psd.offset;
+    }
+
+    return PSDNet{std::move(inputLODs),
+                  std::move(outputLODs),
+                  std::move(inputIndicesPerPSD),
+                  std::move(psds),
+                  minPSD,
+                  maxPSD};
+}
+
+static Vector<ControlInitializer> createInitialControlValues(const dna::Reader* reader, MemoryResource* memRes) {
     auto endsWith = [](StringView str, StringView suffix) {
             #if __cplusplus >= 201402L || (defined(_MSC_VER) && _MSC_VER >= 1900)
                 return (suffix.size() <= str.size()) && std::equal(extd::advanced(str.begin(), str.size() - suffix.size()),
@@ -111,21 +122,66 @@
             return StringView{str, std::strlen(str)};
         };
     const auto qwSuffix = makeSuffix(".w");
-    const auto controlCount = (static_cast<std::size_t>(rawControlCount) +
-                               static_cast<std::size_t>(psdControlCount) +
-                               static_cast<std::size_t>(mlControlCount) +
-                               static_cast<std::size_t>(rbfControlCount));
-    Vector<float> initialValues{controlCount, {}, memRes};
+    const auto rawControlCount = static_cast<std::size_t>(reader->getRawControlCount());
+    const auto psdControlCount = static_cast<std::size_t>(reader->getPSDCount());
+    const auto mlControlCount = static_cast<std::size_t>(reader->getMLControlCount());
+    const auto rbfControlCount = static_cast<std::size_t>(reader->getRBFPoseControlCount());
+    const auto controlCount = rawControlCount + psdControlCount + mlControlCount + rbfControlCount;
+    Vector<ControlInitializer> initialValues{memRes};
+    initialValues.reserve(controlCount - rawControlCount);
     for (std::uint16_t solverIndex = {}; solverIndex < reader->getRBFSolverCount(); ++solverIndex) {
         const auto solverRawControlIndices = reader->getRBFSolverRawControlIndices(solverIndex);
         for (const auto rawControlIndex : solverRawControlIndices) {
             const auto rawControlName = reader->getRawControlName(rawControlIndex);
             if (endsWith(rawControlName, qwSuffix)) {
-                initialValues[rawControlIndex] = 1.0f;
+                initialValues.push_back({rawControlIndex, 1.0f});
             }
         }
     }
+    return initialValues;
+}
 
+static ControlsInputInstance::Factory createInstanceFactory(const Configuration&  /*unused*/,
+                                                            std::uint16_t guiControlCount,
+                                                            std::uint16_t rawControlCount,
+                                                            std::uint16_t psdControlCount,
+                                                            std::uint16_t mlControlCount,
+                                                            std::uint16_t rbfControlCount) {
+    return [ = ](ConstArrayView<ControlInitializer> initialValues, MemoryResource* memRes) {
+               auto factory = UniqueInstance<StandardControlsInputInstance, ControlsInputInstance>::with(memRes);
+               return factory.create(guiControlCount,
+                                     rawControlCount,
+                                     psdControlCount,
+                                     mlControlCount,
+                                     rbfControlCount,
+                                     initialValues,
+                                     memRes);
+    };
+}
+
+Controls::Pointer ControlsFactory::create(const Configuration& config, const RigMetrics& metrics, MemoryResource* memRes) {
     auto instanceFactory = createInstanceFactory(config,
+                                                 metrics.guiControlCount,
+                                                 metrics.rawControlCount,
+                                                 metrics.psdControlCount,
+                                                 metrics.mlControlCount,
+                                                 metrics.rbfControlCount);
+    return UniqueInstance<Controls>::with(memRes).create(ConditionalTable{memRes},
+                                                         PSDNet{memRes},
+                                                         Vector<ControlInitializer>{memRes},
+                                                         instanceFactory);
+}
+
+Controls::Pointer ControlsFactory::create(const Configuration& config, const dna::Reader* reader, MemoryResource* memRes) {
+    const auto guiControlCount = reader->getGUIControlCount();
+    const auto rawControlCount = reader->getRawControlCount();
+    const auto psdControlCount = reader->getPSDCount();
+    const auto mlControlCount = reader->getMLControlCount();
+    const auto rbfControlCount = reader->getRBFPoseControlCount();
+
+    ConditionalTable conditionals = createConditionalTable(reader, memRes);
+    PSDNet psds = createPSDNet(reader, memRes);
+    Vector<ControlInitializer> initialValues = createInitialControlValues(reader, memRes);
+    auto instanceFactory = createInstanceFactory(config,
                                                  guiControlCount,
                                                  rawControlCount,
diff --git a/Engine/Plugins/Animation/RigLogic/Source/RigLogicLib/Private/riglogic/controls/ControlsInputInstance.h b/Engine/Plugins/Animation/RigLogic/Source/RigLogicLib/Private/riglogic/controls/ControlsInputInstance.h
--- a/Engine/Plugins/Animation/RigLogic/Source/RigLogicLib/Private/riglogic/controls/ControlsInputInstance.h
+++ b/Engine/Plugins/Animation/RigLogic/Source/RigLogicLib/Private/riglogic/controls/ControlsInputInstance.h
@@ -15,17 +15,25 @@
 
 namespace rl4 {
 
+struct ControlInitializer {
+    std::uint32_t index;
+    float value;
+};
+
 class ControlsInputInstance {
     public:
         using Pointer = UniqueInstance<ControlsInputInstance>::PointerType;
-        using Factory = std::function<Pointer(ConstArrayView<float>, MemoryResource*)>;
+        using Factory = std::function<Pointer(ConstArrayView<ControlInitializer>, MemoryResource*)>;
 
     public:
         virtual ~ControlsInputInstance();
+        virtual void resetInternalInputBuffer() = 0;
         virtual ArrayView<float> getGUIControlBuffer() = 0;
         virtual ArrayView<float> getInputBuffer() = 0;
+        virtual ArrayView<float> getClampBuffer() = 0;
         virtual ConstArrayView<float> getGUIControlBuffer() const = 0;
         virtual ConstArrayView<float> getInputBuffer() const = 0;
+        virtual ConstArrayView<float> getClampBuffer() const = 0;
         virtual std::uint16_t getGUIControlCount() const = 0;
         virtual std::uint16_t getRawControlCount() const = 0;
         virtual std::uint16_t getPSDControlCount() const = 0;
diff --git a/Engine/Plugins/Animation/RigLogic/Source/RigLogicLib/Private/riglogic/controls/instances/StandardControlsInputInstance.cpp b/Engine/Plugins/Animation/RigLogic/Source/RigLogicLib/Private/riglogic/controls/instances/StandardControlsInputInstance.cpp
--- a/Engine/Plugins/Animation/RigLogic/Source/RigLogicLib/Private/riglogic/controls/instances/StandardControlsInputInstance.cpp
+++ b/Engine/Plugins/Animation/RigLogic/Source/RigLogicLib/Private/riglogic/controls/instances/StandardControlsInputInstance.cpp
@@ -2,6 +2,8 @@
 
 #include "riglogic/controls/instances/StandardControlsInputInstance.h"
 
+#include "riglogic/utils/Extd.h"
+
 #include <cstdint>
 #include <cstddef>
 
@@ -12,15 +14,29 @@ StandardControlsInputInstance::StandardControlsInputInstance(std::uint16_t guiCo
                                                              std::uint16_t psdControlCount_,
                                                              std::uint16_t mlControlCount_,
                                                              std::uint16_t rbfControlCount_,
-                                                             ConstArrayView<float> initialValues,
+                                                             ConstArrayView<ControlInitializer> initialValues_,
                                                              MemoryResource* memRes) :
     guiControlBuffer{guiControlCount_, {}, memRes},
-    inputBuffer{initialValues.begin(), initialValues.end(), memRes},
+    inputBuffer{static_cast<std::size_t>(rawControlCount_) +
+                static_cast<std::size_t>(psdControlCount_) +
+                static_cast<std::size_t>(mlControlCount_) +
+                static_cast<std::size_t>(rbfControlCount_), {}, memRes},
+    clampBuffer{inputBuffer.size(), {}, memRes},
+    initialValues{initialValues_.begin(), initialValues_.end(), memRes},
     guiControlCount{guiControlCount_},
     rawControlCount{rawControlCount_},
     psdControlCount{psdControlCount_},
     mlControlCount{mlControlCount_},
     rbfControlCount{rbfControlCount_} {
+
+    resetInternalInputBuffer();
+}
+
+void StandardControlsInputInstance::resetInternalInputBuffer() {
+    std::fill(extd::advanced(inputBuffer.begin(), rawControlCount), inputBuffer.end(), 0.0f);
+    for (const auto& initializer : initialValues) {
+        inputBuffer[initializer.index] = initializer.value;
+    }
 }
 
 ArrayView<float> StandardControlsInputInstance::getGUIControlBuffer() {
@@ -31,6 +47,10 @@ ArrayView<float> StandardControlsInputInstance::getInputBuffer() {
     return ArrayView<float>{inputBuffer};
 }
 
+ArrayView<float> StandardControlsInputInstance::getClampBuffer() {
+    return ArrayView<float>{clampBuffer};
+}
+
 ConstArrayView<float> StandardControlsInputInstance::getGUIControlBuffer() const {
     return ConstArrayView<float>{guiControlBuffer};
 }
@@ -39,6 +59,10 @@ ConstArrayView<float> StandardControlsInputInstance::getInputBuffer() const {
     return ConstArrayView<float>{inputBuffer};
 }
 
+ConstArrayView<float> StandardControlsInputInstance::getClampBuffer() const {
+    return ConstArrayView<float>{clampBuffer};
+}
+
 std::uint16_t StandardControlsInputInstance::getGUIControlCount() const {
     return guiControlCount;
 }
diff --git a/Engine/Plugins/Animation/RigLogic/Source/RigLogicLib/Private/riglogic/controls/instances/StandardControlsInputInstance.h b/Engine/Plugins/Animation/RigLogic/Source/RigLogicLib/Private/riglogic/controls/instances/StandardControlsInputInstance.h
--- a/Engine/Plugins/Animation/RigLogic/Source/RigLogicLib/Private/riglogic/controls/instances/StandardControlsInputInstance.h
+++ b/Engine/Plugins/Animation/RigLogic/Source/RigLogicLib/Private/riglogic/controls/instances/StandardControlsInputInstance.h
@@ -16,13 +16,16 @@ class StandardControlsInputInstance : public ControlsInputInstance {
                                       std::uint16_t psdControlCount_,
                                       std::uint16_t mlControlCount_,
                                       std::uint16_t rbfControlCount_,
-                                      ConstArrayView<float> initialValues,
+                                      ConstArrayView<ControlInitializer> initialValues_,
                                       MemoryResource* memRes);
 
+        void resetInternalInputBuffer() override;
         ArrayView<float> getGUIControlBuffer() override;
         ArrayView<float> getInputBuffer() override;
+        ArrayView<float> getClampBuffer() override;
         ConstArrayView<float> getGUIControlBuffer() const override;
         ConstArrayView<float> getInputBuffer() const override;
+        ConstArrayView<float> getClampBuffer() const override;
         std::uint16_t getGUIControlCount() const override;
         std::uint16_t getRawControlCount() const override;
         std::uint16_t getPSDControlCount() const override;
@@ -32,6 +35,8 @@ class StandardControlsInputInstance : public ControlsInputInstance {
     private:
         Vector<float> guiControlBuffer;
         Vector<float> inputBuffer;
+        Vector<float> clampBuffer;
+        Vector<ControlInitializer> initialValues;
         std::uint16_t guiControlCount;
         std::uint16_t rawControlCount;
         std::uint16_t psdControlCount;
diff --git a/Engine/Plugins/Animation/RigLogic/Source/RigLogicLib/Private/riglogic/controls/psdnet/PSDNet.cpp b/Engine/Plugins/Animation/RigLogic/Source/RigLogicLib/Private/riglogic/controls/psdnet/PSDNet.cpp
new file mode 100644
--- /dev/null
+++ b/Engine/Plugins/Animation/RigLogic/Source/RigLogicLib/Private/riglogic/controls/psdnet/PSDNet.cpp
@@ -0,0 +1,100 @@
+// Copyright Epic Games, Inc. All Rights Reserved.
+
+#include "riglogic/controls/psdnet/PSDNet.h"
+
+#include "riglogic/TypeDefs.h"
+#include "riglogic/utils/Extd.h"
+
+#include <cstddef>
+
+namespace rl4 {
+
+namespace {
+
+constexpr float minPSDValue = 0.0f;
+constexpr float maxPSDValue = 1.0f;
+
+}  // namespace
+
+PSDNet::PSDNet(MemoryResource* memRes) :
+    inputLODs{memRes},
+    outputLODs{memRes},
+    inputIndicesPerPSD{memRes},
+    psds{memRes},
+    psdMinIndex{},
+    psdMaxIndex{} {
+}
+
+PSDNet::PSDNet(Matrix<std::uint16_t>&& inputLODs_,
+               Matrix<std::uint16_t>&& outputLODs_,
+               Vector<std::uint16_t>&& inputIndicesPerPSD_,
+               Vector<PSD>&& psds_,
+               std::uint16_t psdMinIndex_,
+               std::uint16_t psdMaxIndex_) :
+    inputLODs{inputLODs_},
+    outputLODs{outputLODs_},
+    inputIndicesPerPSD{std::move(inputIndicesPerPSD_)},
+    psds{std::move(psds_)},
+    psdMinIndex{psdMinIndex_},
+    psdMaxIndex{psdMaxIndex_} {
+}
+
+void PSDNet::registerControls(std::uint16_t lod, ConstArrayView<std::uint16_t> controlIndices) {
+    if ((controlIndices.size() == 0ul) || (psds.size() == 0ul)) {
+        return;
+    }
+
+    // Memory already reserved in both buffers
+    auto& inputIndices = inputLODs[lod];
+    auto& outputIndices = outputLODs[lod];
+
+    for (auto controlIndex : controlIndices) {
+        if ((controlIndex >= psdMinIndex) && (controlIndex <= psdMaxIndex)) {
+            outputIndices.push_back(controlIndex);
+            const PSD psd = psds[static_cast<std::size_t>(controlIndex) - static_cast<std::size_t>(psdMinIndex)];
+            for (std::size_t i = psd.offset; i < psd.offset + psd.size; ++i) {
+                inputIndices.push_back(inputIndicesPerPSD[i]);
+            }
+        }
+    }
+
+    std::sort(inputIndices.begin(), inputIndices.end());
+    std::sort(outputIndices.begin(), outputIndices.end());
+    inputIndices.erase(std::unique(inputIndices.begin(), inputIndices.end()), inputIndices.end());
+    outputIndices.erase(std::unique(outputIndices.begin(), outputIndices.end()), outputIndices.end());
+}
+
+std::uint16_t PSDNet::getPSDCount() const {
+    return static_cast<std::uint16_t>(psds.size());
+}
+
+ConstArrayView<std::uint16_t> PSDNet::getPSDInputIndicesForLOD(std::uint16_t lod) const {
+    assert(lod < inputLODs.size());
+    return inputLODs[lod];
+}
+
+ConstArrayView<std::uint16_t> PSDNet::getPSDOutputIndicesForLOD(std::uint16_t lod) const {
+    assert(lod < outputLODs.size());
+    return outputLODs[lod];
+}
+
+void PSDNet::calculate(ArrayView<float> inputs, ArrayView<float> clampBuffer, std::uint16_t lod) const {
+    ConstArrayView<std::uint16_t> inputIndices = inputLODs[lod];
+    ConstArrayView<std::uint16_t> outputIndices = outputLODs[lod];
+
+    for (auto inputIndex : inputIndices) {
+        clampBuffer[inputIndex] = extd::clamp(inputs[inputIndex], minPSDValue, maxPSDValue);
+    }
+
+    for (auto outputIndex : outputIndices) {
+        const PSD psd = psds[static_cast<std::size_t>(outputIndex) - static_cast<std::size_t>(psdMinIndex)];
+        float psdOutput = psd.weight;
+        for (std::size_t i = psd.offset; i < psd.offset + psd.size; ++i) {
+            const float input = clampBuffer[inputIndicesPerPSD[i]];
+            psdOutput *= input;
+        }
+        inputs[outputIndex] = std::min(maxPSDValue, psdOutput);
+    }
+}
+
+}  // namespace rl4
diff --git a/Engine/Plugins/Animation/RigLogic/Source/RigLogicLib/Private/riglogic/controls/psdnet/PSDNet.h b/Engine/Plugins/Animation/RigLogic/Source/RigLogicLib/Private/riglogic/controls/psdnet/PSDNet.h
new file mode 100644
--- /dev/null
+++ b/Engine/Plugins/Animation/RigLogic/Source/RigLogicLib/Private/riglogic/controls/psdnet/PSDNet.h
@@ -0,0 +1,53 @@
+// Copyright Epic Games, Inc. All Rights Reserved.
+
+#pragma once
+
+#include "riglogic/TypeDefs.h"
+
+#include <cstddef>
+
+namespace rl4 {
+
+struct PSD {
+    std::size_t offset;
+    std::size_t size;
+    float weight;
+
+    template<class Archive>
+    void serialize(Archive& archive) {
+        archive(offset, size, weight);
+    }
+
+};
+
+class PSDNet {
+    public:
+        explicit PSDNet(MemoryResource* memRes);
+        PSDNet(Matrix<std::uint16_t>&& inputLODs_,
+               Matrix<std::uint16_t>&& outputLODs_,
+               Vector<std::uint16_t>&& inputIndicesPerPSD_,
+               Vector<PSD>&& psds_,
+               std::uint16_t psdMinIndex_,
+               std::uint16_t psdMaxIndex_);
+
+        void registerControls(std::uint16_t lod, ConstArrayView<std::uint16_t> controlIndices);
+        std::uint16_t getPSDCount() const;
+        ConstArrayView<std::uint16_t> getPSDInputIndicesForLOD(std::uint16_t lod) const;
+        ConstArrayView<std::uint16_t> getPSDOutputIndicesForLOD(std::uint16_t lod) const;
+        void calculate(ArrayView<float> inputs, ArrayView<float> clampBuffer, std::uint16_t lod) const;
+
+        template<class Archive>
+        void serialize(Archive& archive) {
+            archive(inputLODs, outputLODs, inputIndicesPerPSD, psds, psdMinIndex, psdMaxIndex);
+        }
+
+    private:
+        Matrix<std::uint16_t> inputLODs;
+        Matrix<std::uint16_t> outputLODs;
+        Vector<std::uint16_t> inputIndicesPerPSD;
+        Vector<PSD> psds;
+        std::uint16_t psdMinIndex;
+        std::uint16_t psdMaxIndex;
+};
+
+}  // namespace rl4
diff --git a/Engine/Plugins/Animation/RigLogic/Source/RigLogicLib/Private/riglogic/joints/Joints.cpp b/Engine/Plugins/Animation/RigLogic/Source/RigLogicLib/Private/riglogic/joints/Joints.cpp
--- a/Engine/Plugins/Animation/RigLogic/Source/RigLogicLib/Private/riglogic/joints/Joints.cpp
+++ b/Engine/Plugins/Animation/RigLogic/Source/RigLogicLib/Private/riglogic/joints/Joints.cpp
@@ -14,16 +14,19 @@ Joints::Joints(JointsEvaluator::Pointer evaluator_, MemoryResource* memRes) :
     evaluator{std::move(evaluator_)},
     neutralValues{memRes},
     variableAttributeIndices{memRes},
+    jointIndices{memRes},
     jointGroupCount{} {
 }
 
 Joints::Joints(JointsEvaluator::Pointer evaluator_,
                Vector<float>&& neutralValues_,
                Matrix<std::uint16_t>&& variableAttributeIndices_,
+               Matrix<std::uint16_t>&& jointIndices_,
                std::uint16_t jointGroupCount_) :
     evaluator{std::move(evaluator_)},
     neutralValues{std::move(neutralValues_)},
     variableAttributeIndices{std::move(variableAttributeIndices_)},
+    jointIndices{std::move(jointIndices_)},
     jointGroupCount{jointGroupCount_} {
 }
 
@@ -31,6 +34,15 @@ JointsOutputInstance::Pointer Joints::createInstance(MemoryResource* instanceMem
     return evaluator->createInstance(instanceMemRes);
 }
 
+ConstArrayView<std::uint16_t> Joints::getJointIndicesForLOD(std::uint16_t lod) const {
+    assert(lod < jointIndices.size());
+    return jointIndices[lod];
+}
+
+std::uint32_t Joints::getJointDeltaValueCountForLOD(std::uint16_t lod) const {
+    return evaluator->getJointDeltaValueCountForLOD(lod);
+}
+
 void Joints::calculate(const ControlsInputInstance* inputs, JointsOutputInstance* outputs, std::uint16_t lod) const {
     evaluator->calculate(inputs, outputs, lod);
 }
diff --git a/Engine/Plugins/Animation/RigLogic/Source/RigLogicLib/Private/riglogic/joints/Joints.h b/Engine/Plugins/Animation/RigLogic/Source/RigLogicLib/Private/riglogic/joints/Joints.h
--- a/Engine/Plugins/Animation/RigLogic/Source/RigLogicLib/Private/riglogic/joints/Joints.h
+++ b/Engine/Plugins/Animation/RigLogic/Source/RigLogicLib/Private/riglogic/joints/Joints.h
@@ -22,9 +22,12 @@ class Joints {
         Joints(JointsEvaluator::Pointer evaluator_,
                Vector<float>&& neutralValues_,
                Matrix<std::uint16_t>&& variableAttributeIndices_,
+               Matrix<std::uint16_t>&& jointIndices_,
                std::uint16_t jointGroupCount_);
 
         JointsOutputInstance::Pointer createInstance(MemoryResource* instanceMemRes) const;
+        ConstArrayView<std::uint16_t> getJointIndicesForLOD(std::uint16_t lod) const;
+        std::uint32_t getJointDeltaValueCountForLOD(std::uint16_t lod) const;
         void calculate(const ControlsInputInstance* inputs, JointsOutputInstance* outputs, std::uint16_t lod) const;
         void calculate(const ControlsInputInstance* inputs,
                        JointsOutputInstance* outputs,
@@ -34,13 +37,13 @@ class Joints {
         template<class Archive>
         void load(Archive& archive) {
             evaluator->load(archive);
-            archive >> neutralValues >> variableAttributeIndices >> jointGroupCount;
+            archive >> neutralValues >> variableAttributeIndices >> jointIndices >> jointGroupCount;
         }
 
         template<class Archive>
         void save(Archive& archive) {
             evaluator->save(archive);
-            archive << neutralValues << variableAttributeIndices << jointGroupCount;
+            archive << neutralValues << variableAttributeIndices << jointIndices << jointGroupCount;
         }
 
         std::uint16_t getJointGroupCount() const;
@@ -51,6 +54,7 @@ class Joints {
         JointsEvaluator::Pointer evaluator;
         Vector<float> neutralValues;
         Matrix<std::uint16_t> variableAttributeIndices;
+        Matrix<std::uint16_t> jointIndices;
         std::uint16_t jointGroupCount;
 
 };
diff --git a/Engine/Plugins/Animation/RigLogic/Source/RigLogicLib/Private/riglogic/joints/JointsBuilder.h b/Engine/Plugins/Animation/RigLogic/Source/RigLogicLib/Private/riglogic/joints/JointsBuilder.h
--- a/Engine/Plugins/Animation/RigLogic/Source/RigLogicLib/Private/riglogic/joints/JointsBuilder.h
+++ b/Engine/Plugins/Animation/RigLogic/Source/RigLogicLib/Private/riglogic/joints/JointsBuilder.h
@@ -8,6 +8,7 @@
 
 namespace rl4 {
 
+class Controls;
 class JointBehaviorFilter;
 struct RigMetrics;
 
@@ -24,6 +25,7 @@ class JointsBuilder {
         virtual void computeStorageRequirements(const JointBehaviorFilter& source) = 0;
         virtual void allocateStorage(const JointBehaviorFilter& source) = 0;
         virtual void fillStorage(const JointBehaviorFilter& source) = 0;
+        virtual void registerControls(Controls* controls) = 0;
         virtual JointsEvaluator::Pointer build() = 0;
 };
 
diff --git a/Engine/Plugins/Animation/RigLogic/Source/RigLogicLib/Private/riglogic/joints/JointsEvaluator.h b/Engine/Plugins/Animation/RigLogic/Source/RigLogicLib/Private/riglogic/joints/JointsEvaluator.h
--- a/Engine/Plugins/Animation/RigLogic/Source/RigLogicLib/Private/riglogic/joints/JointsEvaluator.h
+++ b/Engine/Plugins/Animation/RigLogic/Source/RigLogicLib/Private/riglogic/joints/JointsEvaluator.h
@@ -20,6 +20,7 @@ class JointsEvaluator {
 
     public:
         virtual JointsOutputInstance::Pointer createInstance(MemoryResource* instanceMemRes) const = 0;
+        virtual std::uint32_t getJointDeltaValueCountForLOD(std::uint16_t lod) const = 0;
         virtual void calculate(const ControlsInputInstance* inputs, JointsOutputInstance* outputs, std::uint16_t lod) const = 0;
         virtual void calculate(const ControlsInputInstance* inputs,
                                JointsOutputInstance* outputs,
diff --git a/Engine/Plugins/Animation/RigLogic/Source/RigLogicLib/Private/riglogic/joints/JointsFactory.cpp b/Engine/Plugins/Animation/RigLogic/Source/RigLogicLib/Private/riglogic/joints/JointsFactory.cpp
--- a/Engine/Plugins/Animation/RigLogic/Source/RigLogicLib/Private/riglogic/joints/JointsFactory.cpp
+++ b/Engine/Plugins/Animation/RigLogic/Source/RigLogicLib/Private/riglogic/joints/JointsFactory.cpp
@@ -128,7 +128,21 @@ static Matrix<std::uint16_t> copyVariableAttributeIndices(const Configuration& c
     return variableAttributeIndices;
 }
 
-Joints::Pointer JointsFactory::create(const Configuration& config, const dna::Reader* reader, MemoryResource* memRes) {
+static Matrix<std::uint16_t> copyJointIndices(const Configuration&  /*unused*/, const dna::Reader* reader,
+                                              MemoryResource* memRes) {
+    Matrix<std::uint16_t> jointIndices{memRes};
+    jointIndices.resize(reader->getLODCount());
+    for (std::uint16_t lod = {}; lod < reader->getLODCount(); ++lod) {
+        const auto indicesForLOD = reader->getJointIndicesForLOD(lod);
+        jointIndices[lod].assign(indicesForLOD.begin(), indicesForLOD.end());
+    }
+    return jointIndices;
+}
+
+Joints::Pointer JointsFactory::create(const Configuration& config,
+                                      const dna::Reader* reader,
+                                      Controls* controls,
+                                      MemoryResource* memRes) {
     if (!config.loadJoints || (reader->getJointCount() == 0u)) {
         auto evaluator = UniqueInstance<JointsNullEvaluator, JointsEvaluator>::with(memRes).create();
         return UniqueInstance<Joints>::with(memRes).create(std::move(evaluator), memRes);
@@ -144,12 +158,15 @@ Joints::Pointer JointsFactory::create(const Configuration& config, const dna::Re
     builder->computeStorageRequirements(filter);
     builder->allocateStorage(filter);
     builder->fillStorage(filter);
+    builder->registerControls(controls);
     auto evaluator = builder->build();
     auto neutralValues = copyNeutralValues(config, reader, memRes);
     auto variableAttributeIndices = copyVariableAttributeIndices(config, reader, memRes);
+    auto jointIndices = copyJointIndices(config, reader, memRes);
     return UniqueInstance<Joints>::with(memRes).create(std::move(evaluator),
                                                        std::move(neutralValues),
                                                        std::move(variableAttributeIndices),
+                                                       std::move(jointIndices),
                                                        filter.getJointGroupCount());
 }
 
diff --git a/Engine/Plugins/Animation/RigLogic/Source/RigLogicLib/Private/riglogic/joints/JointsFactory.h b/Engine/Plugins/Animation/RigLogic/Source/RigLogicLib/Private/riglogic/joints/JointsFactory.h
--- a/Engine/Plugins/Animation/RigLogic/Source/RigLogicLib/Private/riglogic/joints/JointsFactory.h
+++ b/Engine/Plugins/Animation/RigLogic/Source/RigLogicLib/Private/riglogic/joints/JointsFactory.h
@@ -9,9 +9,13 @@ namespace rl4 {
 
 struct Configuration;
 struct RigMetrics;
+class Controls;
 
 struct JointsFactory {
-    static Joints::Pointer create(const Configuration& config, const dna::Reader* reader, MemoryResource* memRes);
+    static Joints::Pointer create(const Configuration& config,
+                                  const dna::Reader* reader,
+                                  Controls* controls,
+                                  MemoryResource* memRes);
     static Joints::Pointer create(const Configuration& config, const RigMetrics& metrics, MemoryResource* memRes);
 
 };
diff --git a/Engine/Plugins/Animation/RigLogic/Source/RigLogicLib/Private/riglogic/joints/JointsNullEvaluator.cpp b/Engine/Plugins/Animation/RigLogic/Source/RigLogicLib/Private/riglogic/joints/JointsNullEvaluator.cpp
--- a/Engine/Plugins/Animation/RigLogic/Source/RigLogicLib/Private/riglogic/joints/JointsNullEvaluator.cpp
+++ b/Engine/Plugins/Animation/RigLogic/Source/RigLogicLib/Private/riglogic/joints/JointsNullEvaluator.cpp
@@ -8,6 +8,10 @@ JointsOutputInstance::Pointer JointsNullEvaluator::createInstance(MemoryResource
     return UniqueInstance<JointsNullOutputInstance, JointsOutputInstance>::with(instanceMemRes).create();
 }
 
+std::uint32_t JointsNullEvaluator::getJointDeltaValueCountForLOD(std::uint16_t  /*unused*/) const {
+    return {};
+}
+
 void JointsNullEvaluator::calculate(const ControlsInputInstance*  /*unused*/, JointsOutputInstance*  /*unused*/,
                                     std::uint16_t  /*unused*/) const {
 }
diff --git a/Engine/Plugins/Animation/RigLogic/Source/RigLogicLib/Private/riglogic/joints/JointsNullEvaluator.h b/Engine/Plugins/Animation/RigLogic/Source/RigLogicLib/Private/riglogic/joints/JointsNullEvaluator.h
--- a/Engine/Plugins/Animation/RigLogic/Source/RigLogicLib/Private/riglogic/joints/JointsNullEvaluator.h
+++ b/Engine/Plugins/Animation/RigLogic/Source/RigLogicLib/Private/riglogic/joints/JointsNullEvaluator.h
@@ -14,6 +14,7 @@ class ControlsInputInstance;
 class JointsNullEvaluator : public JointsEvaluator {
     public:
         JointsOutputInstance::Pointer createInstance(MemoryResource* instanceMemRes) const override;
+        std::uint32_t getJointDeltaValueCountForLOD(std::uint16_t  /*unused*/) const override;
         void calculate(const ControlsInputInstance*  /*unused*/, JointsOutputInstance*  /*unused*/,
                        std::uint16_t  /*unused*/) const override;
         void calculate(const ControlsInputInstance*  /*unused*/,
diff --git a/Engine/Plugins/Animation/RigLogic/Source/RigLogicLib/Private/riglogic/joints/cpu/CPUJointsBuilder.cpp b/Engine/Plugins/Animation/RigLogic/Source/RigLogicLib/Private/riglogic/joints/cpu/CPUJointsBuilder.cpp
--- a/Engine/Plugins/Animation/RigLogic/Source/RigLogicLib/Private/riglogic/joints/cpu/CPUJointsBuilder.cpp
+++ b/Engine/Plugins/Animation/RigLogic/Source/RigLogicLib/Private/riglogic/joints/cpu/CPUJointsBuilder.cpp
@@ -56,6 +56,12 @@ void CPUJointsBuilder::fillStorage(const JointBehaviorFilter& source) {
     twistSwingBuilder->fillStorage(source);
 }
 
+void CPUJointsBuilder::registerControls(Controls* controls) {
+    bpcmBuilder->registerControls(controls);
+    quaternionBuilder->registerControls(controls);
+    twistSwingBuilder->registerControls(controls);
+}
+
 JointsEvaluator::Pointer CPUJointsBuilder::build() {
     const auto attrCount = jointAttributeCount;
     const auto translationType = config.translationType;
diff --git a/Engine/Plugins/Animation/RigLogic/Source/RigLogicLib/Private/riglogic/joints/cpu/CPUJointsBuilder.h b/Engine/Plugins/Animation/RigLogic/Source/RigLogicLib/Private/riglogic/joints/cpu/CPUJointsBuilder.h
--- a/Engine/Plugins/Animation/RigLogic/Source/RigLogicLib/Private/riglogic/joints/cpu/CPUJointsBuilder.h
+++ b/Engine/Plugins/Animation/RigLogic/Source/RigLogicLib/Private/riglogic/joints/cpu/CPUJointsBuilder.h
@@ -21,6 +21,7 @@ class CPUJointsBuilder : public JointsBuilder {
         void computeStorageRequirements(const JointBehaviorFilter& source) override;
         void allocateStorage(const JointBehaviorFilter& source) override;
         void fillStorage(const JointBehaviorFilter& source) override;
+        void registerControls(Controls* controls) override;
         JointsEvaluator::Pointer build() override;
 
     private:
diff --git a/Engine/Plugins/Animation/RigLogic/Source/RigLogicLib/Private/riglogic/joints/cpu/CPUJointsEvaluator.cpp b/Engine/Plugins/Animation/RigLogic/Source/RigLogicLib/Private/riglogic/joints/cpu/CPUJointsEvaluator.cpp
--- a/Engine/Plugins/Animation/RigLogic/Source/RigLogicLib/Private/riglogic/joints/cpu/CPUJointsEvaluator.cpp
+++ b/Engine/Plugins/Animation/RigLogic/Source/RigLogicLib/Private/riglogic/joints/cpu/CPUJointsEvaluator.cpp
@@ -18,6 +18,12 @@ JointsOutputInstance::Pointer CPUJointsEvaluator::createInstance(MemoryResource*
     return instanceFactory(instanceMemRes);
 }
 
+std::uint32_t CPUJointsEvaluator::getJointDeltaValueCountForLOD(std::uint16_t lod) const {
+    const auto bpcmDeltaCount = bpcmEvaluator->getJointDeltaValueCountForLOD(lod);
+    const auto quaternionDeltaCount = quaternionEvaluator->getJointDeltaValueCountForLOD(lod);
+    return bpcmDeltaCount + quaternionDeltaCount;
+}
+
 void CPUJointsEvaluator::calculate(const ControlsInputInstance* inputs, JointsOutputInstance* outputs, std::uint16_t lod) const {
     bpcmEvaluator->calculate(inputs, outputs, lod);
     quaternionEvaluator->calculate(inputs, outputs, lod);
diff --git a/Engine/Plugins/Animation/RigLogic/Source/RigLogicLib/Private/riglogic/joints/cpu/CPUJointsEvaluator.h b/Engine/Plugins/Animation/RigLogic/Source/RigLogicLib/Private/riglogic/joints/cpu/CPUJointsEvaluator.h
--- a/Engine/Plugins/Animation/RigLogic/Source/RigLogicLib/Private/riglogic/joints/cpu/CPUJointsEvaluator.h
+++ b/Engine/Plugins/Animation/RigLogic/Source/RigLogicLib/Private/riglogic/joints/cpu/CPUJointsEvaluator.h
@@ -22,6 +22,7 @@ class CPUJointsEvaluator : public JointsEvaluator {
                            JointsEvaluator::Pointer twistSwingEvaluator_,
                            JointsOutputInstance::Factory instanceFactory_);
         JointsOutputInstance::Pointer createInstance(MemoryResource* instanceMemRes) const override;
+        std::uint32_t getJointDeltaValueCountForLOD(std::uint16_t lod) const override;
         void calculate(const ControlsInputInstance* inputs, JointsOutputInstance* outputs, std::uint16_t lod) const override;
         void calculate(const ControlsInputInstance* inputs,
                        JointsOutputInstance* outputs,
diff --git a/Engine/Plugins/Animation/RigLogic/Source/RigLogicLib/Private/riglogic/joints/cpu/bpcm/BPCMJointsBuilder.h b/Engine/Plugins/Animation/RigLogic/Source/RigLogicLib/Private/riglogic/joints/cpu/bpcm/BPCMJointsBuilder.h
--- a/Engine/Plugins/Animation/RigLogic/Source/RigLogicLib/Private/riglogic/joints/cpu/bpcm/BPCMJointsBuilder.h
+++ b/Engine/Plugins/Animation/RigLogic/Source/RigLogicLib/Private/riglogic/joints/cpu/bpcm/BPCMJointsBuilder.h
@@ -3,13 +3,14 @@
 #pragma once
 
 #include "riglogic/TypeDefs.h"
+#include "riglogic/controls/Controls.h"
 #include "riglogic/joints/JointBehaviorFilter.h"
 #include "riglogic/joints/JointsBuilder.h"
 #include "riglogic/joints/cpu/bpcm/BPCMJointsEvaluator.h"
 #include "riglogic/joints/cpu/bpcm/CalculationStrategy.h"
 #include "riglogic/joints/cpu/bpcm/RotationAdapters.h"
 #include "riglogic/joints/cpu/bpcm/Storage.h"
-#include "riglogic/joints/cpu/bpcm/StorageSize.h"
+#include "riglogic/joints/cpu/utils/JointGroupOptimizer.h"
 #include "riglogic/riglogic/Configuration.h"
 #include "riglogic/riglogic/RigMetrics.h"
 #include "riglogic/types/Aliases.h"
@@ -52,19 +53,16 @@ class BPCMJointsBuilder : public JointsBuilder {
     public:
         BPCMJointsBuilder(const Configuration& config_, MemoryResource* memRes_);
 
-        void computeStorageRequirements(const RigMetrics& source) override;
-        void computeStorageRequirements(const JointBehaviorFilter& source) override;
+        void computeStorageRequirements(const RigMetrics&  /*unused*/) override;
+        void computeStorageRequirements(const JointBehaviorFilter&  /*unused*/) override;
         void allocateStorage(const JointBehaviorFilter& source) override;
         void fillStorage(const JointBehaviorFilter& source) override;
+        void registerControls(Controls* controls) override;
         JointsEvaluator::Pointer build() override;
 
     private:
-        void setValues(const JointBehaviorFilter& source);
-        void setInputIndices(const JointBehaviorFilter& source);
-        void setOutputIndices(const JointBehaviorFilter& source);
-        void setLODs(const JointBehaviorFilter& source);
         void setOutputRotationIndices(const JointBehaviorFilter& source);
-        void setOutputRotationLODs(const JointBehaviorFilter& source,
+        void setOutputRotationLODs(ConstArrayView<LODRegion> lods,
                                    ConstArrayView<std::uint16_t> outputRotationIndices,
                                    std::uint32_t outputOffset,
                                    std::uint16_t jointGroupIndex);
@@ -79,18 +77,18 @@ class BPCMJointsBuilder : public JointsBuilder {
     private:
         Configuration config;
         MemoryResource* memRes;
-        StorageSize sizeReqs;
         JointStorage<TValue> storage;
         dna::RotationUnit rotationUnit;
+        std::uint16_t lodCount;
 };
 
 template<typename TValue, typename TFVec>
 BPCMJointsBuilder<TValue, TFVec>::BPCMJointsBuilder(const Configuration& config_, MemoryResource* memRes_) :
     config{config_},
     memRes{memRes_},
-    sizeReqs{memRes},
     storage{memRes},
-    rotationUnit{} {
+    rotationUnit{},
+    lodCount{} {
 }
 
 template<typename TValue, typename TFVec>
@@ -98,29 +96,87 @@ void BPCMJointsBuilder<TValue, TFVec>::computeStorageRequirements(const RigMetri
 }
 
 template<typename TValue, typename TFVec>
-void BPCMJointsBuilder<TValue, TFVec>::computeStorageRequirements(const JointBehaviorFilter& source) {
-    sizeReqs.computeFrom(source, PadTo());
+void BPCMJointsBuilder<TValue, TFVec>::computeStorageRequirements(const JointBehaviorFilter&  /*unused*/) {
 }
 
 template<typename TValue, typename TFVec>
-void BPCMJointsBuilder<TValue, TFVec>::allocateStorage(const JointBehaviorFilter&  /*unused*/) {
-    storage.values.resize(sizeReqs.valueCount);
-    storage.inputIndices.resize(sizeReqs.inputIndexCount);
-    storage.outputIndices.resize(sizeReqs.outputIndexCount);
-    storage.lodRegions.reserve(sizeReqs.lodRegionCount);
-    storage.jointGroups.resize(sizeReqs.jointGroups.size());
+void BPCMJointsBuilder<TValue, TFVec>::allocateStorage(const JointBehaviorFilter& source) {
+    lodCount = source.getLODCount();
+    storage.jointGroups.resize(source.getJointGroupCount());
+    storage.lodRegions.resize(storage.jointGroups.size() * static_cast<std::size_t>(lodCount));
     if (config.rotationType == RotationType::Quaternions) {
-        storage.outputRotationLODs.resize(sizeReqs.jointGroups.size() * sizeReqs.lodCount);
+        storage.outputRotationLODs.resize(storage.lodRegions.size());
     }
 }
 
 template<typename TValue, typename TFVec>
 void BPCMJointsBuilder<TValue, TFVec>::fillStorage(const JointBehaviorFilter& source) {
     rotationUnit = source.getRotationUnit();
-    setValues(source);
-    setInputIndices(source);
-    setOutputIndices(source);
-    setLODs(source);
+
+    Vector<float> values{memRes};
+    Vector<std::uint16_t> inputIndices{memRes};
+    Vector<std::uint16_t> outputIndices{memRes};
+    std::uint32_t valueOffset = {};
+    std::uint32_t inputOffset = {};
+    std::uint32_t outputOffset = {};
+    std::uint32_t lodOffset = {};
+    for (std::uint16_t i = {}; i < source.getJointGroupCount(); ++i) {
+        const auto colCount = static_cast<std::uint32_t>(source.getColumnCount(i));
+        const auto rowCount = static_cast<std::uint32_t>(source.getRowCount(i));
+
+        values.resize(rowCount * colCount);
+        source.copyValues(i, values);
+
+        inputIndices.resize(colCount);
+        source.copyInputIndices(i, inputIndices);
+
+        outputIndices.resize(rowCount);
+        source.copyOutputIndices(i, outputIndices);
+
+        // This step might reduce both value count and input index count (by eliminating empty columns)
+        ArrayView<LODRegion> lods{storage.lodRegions.data() + lodOffset, source.getLODCount()};
+        JointGroupOptimizer::defragment(source,
+                                        i,
+                                        values,
+                                        inputIndices,
+                                        outputIndices,
+                                        lods,
+                                        config.translationPruningThreshold,
+                                        config.rotationPruningThreshold,
+                                        config.scalePruningThreshold);
+        const auto optimizedColCount = static_cast<std::uint32_t>(inputIndices.size());
+        const auto optimizedRowCount = static_cast<std::uint32_t>(outputIndices.size());
+        const auto padding = extd::roundUp(optimizedRowCount, PadTo()) - optimizedRowCount;
+        const auto paddedRowCount = optimizedRowCount + padding;
+        // Compute padding and actual matrix size after optimization is done
+        storage.jointGroups[i].valuesSize = optimizedColCount * paddedRowCount;
+        storage.values.resize(storage.values.size() + storage.jointGroups[i].valuesSize);
+        storage.inputIndices.resize(storage.inputIndices.size() + inputIndices.size());
+        storage.outputIndices.resize(storage.outputIndices.size() + paddedRowCount);
+
+        // Needs to use the original unpadded rowcount with the number of optimized columns (that were eliminated)
+        const Extent extent{optimizedRowCount, optimizedColCount};
+        using BPCMOptimizer = Optimizer<TFVec, BlockHeight(), PadTo(), 1u>;
+        storage.jointGroups[i].valuesOffset = valueOffset;
+        valueOffset += BPCMOptimizer::optimize(storage.values.data() + valueOffset, values.data(), extent);
+
+        std::copy(inputIndices.begin(), inputIndices.end(), extd::advanced(storage.inputIndices.begin(), inputOffset));
+        storage.jointGroups[i].inputIndicesOffset = inputOffset;
+        inputOffset += optimizedColCount;
+
+        std::copy(outputIndices.begin(), outputIndices.end(), extd::advanced(storage.outputIndices.begin(), outputOffset));
+        storage.jointGroups[i].outputIndicesOffset = outputOffset;
+        outputOffset += paddedRowCount;
+
+        for (auto& lod : lods) {
+            lod.outputLODs = RowLOD(lod.outputLODs.size, paddedRowCount, BlockHeight(), PadTo());
+        }
+        storage.jointGroups[i].lodsOffset = lodOffset;
+        lodOffset += source.getLODCount();
+        storage.jointGroups[i].colCount = optimizedColCount;
+        storage.jointGroups[i].rowCount = paddedRowCount;
+    }
+
     if (config.rotationType == RotationType::Quaternions) {
         // Remap output indices from 9-attribute joints to 10-attribute joints
         remapOutputIndicesForQuaternions(storage.outputIndices.begin(), storage.outputIndices.end());
@@ -130,62 +186,14 @@ void BPCMJointsBuilder<TValue, TFVec>::fillStorage(const JointBehaviorFilter& so
 }
 
 template<typename TValue, typename TFVec>
-void BPCMJointsBuilder<TValue, TFVec>::setValues(const JointBehaviorFilter& source) {
-    std::uint32_t offset = {};
-    Vector<float> buffer{memRes};
-    for (std::uint16_t i = {}; i < source.getJointGroupCount(); ++i) {
-        const auto jointGroupSize = sizeReqs.getJointGroupSize(i);
-        storage.jointGroups[i].valuesOffset = offset;
-        storage.jointGroups[i].valuesSize = jointGroupSize.padded.size();
-
-        buffer.resize(jointGroupSize.padded.size());
-        source.copyValues(i, buffer);
-
-        offset += Optimizer<TFVec, BlockHeight(), PadTo(), 1u>::optimize(storage.values.data() + offset,
-                                                                         buffer.data(),
-                                                                         jointGroupSize.original);
-    }
-}
-
-template<typename TValue, typename TFVec>
-void BPCMJointsBuilder<TValue, TFVec>::setInputIndices(const JointBehaviorFilter& source) {
-    std::uint32_t offset = {};
-    for (std::uint16_t i = {}; i < source.getJointGroupCount(); ++i) {
-        const auto jointGroupSize = sizeReqs.getJointGroupSize(i);
-        source.copyInputIndices(i, {storage.inputIndices.data() + offset, jointGroupSize.padded.cols});
-        storage.jointGroups[i].inputIndicesOffset = offset;
-        storage.jointGroups[i].inputIndicesSize = jointGroupSize.padded.cols;
-        storage.jointGroups[i].inputIndicesSizeAlignedTo4 = jointGroupSize.padded.cols - (jointGroupSize.padded.cols % 4u);
-        storage.jointGroups[i].inputIndicesSizeAlignedTo8 = jointGroupSize.padded.cols - (jointGroupSize.padded.cols % 8u);
-        offset += jointGroupSize.padded.cols;
-    }
-}
-
-template<typename TValue, typename TFVec>
-void BPCMJointsBuilder<TValue, TFVec>::setOutputIndices(const JointBehaviorFilter& source) {
-    std::uint32_t offset = {};
-    for (std::uint16_t i = {}; i < source.getJointGroupCount(); ++i) {
-        const auto jointGroupSize = sizeReqs.getJointGroupSize(i);
-        source.copyOutputIndices(i, {storage.outputIndices.data() + offset, jointGroupSize.padded.rows});
-        storage.jointGroups[i].outputIndicesOffset = offset;
-        offset += jointGroupSize.padded.rows;
-    }
-}
-
-template<typename TValue, typename TFVec>
-void BPCMJointsBuilder<TValue, TFVec>::setLODs(const JointBehaviorFilter& source) {
-    std::uint32_t offset = {};
-    for (std::uint16_t i = {}; i < source.getJointGroupCount(); ++i) {
-        const auto jointGroupSize = sizeReqs.getJointGroupSize(i);
-        const auto paddedRowCount = jointGroupSize.padded.rows;
-        for (std::uint16_t lod = {}; lod < source.getLODCount(); ++lod) {
-            storage.lodRegions.emplace_back(source.getRowCountForLOD(i, lod),
-                                            paddedRowCount,
-                                            BlockHeight(),
-                                            PadTo());
+void BPCMJointsBuilder<TValue, TFVec>::registerControls(Controls* controls) {
+    for (const auto& group : storage.jointGroups) {
+        ConstArrayView<LODRegion> lods{storage.lodRegions.data() + group.lodsOffset, lodCount};
+        for (std::uint16_t lod = {}; lod < lodCount; ++lod) {
+            ConstArrayView<std::uint16_t> inputIndicesForLOD(storage.inputIndices.data() + group.inputIndicesOffset,
+                                                             lods[lod].inputLODs.size);
+            controls->registerControls(lod, inputIndicesForLOD);
         }
-        storage.jointGroups[i].lodsOffset = offset;
-        offset += source.getLODCount();
     }
 }
 
@@ -193,9 +201,33 @@ template<typename TValue, typename TFVec>
 void BPCMJointsBuilder<TValue, TFVec>::setOutputRotationIndices(const JointBehaviorFilter& source) {
     std::uint32_t outputOffset = {};
     for (std::uint16_t jgi = {}; jgi < source.getJointGroupCount(); ++jgi) {
-        const auto rowCount = source.getRowCount(jgi);
+        auto colCount = static_cast<std::uint32_t>(source.getColumnCount(jgi));
+        auto rowCount = static_cast<std::uint32_t>(source.getRowCount(jgi));
+
+        Vector<float> values{rowCount* colCount, {}, memRes};
+        source.copyValues(jgi, values);
+
+        Vector<std::uint16_t> inputIndices{colCount, {}, memRes};
+        source.copyInputIndices(jgi, inputIndices);
+
         Vector<std::uint16_t> outputRotationIndices{rowCount, {}, memRes};
-        source.copyOutputIndices(jgi, {outputRotationIndices.data(), rowCount});
+        source.copyOutputIndices(jgi, outputRotationIndices);
+
+        Vector<LODRegion> lods{source.getLODCount(), {}, memRes};
+        for (std::uint16_t lod = {}; lod < source.getLODCount(); ++lod) {
+            lods[lod].outputLODs.size = source.getRowCountForLOD(jgi, lod);
+        }
+        JointGroupOptimizer::defragment(source,
+                                        jgi,
+                                        values,
+                                        inputIndices,
+                                        outputRotationIndices,
+                                        lods,
+                                        config.translationPruningThreshold,
+                                        config.rotationPruningThreshold,
+                                        config.scalePruningThreshold);
+
+        rowCount = static_cast<std::uint16_t>(outputRotationIndices.size());
         // Remap output indices from 9-attribute joints to 10-attribute joints rx -> qx
         remapOutputIndicesForQuaternions(outputRotationIndices.begin(), outputRotationIndices.end());
         // Given any rotation indices (qx, qy, qz), return only qx indices for all joints in the group
@@ -227,20 +259,19 @@ void BPCMJointsBuilder<TValue, TFVec>::setOutputRotationIndices(const JointBehav
                   std::back_inserter(storage.outputRotationIndices));
         storage.jointGroups[jgi].outputRotationIndicesOffset = outputOffset;
         // Must be called before outputOffset is adjusted
-        setOutputRotationLODs(source, outputRotationIndices, outputOffset, jgi);
+        setOutputRotationLODs(lods, outputRotationIndices, outputOffset, jgi);
         outputOffset += static_cast<std::uint32_t>(outputRotationBaseIndices.size());
     }
 }
 
 template<typename TValue, typename TFVec>
-void BPCMJointsBuilder<TValue, TFVec>::setOutputRotationLODs(const JointBehaviorFilter& source,
+void BPCMJointsBuilder<TValue, TFVec>::setOutputRotationLODs(ConstArrayView<LODRegion> lods,
                                                              ConstArrayView<std::uint16_t> outputRotationIndices,
                                                              std::uint32_t outputOffset,
                                                              std::uint16_t jointGroupIndex) {
-    const auto lodCount = source.getLODCount();
     const auto offset = static_cast<std::uint32_t>(jointGroupIndex * lodCount);
     for (std::uint16_t lod = {}; lod < lodCount; ++lod) {
-        const auto oldLODRowCount = source.getRowCountForLOD(jointGroupIndex, lod);
+        const auto oldLODRowCount = lods[lod].outputLODs.size;
         if (oldLODRowCount == 0) {
             storage.outputRotationLODs[offset + lod] = 0;
             continue;
@@ -272,14 +303,12 @@ typename UniqueInstance<JointGroupLinearCalculationStrategy<T> >::PointerType cr
     #ifdef RL_BUILD_WITH_XYZ_ROTATION_ORDER
         if (rotationOrder == RotationOrder::XYZ) {
             if (rotationUnit == dna::RotationUnit::degrees) {
-                using CalculationStrategy = VectorizedJointGroupLinearCalculationStrategy<T, TFVec,
-                                                                                          EulerAnglesToQuaternions<tdm::fdeg,
-                                                                                                                   tdm::rot_seq::xyz> >;
+                using E2Q = EulerAnglesToQuaternions<tdm::fdeg, tdm::rot_seq::xyz>;
+                using CalculationStrategy = VectorizedJointGroupLinearCalculationStrategy<T, TFVec, E2Q>;
                 return UniqueInstance<CalculationStrategy, JointGroupLinearCalculationStrategy<T> >::with(memRes).create();
             } else {
-                using CalculationStrategy = VectorizedJointGroupLinearCalculationStrategy<T, TFVec,
-                                                                                          EulerAnglesToQuaternions<tdm::frad,
-                                                                                                                   tdm::rot_seq::xyz> >;
+                using E2Q = EulerAnglesToQuaternions<tdm::frad, tdm::rot_seq::xyz>;
+                using CalculationStrategy = VectorizedJointGroupLinearCalculationStrategy<T, TFVec, E2Q>;
                 return UniqueInstance<CalculationStrategy, JointGroupLinearCalculationStrategy<T> >::with(memRes).create();
             }
         }
@@ -288,14 +317,12 @@ typename UniqueInstance<JointGroupLinearCalculationStrategy<T> >::PointerType cr
     #ifdef RL_BUILD_WITH_XZY_ROTATION_ORDER
         if (rotationOrder == RotationOrder::XZY) {
             if (rotationUnit == dna::RotationUnit::degrees) {
-                using CalculationStrategy = VectorizedJointGroupLinearCalculationStrategy<T, TFVec,
-                                                                                          EulerAnglesToQuaternions<tdm::fdeg,
-                                                                                                                   tdm::rot_seq::xzy> >;
+                using E2Q = EulerAnglesToQuaternions<tdm::fdeg, tdm::rot_seq::xzy>;
+                using CalculationStrategy = VectorizedJointGroupLinearCalculationStrategy<T, TFVec, E2Q>;
                 return UniqueInstance<CalculationStrategy, JointGroupLinearCalculationStrategy<T> >::with(memRes).create();
             } else {
-                using CalculationStrategy = VectorizedJointGroupLinearCalculationStrategy<T, TFVec,
-                                                                                          EulerAnglesToQuaternions<tdm::frad,
-                                                                                                                   tdm::rot_seq::xzy> >;
+                using E2Q = EulerAnglesToQuaternions<tdm::frad, tdm::rot_seq::xzy>;
+                using CalculationStrategy = VectorizedJointGroupLinearCalculationStrategy<T, TFVec, E2Q>;
                 return UniqueInstance<CalculationStrategy, JointGroupLinearCalculationStrategy<T> >::with(memRes).create();
             }
         }
@@ -304,14 +331,12 @@ typename UniqueInstance<JointGroupLinearCalculationStrategy<T> >::PointerType cr
     #ifdef RL_BUILD_WITH_YXZ_ROTATION_ORDER
         if (rotationOrder == RotationOrder::YXZ) {
             if (rotationUnit == dna::RotationUnit::degrees) {
-                using CalculationStrategy = VectorizedJointGroupLinearCalculationStrategy<T, TFVec,
-                                                                                          EulerAnglesToQuaternions<tdm::fdeg,
-                                                                                                                   tdm::rot_seq::yxz> >;
+                using E2Q = EulerAnglesToQuaternions<tdm::fdeg, tdm::rot_seq::yxz>;
+                using CalculationStrategy = VectorizedJointGroupLinearCalculationStrategy<T, TFVec, E2Q>;
                 return UniqueInstance<CalculationStrategy, JointGroupLinearCalculationStrategy<T> >::with(memRes).create();
             } else {
-                using CalculationStrategy = VectorizedJointGroupLinearCalculationStrategy<T, TFVec,
-                                                                                          EulerAnglesToQuaternions<tdm::frad,
-                                                                                                                   tdm::rot_seq::yxz> >;
+                using E2Q = EulerAnglesToQuaternions<tdm::frad, tdm::rot_seq::yxz>;
+                using CalculationStrategy = VectorizedJointGroupLinearCalculationStrategy<T, TFVec, E2Q>;
                 return UniqueInstance<CalculationStrategy, JointGroupLinearCalculationStrategy<T> >::with(memRes).create();
             }
         }
@@ -320,14 +345,12 @@ typename UniqueInstance<JointGroupLinearCalculationStrategy<T> >::PointerType cr
     #ifdef RL_BUILD_WITH_YZX_ROTATION_ORDER
         if (rotationOrder == RotationOrder::YZX) {
             if (rotationUnit == dna::RotationUnit::degrees) {
-                using CalculationStrategy = VectorizedJointGroupLinearCalculationStrategy<T, TFVec,
-                                                                                          EulerAnglesToQuaternions<tdm::fdeg,
-                                                                                                                   tdm::rot_seq::yzx> >;
+                using E2Q = EulerAnglesToQuaternions<tdm::fdeg, tdm::rot_seq::yzx>;
+                using CalculationStrategy = VectorizedJointGroupLinearCalculationStrategy<T, TFVec, E2Q>;
                 return UniqueInstance<CalculationStrategy, JointGroupLinearCalculationStrategy<T> >::with(memRes).create();
             } else {
-                using CalculationStrategy = VectorizedJointGroupLinearCalculationStrategy<T, TFVec,
-                                                                                          EulerAnglesToQuaternions<tdm::frad,
-                                                                                                                   tdm::rot_seq::yzx> >;
+                using E2Q = EulerAnglesToQuaternions<tdm::frad, tdm::rot_seq::yzx>;
+                using CalculationStrategy = VectorizedJointGroupLinearCalculationStrategy<T, TFVec, E2Q>;
                 return UniqueInstance<CalculationStrategy, JointGroupLinearCalculationStrategy<T> >::with(memRes).create();
             }
         }
@@ -336,14 +359,12 @@ typename UniqueInstance<JointGroupLinearCalculationStrategy<T> >::PointerType cr
     #ifdef RL_BUILD_WITH_ZXY_ROTATION_ORDER
         if (rotationOrder == RotationOrder::ZXY) {
             if (rotationUnit == dna::RotationUnit::degrees) {
-                using CalculationStrategy = VectorizedJointGroupLinearCalculationStrategy<T, TFVec,
-                                                                                          EulerAnglesToQuaternions<tdm::fdeg,
-                                                                                                                   tdm::rot_seq::zxy> >;
+                using E2Q = EulerAnglesToQuaternions<tdm::fdeg, tdm::rot_seq::zxy>;
+                using CalculationStrategy = VectorizedJointGroupLinearCalculationStrategy<T, TFVec, E2Q>;
                 return UniqueInstance<CalculationStrategy, JointGroupLinearCalculationStrategy<T> >::with(memRes).create();
             } else {
-                using CalculationStrategy = VectorizedJointGroupLinearCalculationStrategy<T, TFVec,
-                                                                                          EulerAnglesToQuaternions<tdm::frad,
-                                                                                                                   tdm::rot_seq::zxy> >;
+                using E2Q = EulerAnglesToQuaternions<tdm::frad, tdm::rot_seq::zxy>;
+                using CalculationStrategy = VectorizedJointGroupLinearCalculationStrategy<T, TFVec, E2Q>;
                 return UniqueInstance<CalculationStrategy, JointGroupLinearCalculationStrategy<T> >::with(memRes).create();
             }
         }
@@ -352,14 +373,12 @@ typename UniqueInstance<JointGroupLinearCalculationStrategy<T> >::PointerType cr
     #ifdef RL_BUILD_WITH_ZYX_ROTATION_ORDER
         if (rotationOrder == RotationOrder::ZYX) {
             if (rotationUnit == dna::RotationUnit::degrees) {
-                using CalculationStrategy = VectorizedJointGroupLinearCalculationStrategy<T, TFVec,
-                                                                                          EulerAnglesToQuaternions<tdm::fdeg,
-                                                                                                                   tdm::rot_seq::zyx> >;
+                using E2Q = EulerAnglesToQuaternions<tdm::fdeg, tdm::rot_seq::zyx>;
+                using CalculationStrategy = VectorizedJointGroupLinearCalculationStrategy<T, TFVec, E2Q>;
                 return UniqueInstance<CalculationStrategy, JointGroupLinearCalculationStrategy<T> >::with(memRes).create();
             } else {
-                using CalculationStrategy = VectorizedJointGroupLinearCalculationStrategy<T, TFVec,
-                                                                                          EulerAnglesToQuaternions<tdm::frad,
-                                                                                                                   tdm::rot_seq::zyx> >;
+                using E2Q = EulerAnglesToQuaternions<tdm::frad, tdm::rot_seq::zyx>;
+                using CalculationStrategy = VectorizedJointGroupLinearCalculationStrategy<T, TFVec, E2Q>;
                 return UniqueInstance<CalculationStrategy, JointGroupLinearCalculationStrategy<T> >::with(memRes).create();
             }
         }
diff --git a/Engine/Plugins/Animation/RigLogic/Source/RigLogicLib/Private/riglogic/joints/cpu/bpcm/BPCMJointsEvaluator.h b/Engine/Plugins/Animation/RigLogic/Source/RigLogicLib/Private/riglogic/joints/cpu/bpcm/BPCMJointsEvaluator.h
--- a/Engine/Plugins/Animation/RigLogic/Source/RigLogicLib/Private/riglogic/joints/cpu/bpcm/BPCMJointsEvaluator.h
+++ b/Engine/Plugins/Animation/RigLogic/Source/RigLogicLib/Private/riglogic/joints/cpu/bpcm/BPCMJointsEvaluator.h
@@ -8,7 +8,6 @@
 #include "riglogic/joints/JointsOutputInstance.h"
 #include "riglogic/joints/cpu/bpcm/CalculationStrategy.h"
 #include "riglogic/joints/cpu/bpcm/Storage.h"
-#include "riglogic/joints/cpu/bpcm/StorageView.h"
 #include "riglogic/riglogic/RigInstanceImpl.h"
 
 #include <cstdint>
@@ -41,6 +40,14 @@ class Evaluator : public JointsEvaluator {
             return instanceFactory(instanceMemRes);
         }
 
+        std::uint32_t getJointDeltaValueCountForLOD(std::uint16_t lod) const override {
+            std::uint32_t deltaCount = {};
+            for (const auto& group : jointGroups) {
+                deltaCount += (group.lods[lod].inputLODs.size * group.lods[lod].outputLODs.size);
+            }
+            return deltaCount;
+        }
+
         void calculate(const ControlsInputInstance* inputs, JointsOutputInstance* outputs, std::uint16_t lod) const override {
             assert(strategy != nullptr);
             for (std::size_t i = {}; i < jointGroups.size(); ++i) {
