Optimize RigLogic for low LOD evaluation

diff --git a/Engine/Plugins/Animation/RigLogic/Source/RigLogicLib/Private/riglogic/joints/cpu/bpcm/CalculationStrategy.h b/Engine/Plugins/Animation/RigLogic/Source/RigLogicLib/Private/riglogic/joints/cpu/bpcm/CalculationStrategy.h
--- a/Engine/Plugins/Animation/RigLogic/Source/RigLogicLib/Private/riglogic/joints/cpu/bpcm/CalculationStrategy.h
+++ b/Engine/Plugins/Animation/RigLogic/Source/RigLogicLib/Private/riglogic/joints/cpu/bpcm/CalculationStrategy.h
@@ -2,10 +2,9 @@
 
 #pragma once
 
-#include "riglogic/joints/cpu/bpcm/Consts.h"
 #include "riglogic/joints/cpu/bpcm/CalculationStrategy.h"
 #include "riglogic/joints/cpu/bpcm/JointGroup.h"
-#include "riglogic/joints/cpu/bpcm/StorageView.h"
+#include "riglogic/joints/cpu/bpcm/Storage.h"
 #include "riglogic/system/simd/SIMD.h"
 #include "riglogic/types/Aliases.h"
 #include "riglogic/utils/Macros.h"
@@ -13,6 +12,7 @@
 #include <cstddef>
 #include <cstdint>
 
+// *INDENT-OFF*
 namespace rl4 {
 
 namespace bpcm {
@@ -47,7 +47,7 @@ template<typename TFVec, typename T>
 static FORCE_INLINE void processBlocks8x1(const std::uint16_t* inputIndicesEndAlignedTo4,
                                           const std::uint16_t* inputIndicesEnd,
                                           ConstArrayView<float> inputs,
-                                          const T*& values,
+                                          const T* values,
                                           TFVec& sum1,
                                           TFVec& sum2) {
     for (const std::uint16_t* inputIndices = inputIndicesEndAlignedTo4;
@@ -91,7 +91,7 @@ static FORCE_INLINE void processBlocks8x4(const std::uint16_t* inputIndicesStart
                                           const std::uint16_t* inputIndicesEndAlignedTo4,
                                           const std::uint16_t* inputIndicesEnd,
                                           ConstArrayView<float> inputs,
-                                          const T*& values,
+                                          const T* values,
                                           float* outbuf) {
     TFVec sum1{};
     TFVec sum2{};
@@ -116,23 +116,6 @@ static FORCE_INLINE void processBlocks8x4(const std::uint16_t* inputIndicesStart
         const TFVec blk6 = TFVec::fromAlignedSource(values + TFVec::size() * 5);
         const TFVec blk7 = TFVec::fromAlignedSource(values + TFVec::size() * 6);
         const TFVec blk8 = TFVec::fromAlignedSource(values + TFVec::size() * 7);
-        #ifndef RL_NO_MANUAL_PREFETCH
-            #ifdef RL_USE_NON_TEMPORAL_PREFETCH
-                TFVec::prefetchNTA(values + (lookAheadOffset * 6ul));
-                TFVec::prefetchNTA(values + (lookAheadOffset * 7ul));
-                TFVec::prefetchNTA(values + (lookAheadOffset * 8ul));
-                TFVec::prefetchNTA(values + (lookAheadOffset * 9ul));
-                TFVec::prefetchNTA(values + (lookAheadOffset * 10ul));
-                TFVec::prefetchNTA(values + (lookAheadOffset * 11ul));
-            #else
-                TFVec::prefetchT0(values + (lookAheadOffset * 6ul));
-                TFVec::prefetchT0(values + (lookAheadOffset * 7ul));
-                TFVec::prefetchT0(values + (lookAheadOffset * 8ul));
-                TFVec::prefetchT0(values + (lookAheadOffset * 9ul));
-                TFVec::prefetchT0(values + (lookAheadOffset * 10ul));
-                TFVec::prefetchT0(values + (lookAheadOffset * 11ul));
-            #endif
-        #endif
         sum1 += (blk1 * inputVec1);
         sum2 += (blk2 * inputVec1);
         sum3 += (blk3 * inputVec2);
@@ -186,7 +169,7 @@ template<typename TFVec, typename T>
 static FORCE_INLINE void processBlocks4x1(const std::uint16_t* inputIndicesEndAlignedTo8,
                                           const std::uint16_t* inputIndicesEnd,
                                           ConstArrayView<float> inputs,
-                                          const T*& values,
+                                          const T* values,
                                           TFVec& sum1) {
     for (const std::uint16_t* inputIndices = inputIndicesEndAlignedTo8;
          inputIndices < inputIndicesEnd;
@@ -230,7 +213,7 @@ static FORCE_INLINE void processBlocks4x8(const std::uint16_t* inputIndicesStart
                                           const std::uint16_t* inputIndicesEndAlignedTo8,
                                           const std::uint16_t* inputIndicesEnd,
                                           ConstArrayView<float> inputs,
-                                          const T*& values,
+                                          const T* values,
                                           float* outbuf) {
     TFVec sum1{};
     TFVec sum2{};
@@ -289,20 +272,21 @@ template<typename TFVec, typename T>
 static FORCE_INLINE void processJointGroupBlock4(const JointGroupView<T>& jointGroup, ConstArrayView<float> inputs,
                                                  ArrayView<float> outputs, std::uint16_t lod) {
     const T* values = jointGroup.values;
-    const std::uint16_t* const inputIndices = jointGroup.inputIndices;
-    const std::uint16_t* const inputIndicesEnd = jointGroup.inputIndicesEnd;
-    const std::uint16_t* const inputIndicesEndAlignedTo4 = jointGroup.inputIndicesEndAlignedTo4;
-    const std::uint16_t* const inputIndicesEndAlignedTo8 = jointGroup.inputIndicesEndAlignedTo8;
     const LODRegion& lodRegion = jointGroup.lods[lod];
+    const std::uint16_t* const inputIndices = jointGroup.inputIndices;
+    const std::uint16_t* const inputIndicesEnd = inputIndices + lodRegion.inputLODs.size;
+    const std::uint16_t* const inputIndicesEndAlignedTo4 = inputIndices + lodRegion.inputLODs.sizeAlignedTo4;
+    const std::uint16_t* const inputIndicesEndAlignedTo8 = inputIndices + lodRegion.inputLODs.sizeAlignedTo8;
     const std::uint16_t* outputIndices = jointGroup.outputIndices;
-    const std::uint16_t* const outputIndicesEnd = outputIndices + lodRegion.size;
-    const std::uint16_t* const outputIndicesEndPaddedToLastFullBlock = outputIndices + lodRegion.sizePaddedToLastFullBlock;
-    const std::uint16_t* const outputIndicesEndPaddedToSecondLastFullBlock = outputIndices +
-        lodRegion.sizePaddedToSecondLastFullBlock;
+    const std::uint16_t* const outputIndicesEnd = outputIndices + lodRegion.outputLODs.size;
+    const std::uint16_t* const outputIndicesEndPaddedToLastFullBlock = outputIndices + lodRegion.outputLODs.sizePaddedToLastFullBlock;
+    const std::uint16_t* const outputIndicesEndPaddedToSecondLastFullBlock = outputIndices + lodRegion.outputLODs.sizePaddedToSecondLastFullBlock;
     constexpr std::size_t halfBlockHeight = TFVec::size();
     constexpr std::size_t fullBlockHeight = 2ul * TFVec::size();
+    const std::size_t halfBlockSize = jointGroup.colCount * halfBlockHeight;
+    const std::size_t fullBlockSize = jointGroup.colCount * fullBlockHeight;
     // Process portion of matrix that's partitionable into 8x4 blocks
-    for (; outputIndices < outputIndicesEndPaddedToSecondLastFullBlock; outputIndices += fullBlockHeight) {
+    for (; outputIndices < outputIndicesEndPaddedToSecondLastFullBlock; outputIndices += fullBlockHeight, values += fullBlockSize) {
         alignas(TFVec::alignment()) float outbuf[fullBlockHeight];
         processBlocks8x4<TFVec>(inputIndices, inputIndicesEndAlignedTo4, inputIndicesEnd, inputs, values,
                                 static_cast<float*>(outbuf));
@@ -312,18 +296,18 @@ static FORCE_INLINE void processJointGroupBlock4(const JointGroupView<T>& jointG
     }
     // Process the last 8x4 block which needs special handling when some of
     // the output values need to be masked-off because of LODs
-    for (; outputIndices < outputIndicesEndPaddedToLastFullBlock; outputIndices += fullBlockHeight) {
+    for (; outputIndices < outputIndicesEndPaddedToLastFullBlock; outputIndices += fullBlockHeight, values += fullBlockSize) {
         alignas(TFVec::alignment()) float outbuf[fullBlockHeight];
         processBlocks8x4<TFVec>(inputIndices, inputIndicesEndAlignedTo4, inputIndicesEnd, inputs, values,
                                 static_cast<float*>(outbuf));
         // Ignore results that came from rows after the last LOD row
-        for (std::size_t i = 0ul; i < (lodRegion.size % fullBlockHeight); ++i) {
+        for (std::size_t i = 0ul; i < (lodRegion.outputLODs.size % fullBlockHeight); ++i) {
             // NOLINTNEXTLINE(cppcoreguidelines-pro-bounds-constant-array-index)
             outputs[outputIndices[i]] = outbuf[i];
         }
     }
     // Process vertical remainder portion of matrix that's partitionable into 4x8 blocks
-    for (; outputIndices < outputIndicesEnd; outputIndices += halfBlockHeight) {
+    for (; outputIndices < outputIndicesEnd; outputIndices += halfBlockHeight, values += halfBlockSize) {
         alignas(TFVec::alignment()) float outbuf[halfBlockHeight];
         processBlocks4x8<TFVec>(inputIndices, inputIndicesEndAlignedTo8, inputIndicesEnd, inputs, values,
                                 static_cast<float*>(outbuf));
@@ -361,3 +345,4 @@ struct VectorizedJointGroupLinearCalculationStrategy : public JointGroupLinearCa
 }  // namespace bpcm
 
 }  // namespace rl4
+// *INDENT-ON*
diff --git a/Engine/Plugins/Animation/RigLogic/Source/RigLogicLib/Private/riglogic/joints/cpu/bpcm/JointGroup.h b/Engine/Plugins/Animation/RigLogic/Source/RigLogicLib/Private/riglogic/joints/cpu/bpcm/JointGroup.h
--- a/Engine/Plugins/Animation/RigLogic/Source/RigLogicLib/Private/riglogic/joints/cpu/bpcm/JointGroup.h
+++ b/Engine/Plugins/Animation/RigLogic/Source/RigLogicLib/Private/riglogic/joints/cpu/bpcm/JointGroup.h
@@ -3,6 +3,7 @@
 #pragma once
 
 #include "riglogic/types/PaddedBlockView.h"
+#include "riglogic/joints/cpu/utils/LODRegion.h"
 
 #include <cstdint>
 
@@ -10,8 +11,6 @@ namespace rl4 {
 
 namespace bpcm {
 
-using LODRegion = PaddedBlockView;
-
 struct JointGroup {
     // Start of non-zero values in storage
     std::uint32_t valuesOffset;
@@ -27,9 +26,8 @@ struct JointGroup {
     std::uint32_t outputRotationLODsOffset;
     // Sizes associated with start offsets
     std::uint32_t valuesSize;
-    std::uint32_t inputIndicesSize;
-    std::uint32_t inputIndicesSizeAlignedTo4;
-    std::uint32_t inputIndicesSizeAlignedTo8;
+    std::uint32_t colCount;
+    std::uint32_t rowCount;
 
     template<class Archive>
     void serialize(Archive& archive) {
@@ -40,9 +38,8 @@ struct JointGroup {
                 outputRotationIndicesOffset,
                 outputRotationLODsOffset,
                 valuesSize,
-                inputIndicesSize,
-                inputIndicesSizeAlignedTo4,
-                inputIndicesSizeAlignedTo8);
+                colCount,
+                rowCount);
     }
 
 };
diff --git a/Engine/Plugins/Animation/RigLogic/Source/RigLogicLib/Private/riglogic/joints/cpu/bpcm/RotationAdapters.h b/Engine/Plugins/Animation/RigLogic/Source/RigLogicLib/Private/riglogic/joints/cpu/bpcm/RotationAdapters.h
--- a/Engine/Plugins/Animation/RigLogic/Source/RigLogicLib/Private/riglogic/joints/cpu/bpcm/RotationAdapters.h
+++ b/Engine/Plugins/Animation/RigLogic/Source/RigLogicLib/Private/riglogic/joints/cpu/bpcm/RotationAdapters.h
@@ -2,7 +2,7 @@
 
 #pragma once
 
-#include "riglogic/joints/cpu/bpcm/StorageView.h"
+#include "riglogic/joints/cpu/bpcm/Storage.h"
 #include "riglogic/system/simd/SIMD.h"
 #include "riglogic/types/Aliases.h"
 #include "riglogic/utils/Macros.h"
diff --git a/Engine/Plugins/Animation/RigLogic/Source/RigLogicLib/Private/riglogic/joints/cpu/bpcm/Storage.h b/Engine/Plugins/Animation/RigLogic/Source/RigLogicLib/Private/riglogic/joints/cpu/bpcm/Storage.h
--- a/Engine/Plugins/Animation/RigLogic/Source/RigLogicLib/Private/riglogic/joints/cpu/bpcm/Storage.h
+++ b/Engine/Plugins/Animation/RigLogic/Source/RigLogicLib/Private/riglogic/joints/cpu/bpcm/Storage.h
@@ -45,6 +45,35 @@ struct JointStorage {
 
 };
 
+template<typename TValue>
+struct JointGroupView {
+    TValue* values;
+    std::uint32_t colCount;
+    std::uint32_t rowCount;
+    std::uint16_t* inputIndices;
+    std::uint16_t* outputIndices;
+    std::uint16_t* outputRotationIndices;
+    std::uint16_t* outputRotationLODs;
+    LODRegion* lods;
+};
+
+template<typename TValue>
+Vector<JointGroupView<TValue> > takeStorageSnapshot(JointStorage<TValue>& storage, MemoryResource* memRes) {
+    Vector<JointGroupView<TValue> > snapshot{storage.jointGroups.size(), {}, memRes};
+    for (std::size_t i = 0ul; i < storage.jointGroups.size(); ++i) {
+        const auto& jointGroup = storage.jointGroups[i];
+        snapshot[i].values = storage.values.data() + jointGroup.valuesOffset;
+        snapshot[i].colCount = jointGroup.colCount;
+        snapshot[i].rowCount = jointGroup.rowCount;
+        snapshot[i].inputIndices = storage.inputIndices.data() + jointGroup.inputIndicesOffset;
+        snapshot[i].outputIndices = storage.outputIndices.data() + jointGroup.outputIndicesOffset;
+        snapshot[i].outputRotationIndices = storage.outputRotationIndices.data() + jointGroup.outputRotationIndicesOffset;
+        snapshot[i].outputRotationLODs = storage.outputRotationLODs.data() + jointGroup.outputRotationLODsOffset;
+        snapshot[i].lods = storage.lodRegions.data() + jointGroup.lodsOffset;
+    }
+    return snapshot;
+}
+
 }  // namespace bpcm
 
 }  // namespace rl4
diff --git a/Engine/Plugins/Animation/RigLogic/Source/RigLogicLib/Private/riglogic/joints/cpu/quaternions/CalculationStrategy.h b/Engine/Plugins/Animation/RigLogic/Source/RigLogicLib/Private/riglogic/joints/cpu/quaternions/CalculationStrategy.h
--- a/Engine/Plugins/Animation/RigLogic/Source/RigLogicLib/Private/riglogic/joints/cpu/quaternions/CalculationStrategy.h
+++ b/Engine/Plugins/Animation/RigLogic/Source/RigLogicLib/Private/riglogic/joints/cpu/quaternions/CalculationStrategy.h
@@ -10,6 +10,7 @@
 #include <cstddef>
 #include <cstdint>
 
+// *INDENT-OFF*
 namespace rl4 {
 
 /*
@@ -252,7 +253,7 @@ static FORCE_INLINE void normalize(TFVec& qxACEG, TFVec& qyACEG, TFVec& qzACEG,
  * G[3]*H[3] - G[0]*H[0] - G[1]*H[1] - G[2]*H[2]
  */
 template<typename TFVec, typename T>
-static FORCE_INLINE void blendQuaternions(const T*& quaternions,
+static FORCE_INLINE void blendQuaternions(const T* quaternions,
                                           ConstArrayView<std::uint16_t> inputIndices,
                                           ConstArrayView<float> inputs,
                                           float* outbuf) {
@@ -302,29 +303,30 @@ struct VectorizedJointGroupQuaternionCalculationStrategy : public JointGroupQuat
     void calculate(const JointGroup<T>& jointGroup, ConstArrayView<float> inputs, ArrayView<float> outputs,
                    std::uint16_t lod) override {
         const T* quaternions = jointGroup.values.data();
-        ConstArrayView<std::uint16_t> inputIndices{jointGroup.inputIndices};
-        const std::uint16_t* outputIndices = jointGroup.outputIndices.data();
         const LODRegion& lodRegion = jointGroup.lods[lod];
-        const std::uint16_t* const outputIndicesEnd = outputIndices + lodRegion.size;
-        const std::uint16_t* const outputIndicesEndPaddedToLastFullBlock = outputIndices + lodRegion.sizePaddedToLastFullBlock;
-        const std::uint16_t* const outputIndicesEndPaddedToSecondLastFullBlock = outputIndices +
-            lodRegion.sizePaddedToSecondLastFullBlock;
+        ConstArrayView<std::uint16_t> inputIndices{jointGroup.inputIndices.data(), lodRegion.inputLODs.size};
+        const std::uint16_t* outputIndices = jointGroup.outputIndices.data();
+        const std::uint16_t* const outputIndicesEnd = outputIndices + lodRegion.outputLODs.size;
+        const std::uint16_t* const outputIndicesEndPaddedToLastFullBlock = outputIndices + lodRegion.outputLODs.sizePaddedToLastFullBlock;
+        const std::uint16_t* const outputIndicesEndPaddedToSecondLastFullBlock = outputIndices + lodRegion.outputLODs.sizePaddedToSecondLastFullBlock;
+        const std::size_t fullBlockSize = (TFVec256::size() * 4) * jointGroup.colCount;
+        const std::size_t halfBlockSize = (TFVec128::size() * 4) * jointGroup.colCount;
 
-        for (; outputIndices < outputIndicesEndPaddedToSecondLastFullBlock; outputIndices += (TFVec256::size() * 4)) {
+        for (; outputIndices < outputIndicesEndPaddedToSecondLastFullBlock; outputIndices += (TFVec256::size() * 4), quaternions += fullBlockSize) {
             alignas(TFVec256::alignment()) float outbuf[TFVec256::size() * 4];
             blendQuaternions<TFVec256>(quaternions, inputIndices, inputs, static_cast<float*>(outbuf));
             TRotationAdapter::forward(outbuf, TFVec256::size(), TFVec256::size(), outputIndices, outputs);
         }
 
-        for (; outputIndices < outputIndicesEndPaddedToLastFullBlock; outputIndices += (TFVec256::size() * 4)) {
+        for (; outputIndices < outputIndicesEndPaddedToLastFullBlock; outputIndices += (TFVec256::size() * 4), quaternions += fullBlockSize) {
             alignas(TFVec256::alignment()) float outbuf[TFVec256::size() * 4];
             blendQuaternions<TFVec256>(quaternions, inputIndices, inputs, static_cast<float*>(outbuf));
             // Ignore results that came from rows after the last LOD row
-            const auto quaternionCount = (lodRegion.size % (TFVec256::size() * 4)) / 4;
+            const auto quaternionCount = (lodRegion.outputLODs.size % (TFVec256::size() * 4)) / 4;
             TRotationAdapter::forward(outbuf, quaternionCount, TFVec256::size(), outputIndices, outputs);
         }
 
-        for (; outputIndices < outputIndicesEnd; outputIndices += (TFVec128::size() * 4)) {
+        for (; outputIndices < outputIndicesEnd; outputIndices += (TFVec128::size() * 4), quaternions += halfBlockSize) {
             alignas(TFVec128::alignment()) float outbuf[TFVec128::size() * 4];
             blendQuaternions<TFVec128>(quaternions, inputIndices, inputs, static_cast<float*>(outbuf));
             // Ignore results that came from rows after the last LOD row
@@ -336,3 +338,4 @@ struct VectorizedJointGroupQuaternionCalculationStrategy : public JointGroupQuat
 };
 
 }  // namespace rl4
+// *INDENT-ON*
diff --git a/Engine/Plugins/Animation/RigLogic/Source/RigLogicLib/Private/riglogic/joints/cpu/quaternions/JointGroup.h b/Engine/Plugins/Animation/RigLogic/Source/RigLogicLib/Private/riglogic/joints/cpu/quaternions/JointGroup.h
--- a/Engine/Plugins/Animation/RigLogic/Source/RigLogicLib/Private/riglogic/joints/cpu/quaternions/JointGroup.h
+++ b/Engine/Plugins/Animation/RigLogic/Source/RigLogicLib/Private/riglogic/joints/cpu/quaternions/JointGroup.h
@@ -3,35 +3,38 @@
 #pragma once
 
 #include "riglogic/TypeDefs.h"
-#include "riglogic/types/PaddedBlockView.h"
+#include "riglogic/joints/cpu/utils/LODRegion.h"
 
 #include <cstdint>
 
 namespace rl4 {
 
-using LODRegion = PaddedBlockView;
-
 template<typename TValue>
 struct JointGroup {
     // All non-zero values
     AlignedVector<TValue> values;
     // Sub-matrix col -> input vector
-    AlignedVector<std::uint16_t> inputIndices;
+    Vector<std::uint16_t> inputIndices;
     // Sub-matrix row -> output vector
-    AlignedVector<std::uint16_t> outputIndices;
+    Vector<std::uint16_t> outputIndices;
     // Output index boundaries for each LOD
     Vector<LODRegion> lods;
+    // Matrix size
+    std::uint32_t colCount;
+    std::uint32_t rowCount;
 
     explicit JointGroup(MemoryResource* memRes) :
         values{memRes},
         inputIndices{memRes},
         outputIndices{memRes},
-        lods{memRes} {
+        lods{memRes},
+        colCount{},
+        rowCount{} {
     }
 
     template<class Archive>
     void serialize(Archive& archive) {
-        archive(values, inputIndices, outputIndices, lods);
+        archive(values, inputIndices, outputIndices, lods, colCount, rowCount);
     }
 
 };
diff --git a/Engine/Plugins/Animation/RigLogic/Source/RigLogicLib/Private/riglogic/joints/cpu/quaternions/QuaternionJointsBuilder.h b/Engine/Plugins/Animation/RigLogic/Source/RigLogicLib/Private/riglogic/joints/cpu/quaternions/QuaternionJointsBuilder.h
--- a/Engine/Plugins/Animation/RigLogic/Source/RigLogicLib/Private/riglogic/joints/cpu/quaternions/QuaternionJointsBuilder.h
+++ b/Engine/Plugins/Animation/RigLogic/Source/RigLogicLib/Private/riglogic/joints/cpu/quaternions/QuaternionJointsBuilder.h
@@ -3,10 +3,12 @@
 #pragma once
 
 #include "riglogic/TypeDefs.h"
+#include "riglogic/controls/Controls.h"
 #include "riglogic/joints/JointBehaviorFilter.h"
 #include "riglogic/joints/JointsBuilder.h"
 #include "riglogic/joints/cpu/quaternions/QuaternionJointsEvaluator.h"
 #include "riglogic/joints/cpu/quaternions/RotationAdapters.h"
+#include "riglogic/joints/cpu/utils/JointGroupOptimizer.h"
 #include "riglogic/riglogic/Configuration.h"
 #include "riglogic/types/bpcm/Optimizer.h"
 #include "riglogic/utils/Extd.h"
@@ -24,14 +26,18 @@ class QuaternionJointsBuilder : public JointsBuilder {
         void computeStorageRequirements(const JointBehaviorFilter& source) override;
         void allocateStorage(const JointBehaviorFilter& source) override;
         void fillStorage(const JointBehaviorFilter& source) override;
+        void registerControls(Controls* controls) override;
         JointsEvaluator::Pointer build() override;
 
     private:
-        void setInputIndices(const JointBehaviorFilter& source, std::uint16_t jointGroupIndex);
-        void setOutputIndices(const JointBehaviorFilter& source, std::uint16_t jointGroupIndex);
-        void setValues(const JointBehaviorFilter& source, std::uint16_t jointGroupIndex);
-        void setLODs(const JointBehaviorFilter& source, std::uint16_t jointGroupIndex);
-        void remapOutputIndices(std::uint16_t jointGroupIndex);
+        void setInputIndices(JointGroup<TValue>& group, ConstArrayView<std::uint16_t> inputIndices);
+        void setOutputIndices(JointGroup<TValue>& group, ConstArrayView<std::uint16_t> outputIndices);
+        void setValues(JointGroup<TValue>& group,
+                       ConstArrayView<float> eulers,
+                       ConstArrayView<std::uint16_t> inputIndices,
+                       ConstArrayView<std::uint16_t> outputIndices);
+        void setLODs(JointGroup<TValue>& group, ConstArrayView<std::uint16_t> outputIndices);
+        void remapOutputIndices(JointGroup<TValue>& group);
 
         static constexpr std::uint32_t BlockHeight() {
             return 32u;
@@ -78,18 +84,14 @@ void QuaternionJointsBuilder<TValue, TFVec256, TFVec128>::allocateStorage(const
 }
 
 template<typename TValue, typename TFVec256, typename TFVec128>
-void QuaternionJointsBuilder<TValue, TFVec256, TFVec128>::setInputIndices(const JointBehaviorFilter& source,
-                                                                          std::uint16_t jointGroupIndex) {
-    JointGroup<TValue>& group = jointGroups[jointGroupIndex];
-    group.inputIndices.resize(source.getColumnCount(jointGroupIndex));
-    source.copyInputIndices(jointGroupIndex, group.inputIndices);
+void QuaternionJointsBuilder<TValue, TFVec256, TFVec128>::setInputIndices(JointGroup<TValue>& group,
+                                                                          ConstArrayView<std::uint16_t> inputIndices) {
+    group.inputIndices.assign(inputIndices.begin(), inputIndices.end());
 }
 
 template<typename TValue, typename TFVec256, typename TFVec128>
-void QuaternionJointsBuilder<TValue, TFVec256, TFVec128>::setOutputIndices(const JointBehaviorFilter& source,
-                                                                           std::uint16_t jointGroupIndex) {
-    Vector<std::uint16_t> outputIndices{source.getRowCount(jointGroupIndex), {}, memRes};
-    source.copyOutputIndices(jointGroupIndex, outputIndices);
+void QuaternionJointsBuilder<TValue, TFVec256, TFVec128>::setOutputIndices(JointGroup<TValue>& group,
+                                                                           ConstArrayView<std::uint16_t> outputIndices) {
     // Given any rotation indices (rx, ry, rz), return only rx indices for all joints in the group
     #if !defined(__clang__) && defined(__GNUC__)
         #pragma GCC diagnostic push
@@ -114,7 +116,6 @@ void QuaternionJointsBuilder<TValue, TFVec256, TFVec128>::setOutputIndices(const
         });
     deduplicate(outputRotationBaseIndices);
     // Expand output rotation base indices (qx) into (qx, qy, qz, qw) for all rotations
-    JointGroup<TValue>& group = jointGroups[jointGroupIndex];
     group.outputIndices.reserve(outputRotationBaseIndices.size() * static_cast<std::uint8_t>(RotationType::Quaternions));
     for (const auto baseIndex : outputRotationBaseIndices) {
         // Remap output indices from 9-attribute joints to 10-attribute joints rx -> qx
@@ -128,8 +129,10 @@ void QuaternionJointsBuilder<TValue, TFVec256, TFVec128>::setOutputIndices(const
 }
 
 template<typename TValue, typename TFVec256, typename TFVec128>
-void QuaternionJointsBuilder<TValue, TFVec256, TFVec128>::setValues(const JointBehaviorFilter& source,
-                                                                    std::uint16_t jointGroupIndex) {
+void QuaternionJointsBuilder<TValue, TFVec256, TFVec128>::setValues(JointGroup<TValue>& group,
+                                                                    ConstArrayView<float> eulers,
+                                                                    ConstArrayView<std::uint16_t> inputIndices,
+                                                                    ConstArrayView<std::uint16_t> outputIndices) {
     // Convert euler angles to quaternions
     std::function<tdm::frad(float)> angConv;
     if (rotationUnit == dna::RotationUnit::degrees) {
@@ -142,15 +145,9 @@ void QuaternionJointsBuilder<TValue, TFVec256, TFVec128>::setValues(const JointB
             };
     }
 
-    const auto rowCount = source.getRowCount(jointGroupIndex);
-    const auto colCount = source.getColumnCount(jointGroupIndex);
-    Vector<float> eulers(static_cast<std::size_t>(rowCount) * static_cast<std::size_t>(colCount), {}, memRes);
-    source.copyValues(jointGroupIndex, eulers);
+    const auto colCount = static_cast<std::uint16_t>(inputIndices.size());
+    const auto rowCount = static_cast<std::uint16_t>(outputIndices.size());
 
-    Vector<std::uint16_t> outputIndices{source.getRowCount(jointGroupIndex), {}, memRes};
-    source.copyOutputIndices(jointGroupIndex, outputIndices);
-
-    JointGroup<TValue>& group = jointGroups[jointGroupIndex];
     Vector<float> quaternions(static_cast<std::size_t>(group.outputIndices.size() * colCount), {}, memRes);
     for (std::size_t col = {}; col < colCount; ++col) {
         for (std::size_t row = {}, quatIndex = {}; row < rowCount; ++quatIndex) {
@@ -175,13 +172,15 @@ void QuaternionJointsBuilder<TValue, TFVec256, TFVec128>::setValues(const JointB
     const auto newRowCount = static_cast<std::uint32_t>(group.outputIndices.size());
     const auto paddedRowCount = extd::roundUp(newRowCount, PadTo());
     group.values.resize(paddedRowCount * colCount);
-    bpcm::Optimizer<TFVec256, BlockHeight(), PadTo(), Stride()>::optimize(group.values.data(),
-                                                                          quaternions.data(), {newRowCount, colCount});
+    group.colCount = colCount;
+    group.rowCount = paddedRowCount;
+    using BPCMOptimizer = bpcm::Optimizer<TFVec256, BlockHeight(), PadTo(), Stride()>;
+    BPCMOptimizer::optimize(group.values.data(), quaternions.data(), {newRowCount, colCount});
 }
 
 template<typename TValue, typename TFVec256, typename TFVec128>
-void QuaternionJointsBuilder<TValue, TFVec256, TFVec128>::setLODs(const JointBehaviorFilter& source,
-                                                                  std::uint16_t jointGroupIndex) {
+void QuaternionJointsBuilder<TValue, TFVec256, TFVec128>::setLODs(JointGroup<TValue>& group,
+                                                                  ConstArrayView<std::uint16_t> outputIndices) {
     const auto maxRemappedRotationIndex = [](std::uint16_t absRotAttrIndex) {
             const auto jointIndex = static_cast<std::uint16_t>(absRotAttrIndex / 9u);
             const auto newAttrBase = static_cast<std::uint16_t>(jointIndex * 10u);
@@ -190,29 +189,25 @@ void QuaternionJointsBuilder<TValue, TFVec256, TFVec128>::setLODs(const JointBeh
             return static_cast<std::uint16_t>(newAttrBase + 6);
         };
 
-    Vector<std::uint16_t> outputRotationIndices{source.getRowCount(jointGroupIndex), {}, memRes};
-    source.copyOutputIndices(jointGroupIndex, outputRotationIndices);
-
-    JointGroup<TValue>& group = jointGroups[jointGroupIndex];
     const auto newRowCount = static_cast<std::uint32_t>(group.outputIndices.size());
     const auto paddedRowCount = extd::roundUp(newRowCount, PadTo());
-    const auto lodCount = source.getLODCount();
+    const auto lodCount = static_cast<std::uint16_t>(group.lods.size());
     for (std::uint16_t lod = {}; lod < lodCount; ++lod) {
-        const std::uint32_t oldLODRowCount = source.getRowCountForLOD(jointGroupIndex, lod);
+        const std::uint32_t oldLODRowCount = group.lods[lod].outputLODs.size;
         std::uint32_t newLODRowCount = {};
         if (oldLODRowCount != 0) {
-            const auto qwRotationIndexAtOldLODRowCount = maxRemappedRotationIndex(outputRotationIndices[oldLODRowCount - 1ul]);
+            const auto qwRotationIndexAtOldLODRowCount = maxRemappedRotationIndex(outputIndices[oldLODRowCount - 1ul]);
             auto it = std::find(group.outputIndices.begin(), group.outputIndices.end(), qwRotationIndexAtOldLODRowCount);
             assert(it != group.outputIndices.end());
             newLODRowCount = static_cast<std::uint16_t>(std::distance(group.outputIndices.begin(), it) + 1);
         }
-        group.lods[lod] = LODRegion{newLODRowCount, paddedRowCount, BlockHeight(), PadTo()};
+        RowLOD outputLOD{newLODRowCount, paddedRowCount, BlockHeight(), PadTo()};
+        group.lods[lod].outputLODs = outputLOD;
     }
 }
 
 template<typename TValue, typename TFVec256, typename TFVec128>
-void QuaternionJointsBuilder<TValue, TFVec256, TFVec128>::remapOutputIndices(std::uint16_t jointGroupIndex) {
-    JointGroup<TValue>& group = jointGroups[jointGroupIndex];
+void QuaternionJointsBuilder<TValue, TFVec256, TFVec128>::remapOutputIndices(JointGroup<TValue>& group) {
     for (auto& outputIndex : group.outputIndices) {
         const auto jointIndex = static_cast<std::uint16_t>(outputIndex / 10u);
         const auto relAttrIndex = static_cast<std::uint16_t>(outputIndex % 10u);
@@ -228,20 +223,52 @@ void QuaternionJointsBuilder<TValue, TFVec256, TFVec128>::fillStorage(const Join
     rotationUnit = source.getRotationUnit();
 
     for (std::uint16_t jgi = {}; jgi < static_cast<std::uint16_t>(jointGroups.size()); ++jgi) {
-        const auto rowCount = source.getRowCount(jgi);
-        const auto colCount = source.getColumnCount(jgi);
+        auto rowCount = source.getRowCount(jgi);
+        auto colCount = source.getColumnCount(jgi);
         if ((rowCount == 0u) || (colCount == 0u)) {
             continue;
         }
 
-        setInputIndices(source, jgi);
-        setOutputIndices(source, jgi);
-        setValues(source, jgi);
-        setLODs(source, jgi);
+        JointGroup<TValue>& group = jointGroups[jgi];
+
+        Vector<float> eulers(static_cast<std::size_t>(rowCount) * static_cast<std::size_t>(colCount), {}, memRes);
+        source.copyValues(jgi, eulers);
+
+        Vector<std::uint16_t> inputIndices{colCount, {}, memRes};
+        source.copyInputIndices(jgi, inputIndices);
+
+        Vector<std::uint16_t> outputIndices{rowCount, {}, memRes};
+        source.copyOutputIndices(jgi, outputIndices);
+
+        JointGroupOptimizer::defragment(source,
+                                        jgi,
+                                        eulers,
+                                        inputIndices,
+                                        outputIndices,
+                                        group.lods,
+                                        config.translationPruningThreshold,
+                                        config.rotationPruningThreshold,
+                                        config.scalePruningThreshold);
+
+        setInputIndices(group, inputIndices);
+        setOutputIndices(group, outputIndices);
+        setValues(group, eulers, inputIndices, outputIndices);
+        setLODs(group, outputIndices);
+
         // If the selected RigLogic output is in quaternions, then the output indices are already setup as needed.
         // But if Euler angles were requested, the output indices need to be mapped back to 9-attribute joint output indices
         if (config.rotationType == RotationType::EulerAngles) {
-            remapOutputIndices(jgi);
+            remapOutputIndices(group);
+        }
+    }
+}
+
+template<typename TValue, typename TFVec256, typename TFVec128>
+void QuaternionJointsBuilder<TValue, TFVec256, TFVec128>::registerControls(Controls* controls) {
+    for (const auto& group : jointGroups) {
+        for (std::uint16_t lod = {}; lod < static_cast<std::uint16_t>(group.lods.size()); ++lod) {
+            ConstArrayView<std::uint16_t> inputIndicesForLOD(group.inputIndices.data(), group.lods[lod].inputLODs.size);
+            controls->registerControls(lod, inputIndicesForLOD);
         }
     }
 }
@@ -261,14 +288,12 @@ typename UniqueInstance<JointGroupQuaternionCalculationStrategy<T> >::PointerTyp
     #ifdef RL_BUILD_WITH_XYZ_ROTATION_ORDER
         if (rotationOrder == RotationOrder::XYZ) {
             if (rotationUnit == dna::RotationUnit::degrees) {
-                using CalculationStrategy = VectorizedJointGroupQuaternionCalculationStrategy<T, TFVec256, TFVec128,
-                                                                                              QuaternionsToEulerAngles<tdm::fdeg,
-                                                                                                                       tdm::rot_seq::xyz> >;
+                using Q2E = QuaternionsToEulerAngles<tdm::fdeg, tdm::rot_seq::xyz>;
+                using CalculationStrategy = VectorizedJointGroupQuaternionCalculationStrategy<T, TFVec256, TFVec128, Q2E>;
                 return UniqueInstance<CalculationStrategy, JointGroupQuaternionCalculationStrategy<T> >::with(memRes).create();
             } else {
-                using CalculationStrategy = VectorizedJointGroupQuaternionCalculationStrategy<T, TFVec256, TFVec128,
-                                                                                              QuaternionsToEulerAngles<tdm::frad,
-                                                                                                                       tdm::rot_seq::xyz> >;
+                using Q2E = QuaternionsToEulerAngles<tdm::frad, tdm::rot_seq::xyz>;
+                using CalculationStrategy = VectorizedJointGroupQuaternionCalculationStrategy<T, TFVec256, TFVec128, Q2E>;
                 return UniqueInstance<CalculationStrategy, JointGroupQuaternionCalculationStrategy<T> >::with(memRes).create();
             }
         }
@@ -277,14 +302,12 @@ typename UniqueInstance<JointGroupQuaternionCalculationStrategy<T> >::PointerTyp
     #ifdef RL_BUILD_WITH_XZY_ROTATION_ORDER
         if (rotationOrder == RotationOrder::XZY) {
             if (rotationUnit == dna::RotationUnit::degrees) {
-                using CalculationStrategy = VectorizedJointGroupQuaternionCalculationStrategy<T, TFVec256, TFVec128,
-                                                                                              QuaternionsToEulerAngles<tdm::fdeg,
-                                                                                                                       tdm::rot_seq::xzy> >;
+                using Q2E = QuaternionsToEulerAngles<tdm::fdeg, tdm::rot_seq::xzy>;
+                using CalculationStrategy = VectorizedJointGroupQuaternionCalculationStrategy<T, TFVec256, TFVec128, Q2E>;
                 return UniqueInstance<CalculationStrategy, JointGroupQuaternionCalculationStrategy<T> >::with(memRes).create();
             } else {
-                using CalculationStrategy = VectorizedJointGroupQuaternionCalculationStrategy<T, TFVec256, TFVec128,
-                                                                                              QuaternionsToEulerAngles<tdm::frad,
-                                                                                                                       tdm::rot_seq::xzy> >;
+                using Q2E = QuaternionsToEulerAngles<tdm::frad, tdm::rot_seq::xzy>;
+                using CalculationStrategy = VectorizedJointGroupQuaternionCalculationStrategy<T, TFVec256, TFVec128, Q2E>;
                 return UniqueInstance<CalculationStrategy, JointGroupQuaternionCalculationStrategy<T> >::with(memRes).create();
             }
         }
@@ -293,14 +316,12 @@ typename UniqueInstance<JointGroupQuaternionCalculationStrategy<T> >::PointerTyp
     #ifdef RL_BUILD_WITH_YXZ_ROTATION_ORDER
         if (rotationOrder == RotationOrder::YXZ) {
             if (rotationUnit == dna::RotationUnit::degrees) {
-                using CalculationStrategy = VectorizedJointGroupQuaternionCalculationStrategy<T, TFVec256, TFVec128,
-                                                                                              QuaternionsToEulerAngles<tdm::fdeg,
-                                                                                                                       tdm::rot_seq::yxz> >;
+                using Q2E = QuaternionsToEulerAngles<tdm::fdeg, tdm::rot_seq::yxz>;
+                using CalculationStrategy = VectorizedJointGroupQuaternionCalculationStrategy<T, TFVec256, TFVec128, Q2E>;
                 return UniqueInstance<CalculationStrategy, JointGroupQuaternionCalculationStrategy<T> >::with(memRes).create();
             } else {
-                using CalculationStrategy = VectorizedJointGroupQuaternionCalculationStrategy<T, TFVec256, TFVec128,
-                                                                                              QuaternionsToEulerAngles<tdm::frad,
-                                                                                                                       tdm::rot_seq::yxz> >;
+                using Q2E = QuaternionsToEulerAngles<tdm::frad, tdm::rot_seq::yxz>;
+                using CalculationStrategy = VectorizedJointGroupQuaternionCalculationStrategy<T, TFVec256, TFVec128, Q2E>;
                 return UniqueInstance<CalculationStrategy, JointGroupQuaternionCalculationStrategy<T> >::with(memRes).create();
             }
         }
@@ -309,14 +330,12 @@ typename UniqueInstance<JointGroupQuaternionCalculationStrategy<T> >::PointerTyp
     #ifdef RL_BUILD_WITH_YZX_ROTATION_ORDER
         if (rotationOrder == RotationOrder::YZX) {
             if (rotationUnit == dna::RotationUnit::degrees) {
-                using CalculationStrategy = VectorizedJointGroupQuaternionCalculationStrategy<T, TFVec256, TFVec128,
-                                                                                              QuaternionsToEulerAngles<tdm::fdeg,
-                                                                                                                       tdm::rot_seq::yzx> >;
+                using Q2E = QuaternionsToEulerAngles<tdm::fdeg, tdm::rot_seq::yzx>;
+                using CalculationStrategy = VectorizedJointGroupQuaternionCalculationStrategy<T, TFVec256, TFVec128, Q2E>;
                 return UniqueInstance<CalculationStrategy, JointGroupQuaternionCalculationStrategy<T> >::with(memRes).create();
             } else {
-                using CalculationStrategy = VectorizedJointGroupQuaternionCalculationStrategy<T, TFVec256, TFVec128,
-                                                                                              QuaternionsToEulerAngles<tdm::frad,
-                                                                                                                       tdm::rot_seq::yzx> >;
+                using Q2E = QuaternionsToEulerAngles<tdm::frad, tdm::rot_seq::yzx>;
+                using CalculationStrategy = VectorizedJointGroupQuaternionCalculationStrategy<T, TFVec256, TFVec128, Q2E>;
                 return UniqueInstance<CalculationStrategy, JointGroupQuaternionCalculationStrategy<T> >::with(memRes).create();
             }
         }
@@ -325,14 +344,12 @@ typename UniqueInstance<JointGroupQuaternionCalculationStrategy<T> >::PointerTyp
     #ifdef RL_BUILD_WITH_ZXY_ROTATION_ORDER
         if (rotationOrder == RotationOrder::ZXY) {
             if (rotationUnit == dna::RotationUnit::degrees) {
-                using CalculationStrategy = VectorizedJointGroupQuaternionCalculationStrategy<T, TFVec256, TFVec128,
-                                                                                              QuaternionsToEulerAngles<tdm::fdeg,
-                                                                                                                       tdm::rot_seq::zxy> >;
+                using Q2E = QuaternionsToEulerAngles<tdm::fdeg, tdm::rot_seq::zxy>;
+                using CalculationStrategy = VectorizedJointGroupQuaternionCalculationStrategy<T, TFVec256, TFVec128, Q2E>;
                 return UniqueInstance<CalculationStrategy, JointGroupQuaternionCalculationStrategy<T> >::with(memRes).create();
             } else {
-                using CalculationStrategy = VectorizedJointGroupQuaternionCalculationStrategy<T, TFVec256, TFVec128,
-                                                                                              QuaternionsToEulerAngles<tdm::frad,
-                                                                                                                       tdm::rot_seq::zxy> >;
+                using Q2E = QuaternionsToEulerAngles<tdm::frad, tdm::rot_seq::zxy>;
+                using CalculationStrategy = VectorizedJointGroupQuaternionCalculationStrategy<T, TFVec256, TFVec128, Q2E>;
                 return UniqueInstance<CalculationStrategy, JointGroupQuaternionCalculationStrategy<T> >::with(memRes).create();
             }
         }
@@ -341,14 +358,12 @@ typename UniqueInstance<JointGroupQuaternionCalculationStrategy<T> >::PointerTyp
     #ifdef RL_BUILD_WITH_ZYX_ROTATION_ORDER
         if (rotationOrder == RotationOrder::ZYX) {
             if (rotationUnit == dna::RotationUnit::degrees) {
-                using CalculationStrategy = VectorizedJointGroupQuaternionCalculationStrategy<T, TFVec256, TFVec128,
-                                                                                              QuaternionsToEulerAngles<tdm::fdeg,
-                                                                                                                       tdm::rot_seq::zyx> >;
+                using Q2E = QuaternionsToEulerAngles<tdm::fdeg, tdm::rot_seq::zyx>;
+                using CalculationStrategy = VectorizedJointGroupQuaternionCalculationStrategy<T, TFVec256, TFVec128, Q2E>;
                 return UniqueInstance<CalculationStrategy, JointGroupQuaternionCalculationStrategy<T> >::with(memRes).create();
             } else {
-                using CalculationStrategy = VectorizedJointGroupQuaternionCalculationStrategy<T, TFVec256, TFVec128,
-                                                                                              QuaternionsToEulerAngles<tdm::frad,
-                                                                                                                       tdm::rot_seq::zyx> >;
+                using Q2E = QuaternionsToEulerAngles<tdm::frad, tdm::rot_seq::zyx>;
+                using CalculationStrategy = VectorizedJointGroupQuaternionCalculationStrategy<T, TFVec256, TFVec128, Q2E>;
                 return UniqueInstance<CalculationStrategy, JointGroupQuaternionCalculationStrategy<T> >::with(memRes).create();
             }
         }
diff --git a/Engine/Plugins/Animation/RigLogic/Source/RigLogicLib/Private/riglogic/joints/cpu/quaternions/QuaternionJointsEvaluator.h b/Engine/Plugins/Animation/RigLogic/Source/RigLogicLib/Private/riglogic/joints/cpu/quaternions/QuaternionJointsEvaluator.h
--- a/Engine/Plugins/Animation/RigLogic/Source/RigLogicLib/Private/riglogic/joints/cpu/quaternions/QuaternionJointsEvaluator.h
+++ b/Engine/Plugins/Animation/RigLogic/Source/RigLogicLib/Private/riglogic/joints/cpu/quaternions/QuaternionJointsEvaluator.h
@@ -29,6 +29,7 @@ class QuaternionJointsEvaluator : public JointsEvaluator {
                                   MemoryResource* memRes);
 
         JointsOutputInstance::Pointer createInstance(MemoryResource* instanceMemRes) const override;
+        std::uint32_t getJointDeltaValueCountForLOD(std::uint16_t lod) const override;
         void calculate(const ControlsInputInstance* inputs, JointsOutputInstance* outputs, std::uint16_t lod) const override;
         void calculate(const ControlsInputInstance* inputs,
                        JointsOutputInstance* outputs,
@@ -59,6 +60,16 @@ JointsOutputInstance::Pointer QuaternionJointsEvaluator<TValue>::createInstance(
     return instanceFactory(instanceMemRes);
 }
 
+template<typename TValue>
+std::uint32_t QuaternionJointsEvaluator<TValue>::getJointDeltaValueCountForLOD(std::uint16_t lod) const {
+    std::uint32_t deltaCount = {};
+    for (const auto& group : jointGroups) {
+        assert(lod < group.lods.size());
+        deltaCount += (group.lods[lod].inputLODs.size * group.lods[lod].outputLODs.size);
+    }
+    return deltaCount;
+}
+
 template<typename TValue>
 void QuaternionJointsEvaluator<TValue>::calculate(const ControlsInputInstance* inputs,
                                                   JointsOutputInstance* outputs,
diff --git a/Engine/Plugins/Animation/RigLogic/Source/RigLogicLib/Private/riglogic/joints/cpu/twistswing/TwistSwingJointsBuilder.h b/Engine/Plugins/Animation/RigLogic/Source/RigLogicLib/Private/riglogic/joints/cpu/twistswing/TwistSwingJointsBuilder.h
--- a/Engine/Plugins/Animation/RigLogic/Source/RigLogicLib/Private/riglogic/joints/cpu/twistswing/TwistSwingJointsBuilder.h
+++ b/Engine/Plugins/Animation/RigLogic/Source/RigLogicLib/Private/riglogic/joints/cpu/twistswing/TwistSwingJointsBuilder.h
@@ -34,6 +34,7 @@ class TwistSwingJointsBuilder : public JointsBuilder {
         void computeStorageRequirements(const JointBehaviorFilter& source) override;
         void allocateStorage(const JointBehaviorFilter& source) override;
         void fillStorage(const JointBehaviorFilter& source) override;
+        void registerControls(Controls* controls) override;
         JointsEvaluator::Pointer build() override;
 
     private:
@@ -161,6 +162,10 @@ void TwistSwingJointsBuilder<TValue, TFVec256, TFVec128>::fillStorage(const Join
 
 }
 
+template<typename TValue, typename TFVec256, typename TFVec128>
+void TwistSwingJointsBuilder<TValue, TFVec256, TFVec128>::registerControls(Controls*  /*unused*/) {
+}
+
 template<typename TValue, typename TFVec256, typename TFVec128>
 JointsEvaluator::Pointer TwistSwingJointsBuilder<TValue, TFVec256, TFVec128>::build() {
 
diff --git a/Engine/Plugins/Animation/RigLogic/Source/RigLogicLib/Private/riglogic/joints/cpu/twistswing/TwistSwingJointsEvaluator.h b/Engine/Plugins/Animation/RigLogic/Source/RigLogicLib/Private/riglogic/joints/cpu/twistswing/TwistSwingJointsEvaluator.h
--- a/Engine/Plugins/Animation/RigLogic/Source/RigLogicLib/Private/riglogic/joints/cpu/twistswing/TwistSwingJointsEvaluator.h
+++ b/Engine/Plugins/Animation/RigLogic/Source/RigLogicLib/Private/riglogic/joints/cpu/twistswing/TwistSwingJointsEvaluator.h
@@ -43,6 +43,7 @@ class TwistSwingJointsEvaluator : public JointsEvaluator {
                                            MemoryResource* memRes);
 
         JointsOutputInstance::Pointer createInstance(MemoryResource* instanceMemRes) const override;
+        std::uint32_t getJointDeltaValueCountForLOD(std::uint16_t lod) const override;
         void calculate(const ControlsInputInstance* inputs, JointsOutputInstance* outputs, std::uint16_t lod) const override;
         void calculate(const ControlsInputInstance* inputs,
                        JointsOutputInstance* outputs,
@@ -73,6 +74,12 @@ JointsOutputInstance::Pointer TwistSwingJointsEvaluator<TValue, TFVec256, TFVec1
     return instanceFactory(instanceMemRes);
 }
 
+template<typename TValue, typename TFVec256, typename TFVec128, class TRotationAdapter>
+std::uint32_t TwistSwingJointsEvaluator<TValue, TFVec256, TFVec128, TRotationAdapter>::getJointDeltaValueCountForLOD(
+    std::uint16_t  /*unused*/) const {
+    return {};
+}
+
 template<typename TValue, typename TFVec256, typename TFVec128, class TRotationAdapter>
 void TwistSwingJointsEvaluator<TValue, TFVec256, TFVec128, TRotationAdapter>::calculate(const ControlsInputInstance* inputs,
                                                                                         JointsOutputInstance* outputs,
diff --git a/Engine/Plugins/Animation/RigLogic/Source/RigLogicLib/Private/riglogic/ml/MachineLearnedBehavior.cpp b/Engine/Plugins/Animation/RigLogic/Source/RigLogicLib/Private/riglogic/ml/MachineLearnedBehavior.cpp
--- a/Engine/Plugins/Animation/RigLogic/Source/RigLogicLib/Private/riglogic/ml/MachineLearnedBehavior.cpp
+++ b/Engine/Plugins/Animation/RigLogic/Source/RigLogicLib/Private/riglogic/ml/MachineLearnedBehavior.cpp
@@ -26,6 +26,10 @@ MachineLearnedBehaviorOutputInstance::Pointer MachineLearnedBehavior::createInst
     return evaluator->createInstance(instanceMemRes);
 }
 
+ConstArrayView<std::uint32_t> MachineLearnedBehavior::getNeuralNetworkIndicesForLOD(std::uint16_t lod) const {
+    return evaluator->getNeuralNetworkIndicesForLOD(lod);
+}
+
 void MachineLearnedBehavior::calculate(ControlsInputInstance* inputs,
                                        MachineLearnedBehaviorOutputInstance* intermediateOutputs,
                                        std::uint16_t lod) const {
diff --git a/Engine/Plugins/Animation/RigLogic/Source/RigLogicLib/Private/riglogic/ml/MachineLearnedBehavior.h b/Engine/Plugins/Animation/RigLogic/Source/RigLogicLib/Private/riglogic/ml/MachineLearnedBehavior.h
--- a/Engine/Plugins/Animation/RigLogic/Source/RigLogicLib/Private/riglogic/ml/MachineLearnedBehavior.h
+++ b/Engine/Plugins/Animation/RigLogic/Source/RigLogicLib/Private/riglogic/ml/MachineLearnedBehavior.h
@@ -24,6 +24,7 @@ class MachineLearnedBehavior {
                                Vector<Matrix<std::uint16_t> >&& neuralNetworkIndicesPerMeshRegion_);
 
         MachineLearnedBehaviorOutputInstance::Pointer createInstance(MemoryResource* instanceMemRes) const;
+        ConstArrayView<std::uint32_t> getNeuralNetworkIndicesForLOD(std::uint16_t lod) const;
         void calculate(ControlsInputInstance* inputs, MachineLearnedBehaviorOutputInstance* intermediateOutputs,
                        std::uint16_t lod) const;
         void calculate(ControlsInputInstance* inputs,
diff --git a/Engine/Plugins/Animation/RigLogic/Source/RigLogicLib/Private/riglogic/ml/MachineLearnedBehaviorEvaluator.h b/Engine/Plugins/Animation/RigLogic/Source/RigLogicLib/Private/riglogic/ml/MachineLearnedBehaviorEvaluator.h
--- a/Engine/Plugins/Animation/RigLogic/Source/RigLogicLib/Private/riglogic/ml/MachineLearnedBehaviorEvaluator.h
+++ b/Engine/Plugins/Animation/RigLogic/Source/RigLogicLib/Private/riglogic/ml/MachineLearnedBehaviorEvaluator.h
@@ -18,6 +18,7 @@ class MachineLearnedBehaviorEvaluator {
 
     public:
         virtual MachineLearnedBehaviorOutputInstance::Pointer createInstance(MemoryResource* instanceMemRes) const = 0;
+        virtual ConstArrayView<std::uint32_t> getNeuralNetworkIndicesForLOD(std::uint16_t lod) const = 0;
         virtual void calculate(ControlsInputInstance* inputs,
                                MachineLearnedBehaviorOutputInstance* intermediateOutputs,
                                std::uint16_t lod) const = 0;
diff --git a/Engine/Plugins/Animation/RigLogic/Source/RigLogicLib/Private/riglogic/ml/MachineLearnedBehaviorNullEvaluator.cpp b/Engine/Plugins/Animation/RigLogic/Source/RigLogicLib/Private/riglogic/ml/MachineLearnedBehaviorNullEvaluator.cpp
--- a/Engine/Plugins/Animation/RigLogic/Source/RigLogicLib/Private/riglogic/ml/MachineLearnedBehaviorNullEvaluator.cpp
+++ b/Engine/Plugins/Animation/RigLogic/Source/RigLogicLib/Private/riglogic/ml/MachineLearnedBehaviorNullEvaluator.cpp
@@ -12,6 +12,11 @@ const {
                           MachineLearnedBehaviorOutputInstance>::with(instanceMemRes).create();
 }
 
+ConstArrayView<std::uint32_t> MachineLearnedBehaviorNullEvaluator::getNeuralNetworkIndicesForLOD(std::uint16_t  /*unused*/) const
+{
+    return {};
+}
+
 void MachineLearnedBehaviorNullEvaluator::calculate(ControlsInputInstance*  /*unused*/,
                                                     MachineLearnedBehaviorOutputInstance*  /*unused*/,
                                                     std::uint16_t  /*unused*/) const {
diff --git a/Engine/Plugins/Animation/RigLogic/Source/RigLogicLib/Private/riglogic/ml/MachineLearnedBehaviorNullEvaluator.h b/Engine/Plugins/Animation/RigLogic/Source/RigLogicLib/Private/riglogic/ml/MachineLearnedBehaviorNullEvaluator.h
--- a/Engine/Plugins/Animation/RigLogic/Source/RigLogicLib/Private/riglogic/ml/MachineLearnedBehaviorNullEvaluator.h
+++ b/Engine/Plugins/Animation/RigLogic/Source/RigLogicLib/Private/riglogic/ml/MachineLearnedBehaviorNullEvaluator.h
@@ -13,6 +13,7 @@ class ControlsInputInstance;
 class MachineLearnedBehaviorNullEvaluator : public MachineLearnedBehaviorEvaluator {
     public:
         MachineLearnedBehaviorOutputInstance::Pointer createInstance(MemoryResource* instanceMemRes) const override;
+        ConstArrayView<std::uint32_t> getNeuralNetworkIndicesForLOD(std::uint16_t  /*unused*/) const override;
         void calculate(ControlsInputInstance*  /*unused*/, MachineLearnedBehaviorOutputInstance*  /*unused*/,
                        std::uint16_t  /*unused*/) const override;
         void calculate(ControlsInputInstance*  /*unused*/,
diff --git a/Engine/Plugins/Animation/RigLogic/Source/RigLogicLib/Private/riglogic/ml/cpu/CPUMachineLearnedBehaviorEvaluator.h b/Engine/Plugins/Animation/RigLogic/Source/RigLogicLib/Private/riglogic/ml/cpu/CPUMachineLearnedBehaviorEvaluator.h
--- a/Engine/Plugins/Animation/RigLogic/Source/RigLogicLib/Private/riglogic/ml/cpu/CPUMachineLearnedBehaviorEvaluator.h
+++ b/Engine/Plugins/Animation/RigLogic/Source/RigLogicLib/Private/riglogic/ml/cpu/CPUMachineLearnedBehaviorEvaluator.h
@@ -41,6 +41,11 @@ class Evaluator : public MachineLearnedBehaviorEvaluator {
             return instanceFactory(ConstArrayView<std::uint32_t>{maxLayerOutputCounts}, instanceMemRes);
         }
 
+        ConstArrayView<std::uint32_t> getNeuralNetworkIndicesForLOD(std::uint16_t lod) const override {
+            assert(lod < lods.indicesPerLOD.size());
+            return lods.indicesPerLOD[lod];
+        }
+
         void calculate(ControlsInputInstance* inputs, MachineLearnedBehaviorOutputInstance* intermediateOutputs,
                        std::uint16_t lod) const override {
             assert(lod < lods.indicesPerLOD.size());
diff --git a/Engine/Plugins/Animation/RigLogic/Source/RigLogicLib/Private/riglogic/ml/cpu/layers/LayerEvaluator.h b/Engine/Plugins/Animation/RigLogic/Source/RigLogicLib/Private/riglogic/ml/cpu/layers/LayerEvaluator.h
--- a/Engine/Plugins/Animation/RigLogic/Source/RigLogicLib/Private/riglogic/ml/cpu/layers/LayerEvaluator.h
+++ b/Engine/Plugins/Animation/RigLogic/Source/RigLogicLib/Private/riglogic/ml/cpu/layers/LayerEvaluator.h
@@ -62,25 +62,6 @@ static FORCE_INLINE void processBlocks8x4(const float* inputVectorStart,
         const TF256 blk2 = TF256::fromAlignedSource(weights + TF256::size() * 1);
         const TF256 blk3 = TF256::fromAlignedSource(weights + TF256::size() * 2);
         const TF256 blk4 = TF256::fromAlignedSource(weights + TF256::size() * 3);
-        /*
-        #ifndef RL_NO_MANUAL_PREFETCH
-            #ifdef RL_USE_NON_TEMPORAL_PREFETCH
-                TF256::prefetchNTA(values + (lookAheadOffset * 6ul));
-                TF256::prefetchNTA(values + (lookAheadOffset * 7ul));
-                TF256::prefetchNTA(values + (lookAheadOffset * 8ul));
-                TF256::prefetchNTA(values + (lookAheadOffset * 9ul));
-                TF256::prefetchNTA(values + (lookAheadOffset * 10ul));
-                TF256::prefetchNTA(values + (lookAheadOffset * 11ul));
-            #else
-                TF256::prefetchT0(values + (lookAheadOffset * 6ul));
-                TF256::prefetchT0(values + (lookAheadOffset * 7ul));
-                TF256::prefetchT0(values + (lookAheadOffset * 8ul));
-                TF256::prefetchT0(values + (lookAheadOffset * 9ul));
-                TF256::prefetchT0(values + (lookAheadOffset * 10ul));
-                TF256::prefetchT0(values + (lookAheadOffset * 11ul));
-            #endif
-        #endif
-        */
         sum1 += (blk1 * input1);
         sum2 += (blk2 * input2);
         sum3 += (blk3 * input3);
diff --git a/Engine/Plugins/Animation/RigLogic/Source/RigLogicLib/Private/riglogic/rbf/RBFBehavior.cpp b/Engine/Plugins/Animation/RigLogic/Source/RigLogicLib/Private/riglogic/rbf/RBFBehavior.cpp
--- a/Engine/Plugins/Animation/RigLogic/Source/RigLogicLib/Private/riglogic/rbf/RBFBehavior.cpp
+++ b/Engine/Plugins/Animation/RigLogic/Source/RigLogicLib/Private/riglogic/rbf/RBFBehavior.cpp
@@ -19,6 +19,10 @@ RBFBehaviorOutputInstance::Pointer RBFBehavior::createInstance(MemoryResource* i
     return evaluator->createInstance(instanceMemRes);
 }
 
+ConstArrayView<std::uint16_t> RBFBehavior::getSolverIndicesForLOD(std::uint16_t lod) const {
+    return evaluator->getSolverIndicesForLOD(lod);
+}
+
 void RBFBehavior::calculate(ControlsInputInstance* inputs, RBFBehaviorOutputInstance* intermediateOutputs,
                             std::uint16_t lod) const {
     evaluator->calculate(inputs, intermediateOutputs, lod);
diff --git a/Engine/Plugins/Animation/RigLogic/Source/RigLogicLib/Private/riglogic/rbf/RBFBehavior.h b/Engine/Plugins/Animation/RigLogic/Source/RigLogicLib/Private/riglogic/rbf/RBFBehavior.h
--- a/Engine/Plugins/Animation/RigLogic/Source/RigLogicLib/Private/riglogic/rbf/RBFBehavior.h
+++ b/Engine/Plugins/Animation/RigLogic/Source/RigLogicLib/Private/riglogic/rbf/RBFBehavior.h
@@ -22,6 +22,7 @@ class RBFBehavior {
         explicit RBFBehavior(RBFBehaviorEvaluator::Pointer evaluator_);
 
         RBFBehaviorOutputInstance::Pointer createInstance(MemoryResource* instanceMemRes) const;
+        ConstArrayView<std::uint16_t> getSolverIndicesForLOD(std::uint16_t lod) const;
         void calculate(ControlsInputInstance* inputs, RBFBehaviorOutputInstance* intermediateOutputs, std::uint16_t lod) const;
         void calculate(ControlsInputInstance* inputs,
                        RBFBehaviorOutputInstance* intermediateOutputs,
diff --git a/Engine/Plugins/Animation/RigLogic/Source/RigLogicLib/Private/riglogic/rbf/RBFBehaviorEvaluator.h b/Engine/Plugins/Animation/RigLogic/Source/RigLogicLib/Private/riglogic/rbf/RBFBehaviorEvaluator.h
--- a/Engine/Plugins/Animation/RigLogic/Source/RigLogicLib/Private/riglogic/rbf/RBFBehaviorEvaluator.h
+++ b/Engine/Plugins/Animation/RigLogic/Source/RigLogicLib/Private/riglogic/rbf/RBFBehaviorEvaluator.h
@@ -18,6 +18,7 @@ class RBFBehaviorEvaluator {
 
     public:
         virtual RBFBehaviorOutputInstance::Pointer createInstance(MemoryResource* instanceMemRes) const = 0;
+        virtual ConstArrayView<std::uint16_t> getSolverIndicesForLOD(std::uint16_t lod) const = 0;
         virtual void calculate(ControlsInputInstance* inputs, RBFBehaviorOutputInstance* intermediateOutputs,
                                std::uint16_t lod) const = 0;
         virtual void calculate(ControlsInputInstance* inputs,
diff --git a/Engine/Plugins/Animation/RigLogic/Source/RigLogicLib/Private/riglogic/rbf/RBFBehaviorNullEvaluator.cpp b/Engine/Plugins/Animation/RigLogic/Source/RigLogicLib/Private/riglogic/rbf/RBFBehaviorNullEvaluator.cpp
--- a/Engine/Plugins/Animation/RigLogic/Source/RigLogicLib/Private/riglogic/rbf/RBFBehaviorNullEvaluator.cpp
+++ b/Engine/Plugins/Animation/RigLogic/Source/RigLogicLib/Private/riglogic/rbf/RBFBehaviorNullEvaluator.cpp
@@ -12,6 +12,10 @@ const {
                           RBFBehaviorOutputInstance>::with(instanceMemRes).create();
 }
 
+ConstArrayView<std::uint16_t> RBFBehaviorNullEvaluator::getSolverIndicesForLOD(std::uint16_t  /*unused*/) const {
+    return {};
+}
+
 void RBFBehaviorNullEvaluator::calculate(ControlsInputInstance*  /*unused*/, RBFBehaviorOutputInstance*  /*unused*/,
                                          std::uint16_t  /*unused*/) const {
 }
diff --git a/Engine/Plugins/Animation/RigLogic/Source/RigLogicLib/Private/riglogic/rbf/RBFBehaviorNullEvaluator.h b/Engine/Plugins/Animation/RigLogic/Source/RigLogicLib/Private/riglogic/rbf/RBFBehaviorNullEvaluator.h
--- a/Engine/Plugins/Animation/RigLogic/Source/RigLogicLib/Private/riglogic/rbf/RBFBehaviorNullEvaluator.h
+++ b/Engine/Plugins/Animation/RigLogic/Source/RigLogicLib/Private/riglogic/rbf/RBFBehaviorNullEvaluator.h
@@ -13,6 +13,7 @@ class ControlsInputInstance;
 class RBFBehaviorNullEvaluator : public RBFBehaviorEvaluator {
     public:
         RBFBehaviorOutputInstance::Pointer createInstance(MemoryResource* instanceMemRes) const override;
+        ConstArrayView<std::uint16_t> getSolverIndicesForLOD(std::uint16_t  /*unused*/) const override;
         void calculate(ControlsInputInstance*  /*unused*/, RBFBehaviorOutputInstance*  /*unused*/,
                        std::uint16_t  /*unused*/) const override;
         void calculate(ControlsInputInstance*  /*unused*/,
diff --git a/Engine/Plugins/Animation/RigLogic/Source/RigLogicLib/Private/riglogic/rbf/cpu/CPURBFBehaviorEvaluator.h b/Engine/Plugins/Animation/RigLogic/Source/RigLogicLib/Private/riglogic/rbf/cpu/CPURBFBehaviorEvaluator.h
--- a/Engine/Plugins/Animation/RigLogic/Source/RigLogicLib/Private/riglogic/rbf/cpu/CPURBFBehaviorEvaluator.h
+++ b/Engine/Plugins/Animation/RigLogic/Source/RigLogicLib/Private/riglogic/rbf/cpu/CPURBFBehaviorEvaluator.h
@@ -58,6 +58,11 @@ class Evaluator : public RBFBehaviorEvaluator {
             return instanceFactory(maximumInputCount, maxTargetCount, instanceMemRes);
         }
 
+        ConstArrayView<std::uint16_t> getSolverIndicesForLOD(std::uint16_t lod) const override {
+            assert(lod < lods.indicesPerLOD.size());
+            return lods.indicesPerLOD[lod];
+        }
+
         void calculate(std::uint16_t solverIndex,
                        ArrayView<float> rawControls,
                        ArrayView<float> inputBuffer,
diff --git a/Engine/Plugins/Animation/RigLogic/Source/RigLogicLib/Private/riglogic/riglogic/RigInstanceImpl.cpp b/Engine/Plugins/Animation/RigLogic/Source/RigLogicLib/Private/riglogic/riglogic/RigInstanceImpl.cpp
--- a/Engine/Plugins/Animation/RigLogic/Source/RigLogicLib/Private/riglogic/riglogic/RigInstanceImpl.cpp
+++ b/Engine/Plugins/Animation/RigLogic/Source/RigLogicLib/Private/riglogic/riglogic/RigInstanceImpl.cpp
@@ -200,6 +200,7 @@ std::uint16_t RigInstanceImpl::getLOD() const {
 
 void RigInstanceImpl::setLOD(std::uint16_t level) {
     if (level != lodLevel) {
+        controlsInstance->resetInternalInputBuffer();
         jointsInstance->resetOutputBuffer();
         blendShapesInstance->resetOutputBuffer();
         animatedMapsInstance->resetOutputBuffer();
diff --git a/Engine/Plugins/Animation/RigLogic/Source/RigLogicLib/Private/riglogic/riglogic/RigLogicImpl.cpp b/Engine/Plugins/Animation/RigLogic/Source/RigLogicLib/Private/riglogic/riglogic/RigLogicImpl.cpp
--- a/Engine/Plugins/Animation/RigLogic/Source/RigLogicLib/Private/riglogic/riglogic/RigLogicImpl.cpp
+++ b/Engine/Plugins/Animation/RigLogic/Source/RigLogicLib/Private/riglogic/riglogic/RigLogicImpl.cpp
@@ -13,6 +13,8 @@
 #include "riglogic/riglogic/ConfigurationSerializer.h"
 #include "riglogic/riglogic/RigInstanceImpl.h"
 #include "riglogic/riglogic/RigMetrics.h"
+#include "riglogic/riglogic/Stats.h"
+#include "riglogic/system/simd/Utils.h"
 #include "riglogic/utils/Extd.h"
 
 #ifdef _MSC_VER
@@ -56,17 +58,19 @@ static RigMetrics::Pointer computeRigMetrics(const dna::Reader* reader, const Co
 RigLogic::~RigLogic() = default;
 
 RigLogic* RigLogic::create(const dna::Reader* reader, const Configuration& config, MemoryResource* memRes) {
+    const ActiveFeatures activeFeatures = getActiveFeatures(config);
     auto metrics = computeRigMetrics(reader, config, memRes);
 
     auto controls = ControlsFactory::create(config, reader, memRes);
     auto machineLearnedBlendShapes = MachineLearnedBehaviorFactory::create(config, reader, memRes);
     auto rbfBehavior = RBFBehaviorFactory::create(config, reader, memRes);
-    auto joints = JointsFactory::create(config, reader, memRes);
-    auto blendShapes = BlendShapesFactory::create(config, reader, memRes);
-    auto animatedMaps = AnimatedMapsFactory::create(config, reader, memRes);
+    auto joints = JointsFactory::create(config, reader, controls.get(), memRes);
+    auto blendShapes = BlendShapesFactory::create(config, reader, controls.get(), memRes);
+    auto animatedMaps = AnimatedMapsFactory::create(config, reader, controls.get(), memRes);
 
     PolyAllocator<RigLogicImpl> alloc{memRes};
     return alloc.newObject(config,
+                           activeFeatures,
                            std::move(metrics),
                            std::move(controls),
                            std::move(machineLearnedBlendShapes),
@@ -92,6 +96,8 @@ RigLogic* RigLogic::restore(BoundedIOStream* source, MemoryResource* memRes) {
     Configuration config;
     archive >> config;
 
+    const ActiveFeatures activeFeatures = getActiveFeatures(config);
+
     RigMetrics::Pointer metrics = UniqueInstance<RigMetrics>::with(memRes).create(memRes);
     archive >> *metrics;
 
@@ -107,6 +113,7 @@ RigLogic* RigLogic::restore(BoundedIOStream* source, MemoryResource* memRes) {
 
     archive >> *controls >> *machineLearnedBehavior >> *rbfBehavior >> *joints >> blendShapesProxy >> animatedMapsProxy;
     return alloc.newObject(config,
+                           activeFeatures,
                            std::move(metrics),
                            std::move(controls),
                            std::move(machineLearnedBehavior),
@@ -118,6 +125,7 @@ RigLogic* RigLogic::restore(BoundedIOStream* source, MemoryResource* memRes) {
 }
 
 RigLogicImpl::RigLogicImpl(const Configuration& config_,
+                           ActiveFeatures activeFeatures_,
                            RigMetrics::Pointer metrics_,
                            Controls::Pointer controls_,
                            MachineLearnedBehavior::Pointer machineLearnedBehavior_,
@@ -128,6 +136,7 @@ RigLogicImpl::RigLogicImpl(const Configuration& config_,
                            MemoryResource* memRes_) :
     memRes{memRes_},
     config{config_},
+    activeFeatures{activeFeatures_},
     metrics{std::move(metrics_)},
     controls{std::move(controls_)},
     machineLearnedBehavior{std::move(machineLearnedBehavior_)},
@@ -165,6 +174,26 @@ std::uint16_t RigLogicImpl::getLODCount() const {
     return metrics->lodCount;
 }
 
+ConstArrayView<std::uint16_t> RigLogicImpl::getRBFSolverIndicesForLOD(std::uint16_t lod) const {
+    return rbfBehavior->getSolverIndicesForLOD(lod);
+}
+
+ConstArrayView<std::uint32_t> RigLogicImpl::getNeuralNetworkIndicesForLOD(std::uint16_t lod) const {
+    return machineLearnedBehavior->getNeuralNetworkIndicesForLOD(lod);
+}
+
+ConstArrayView<std::uint16_t> RigLogicImpl::getBlendShapeChannelIndicesForLOD(std::uint16_t lod) const {
+    return blendShapes->getBlendShapeChannelIndicesForLOD(lod);
+}
+
+ConstArrayView<std::uint16_t> RigLogicImpl::getAnimatedMapIndicesForLOD(std::uint16_t lod) const {
+    return animatedMaps->getAnimatedMapIndicesForLOD(lod);
+}
+
+ConstArrayView<std::uint16_t> RigLogicImpl::getJointIndicesForLOD(std::uint16_t lod) const {
+    return joints->getJointIndicesForLOD(lod);
+}
+
 ConstArrayView<float> RigLogicImpl::getNeutralJointValues() const {
     return joints->getNeutralValues();
 }
@@ -235,7 +264,7 @@ void RigLogicImpl::mapRawToGUIControls(RigInstance* instance) const {
 
 void RigLogicImpl::calculateControls(RigInstance* instance) const {
     auto pRigInstance = castInstance(instance);
-    controls->calculate(pRigInstance->getControlsInputInstance());
+    controls->calculate(pRigInstance->getControlsInputInstance(), pRigInstance->getLOD());
 }
 
 void RigLogicImpl::calculateMachineLearnedBehaviorControls(RigInstance* instance) const {
@@ -304,6 +333,19 @@ void RigLogicImpl::calculate(RigInstance* instance) const {
     calculateAnimatedMaps(instance);
 }
 
+void RigLogicImpl::collectCalculationStats(const RigInstance* instance, Stats* stats) const {
+    const auto lod = instance->getLOD();
+    stats->calculationType = activeFeatures.calculationType;
+    stats->floatingPointType = activeFeatures.floatingPointType;
+    stats->rbfSolverCount = static_cast<std::uint16_t>(rbfBehavior->getSolverIndicesForLOD(lod).size());
+    stats->neuralNetworkCount = static_cast<std::uint16_t>(machineLearnedBehavior->getNeuralNetworkIndicesForLOD(lod).size());
+    stats->psdCount = static_cast<std::uint16_t>(controls->getPSDIndicesForLOD(lod).size());
+    stats->blendShapeChannelCount = static_cast<std::uint16_t>(blendShapes->getBlendShapeChannelIndicesForLOD(lod).size());
+    stats->animatedMapCount = static_cast<std::uint16_t>(animatedMaps->getAnimatedMapIndicesForLOD(lod).size());
+    stats->jointCount = static_cast<std::uint16_t>(joints->getJointIndicesForLOD(lod).size());
+    stats->jointDeltaValueCount = static_cast<std::uint16_t>(joints->getJointDeltaValueCountForLOD(lod));
+}
+
 MemoryResource* RigLogicImpl::getMemoryResource() {
     return memRes;
 }
diff --git a/Engine/Plugins/Animation/RigLogic/Source/RigLogicLib/Private/riglogic/riglogic/RigLogicImpl.h b/Engine/Plugins/Animation/RigLogic/Source/RigLogicLib/Private/riglogic/riglogic/RigLogicImpl.h
--- a/Engine/Plugins/Animation/RigLogic/Source/RigLogicLib/Private/riglogic/riglogic/RigLogicImpl.h
+++ b/Engine/Plugins/Animation/RigLogic/Source/RigLogicLib/Private/riglogic/riglogic/RigLogicImpl.h
@@ -17,12 +17,14 @@
 #include "riglogic/riglogic/Configuration.h"
 #include "riglogic/riglogic/RigLogic.h"
 #include "riglogic/riglogic/RigMetrics.h"
+#include "riglogic/system/simd/Utils.h"
 
 namespace rl4 {
 
 class RigLogicImpl : public RigLogic {
     public:
         RigLogicImpl(const Configuration& config_,
+                     ActiveFeatures activeFeatures_,
                      RigMetrics::Pointer metrics_,
                      Controls::Pointer controls_,
                      MachineLearnedBehavior::Pointer machineLearnedBehavior_,
@@ -33,9 +35,14 @@ class RigLogicImpl : public RigLogic {
                      MemoryResource* memRes_);
 
         void dump(BoundedIOStream* destination) const override;
-        const Configuration& getConfiguration() const;
+        const Configuration& getConfiguration() const override;
         const RigMetrics& getRigMetrics() const;
         std::uint16_t getLODCount() const override;
+        ConstArrayView<std::uint16_t> getRBFSolverIndicesForLOD(std::uint16_t lod) const override;
+        ConstArrayView<std::uint32_t> getNeuralNetworkIndicesForLOD(std::uint16_t lod) const override;
+        ConstArrayView<std::uint16_t> getBlendShapeChannelIndicesForLOD(std::uint16_t lod) const override;
+        ConstArrayView<std::uint16_t> getAnimatedMapIndicesForLOD(std::uint16_t lod) const override;
+        ConstArrayView<std::uint16_t> getJointIndicesForLOD(std::uint16_t lod) const override;
         ConstArrayView<float> getNeutralJointValues() const override;
         ConstArrayView<std::uint16_t> getJointVariableAttributeIndices(std::uint16_t lod) const override;
         std::uint16_t getJointGroupCount() const override;
@@ -64,12 +71,14 @@ class RigLogicImpl : public RigLogic {
         void calculateBlendShapes(RigInstance* instance) const override;
         void calculateAnimatedMaps(RigInstance* instance) const override;
         void calculate(RigInstance* instance) const override;
+        void collectCalculationStats(const RigInstance* instance, Stats* stats) const override;
 
         MemoryResource* getMemoryResource();
 
     private:
         MemoryResource* memRes;
         Configuration config;
+        ActiveFeatures activeFeatures;
         RigMetrics::Pointer metrics;
         Controls::Pointer controls;
         MachineLearnedBehavior::Pointer machineLearnedBehavior;
diff --git a/Engine/Plugins/Animation/RigLogic/Source/RigLogicLib/Public/riglogic/riglogic/Configuration.h b/Engine/Plugins/Animation/RigLogic/Source/RigLogicLib/Public/riglogic/riglogic/Configuration.h
--- a/Engine/Plugins/Animation/RigLogic/Source/RigLogicLib/Public/riglogic/riglogic/Configuration.h
+++ b/Engine/Plugins/Animation/RigLogic/Source/RigLogicLib/Public/riglogic/riglogic/Configuration.h
@@ -19,6 +19,14 @@ enum class CalculationType : std::uint8_t {
     AnyVector  ///< Pick any available vectorization
 };
 
+/**
+    @brief Floating point type used in vectorized calculations.
+*/
+enum class FloatingPointType : std::uint8_t {
+    Float,
+    HalfFloat
+};
+
 /**
     @brief Translation type to be used by RigLogic.
 */
@@ -65,6 +73,9 @@ struct Configuration {
     RotationType rotationType = RotationType::EulerAngles;
     RotationOrder rotationOrder = RotationOrder::XYZ;
     ScaleType scaleType = ScaleType::Vector;
+    float translationPruningThreshold = 0.0f;  // Reasonably safe to try 0.0001f;
+    float rotationPruningThreshold = 0.0f;  // Reasonably safe to try 0.1f
+    float scalePruningThreshold = 0.0f;  // Reasonably safe to try 0.001f;
 };
 
 }  // namespace rl4
diff --git a/Engine/Plugins/Animation/RigLogic/Source/RigLogicLib/Public/riglogic/riglogic/RigLogic.h b/Engine/Plugins/Animation/RigLogic/Source/RigLogicLib/Public/riglogic/riglogic/RigLogic.h
--- a/Engine/Plugins/Animation/RigLogic/Source/RigLogicLib/Public/riglogic/riglogic/RigLogic.h
+++ b/Engine/Plugins/Animation/RigLogic/Source/RigLogicLib/Public/riglogic/riglogic/RigLogic.h
@@ -4,6 +4,7 @@
 
 #include "riglogic/Defs.h"
 #include "riglogic/riglogic/Configuration.h"
+#include "riglogic/riglogic/Stats.h"
 #include "riglogic/types/Aliases.h"
 
 #include <cstdint>
@@ -75,12 +76,46 @@ class RLAPI RigLogic {
             @see restore
         */
         virtual void dump(BoundedIOStream* destination) const = 0;
+        /**
+            @brief Retrieve the configuration that RigLogic was initialized with.
+        */
+        virtual const Configuration& getConfiguration() const = 0;
         /**
             @brief Available levels of detail (e.g. 6 which means the following levels are available:
                 [0,1,2,3,4,5], where 0 is the LOD with the highest details, and 5 is the LOD with
                 lowest details).
         */
         virtual std::uint16_t getLODCount() const = 0;
+        /**
+            @brief List of RBF solver indices that need to be computed on the specified LOD.
+            @param lod
+                The LOD for which the indices are requested.
+        */
+        virtual ConstArrayView<std::uint16_t> getRBFSolverIndicesForLOD(std::uint16_t lod) const = 0;
+        /**
+            @brief List of neural network indices that need to be computed on the specified LOD.
+            @param lod
+                The LOD for which the indices are requested.
+        */
+        virtual ConstArrayView<std::uint32_t> getNeuralNetworkIndicesForLOD(std::uint16_t lod) const = 0;
+        /**
+            @brief List of blend shape channel indices that need to be computed on the specified LOD.
+            @param lod
+                The LOD for which the indices are requested.
+        */
+        virtual ConstArrayView<std::uint16_t> getBlendShapeChannelIndicesForLOD(std::uint16_t lod) const = 0;
+        /**
+            @brief List of animated map indices that need to be computed on the specified LOD.
+            @param lod
+                The LOD for which the indices are requested.
+        */
+        virtual ConstArrayView<std::uint16_t> getAnimatedMapIndicesForLOD(std::uint16_t lod) const = 0;
+        /**
+            @brief List of joint indices that need to be computed on the specified LOD.
+            @param lod
+                The LOD for which the indices are requested.
+        */
+        virtual ConstArrayView<std::uint16_t> getJointIndicesForLOD(std::uint16_t lod) const = 0;
         /**
             @brief Neutral values of joint transformations.
             @note
@@ -245,6 +280,14 @@ class RLAPI RigLogic {
                 The rig instance whose outputs are to be calculated.
         */
         virtual void calculate(RigInstance* instance) const = 0;
+        /**
+            @brief Collect stats that summarize the overall amount of data that participate in computation for the current LOD.
+            @param instance
+                The rig instance which provides the current LOD to compute the necessary stats.
+            @param stats
+                The object into which the collected data will be written.
+        */
+        virtual void collectCalculationStats(const RigInstance* instance, Stats* stats) const = 0;
 
 };
 
diff --git a/Engine/Plugins/Animation/RigLogic/Source/RigLogicLibTest/Private/rltests/animatedmaps/TestAnimatedMaps.cpp b/Engine/Plugins/Animation/RigLogic/Source/RigLogicLibTest/Private/rltests/animatedmaps/TestAnimatedMaps.cpp
--- a/Engine/Plugins/Animation/RigLogic/Source/RigLogicLibTest/Private/rltests/animatedmaps/TestAnimatedMaps.cpp
+++ b/Engine/Plugins/Animation/RigLogic/Source/RigLogicLibTest/Private/rltests/animatedmaps/TestAnimatedMaps.cpp
@@ -41,7 +41,7 @@ TEST_P(AnimatedMapsTest, LODLimitsCondTableSize) {
         {0.3f, 0.6f},  // LOD0
         {0.3f, 0.0f}  // LOD1
     };
-    rl4::Vector<float> initialValues(conditionalTableInputs.size());
+    rl4::Vector<rl4::ControlInitializer> initialValues;
     auto inputInstance = inputInstanceFactory(initialValues, &amr);
     auto inputBuffer = inputInstance->getInputBuffer();
     std::copy(conditionalTableInputs.begin(), conditionalTableInputs.end(), inputBuffer.begin());
diff --git a/Engine/Plugins/Animation/RigLogic/Source/RigLogicLibTest/Private/rltests/blendshapes/TestBlendShapes.cpp b/Engine/Plugins/Animation/RigLogic/Source/RigLogicLibTest/Private/rltests/blendshapes/TestBlendShapes.cpp
--- a/Engine/Plugins/Animation/RigLogic/Source/RigLogicLibTest/Private/rltests/blendshapes/TestBlendShapes.cpp
+++ b/Engine/Plugins/Animation/RigLogic/Source/RigLogicLibTest/Private/rltests/blendshapes/TestBlendShapes.cpp
@@ -37,7 +37,7 @@ TEST_P(BlendShapesTest, Calculate) {
                                                                     0,
                                                                     0,
                                                                     0);
-    rl4::Vector<float> initialValues(blendShapeInputs.size());
+    rl4::Vector<rl4::ControlInitializer> initialValues;
     auto inputInstance = inputInstanceFactory(initialValues, &amr);
     auto inputBuffer = inputInstance->getInputBuffer();
     std::copy(blendShapeInputs.begin(), blendShapeInputs.end(), inputBuffer.begin());
diff --git a/Engine/Plugins/Animation/RigLogic/Source/RigLogicLibTest/Private/rltests/controls/ControlFixtures.cpp b/Engine/Plugins/Animation/RigLogic/Source/RigLogicLibTest/Private/rltests/controls/ControlFixtures.cpp
--- a/Engine/Plugins/Animation/RigLogic/Source/RigLogicLibTest/Private/rltests/controls/ControlFixtures.cpp
+++ b/Engine/Plugins/Animation/RigLogic/Source/RigLogicLibTest/Private/rltests/controls/ControlFixtures.cpp
@@ -11,7 +11,7 @@ rl4::ControlsInputInstance::Factory ControlsFactory::getInstanceFactory(std::uin
                                                                         std::uint16_t psdControlCount,
                                                                         std::uint16_t mlControlCount,
                                                                         std::uint16_t rbfControlCount) {
-    return [ = ](rl4::ConstArrayView<float> initialValues, rl4::MemoryResource* memRes) {
+    return [ = ](rl4::ConstArrayView<rl4::ControlInitializer> initialValues, rl4::MemoryResource* memRes) {
                return pma::UniqueInstance<rl4::StandardControlsInputInstance, rl4::ControlsInputInstance>::with(memRes).create(
                    guiControlCount,
                    rawControlCount,
diff --git a/Engine/Plugins/Animation/RigLogic/Source/RigLogicLibTest/Private/rltests/controls/TestControls.cpp b/Engine/Plugins/Animation/RigLogic/Source/RigLogicLibTest/Private/rltests/controls/TestControls.cpp
--- a/Engine/Plugins/Animation/RigLogic/Source/RigLogicLibTest/Private/rltests/controls/TestControls.cpp
+++ b/Engine/Plugins/Animation/RigLogic/Source/RigLogicLibTest/Private/rltests/controls/TestControls.cpp
@@ -6,6 +6,7 @@
 
 #include "riglogic/TypeDefs.h"
 #include "riglogic/controls/Controls.h"
+#include "riglogic/controls/psdnet/PSDNet.h"
 
 #ifdef _MSC_VER
     #pragma warning(push)
@@ -20,13 +21,13 @@ TEST(ControlsTest, GUIToRawMapping) {
     pma::AlignedMemoryResource amr;
     auto conditionals = ConditionalTableFactory::withMultipleIODefaults(&amr);
 
-    rl4::PSDMatrix psds{0u, {}, {}, {}};
+    rl4::PSDNet psds{{}, {}, {}, {}, {}, {}};
     const std::uint16_t guiControlCount = conditionals.getInputCount();
     const std::uint16_t rawControlCount = conditionals.getOutputCount();
-    const std::uint16_t psdControlCount = psds.getDistinctPSDCount();
+    const std::uint16_t psdControlCount = psds.getPSDCount();
     auto instanceFactory = ControlsFactory::getInstanceFactory(guiControlCount, rawControlCount, psdControlCount, 0u, 0u);
 
-    rl4::Vector<float> initialValues(static_cast<std::size_t>(rawControlCount) + static_cast<std::size_t>(psdControlCount));
+    rl4::Vector<rl4::ControlInitializer> initialValues;
     rl4::Controls controls{std::move(conditionals), std::move(psds), std::move(initialValues), instanceFactory};
 
     const rl4::Vector<float> guiControls{0.1f, 0.2f};
@@ -47,25 +48,32 @@ TEST(ControlsTest, PSDsAppendToOutput) {
     auto conditionals = ConditionalTableFactory::withMultipleIODefaults(&amr);
 
     rl4::Vector<float> rawControls{0.1f, 0.2f};
-    const rl4::Vector<float> expected{0.1f, 0.2f, 0.4f, 0.6f};
+    const rl4::Vector<float> expected{0.1f, 0.2f, 0.24f, 0.02f};
 
-    rl4::Vector<std::uint16_t> rows{2u, 3u};
-    rl4::Vector<std::uint16_t> cols{0u, 1u};
-    rl4::Vector<float> values{4.0f, 3.0f};
-    rl4::PSDMatrix psds{2u, std::move(rows), std::move(cols), std::move(values)};
+    rl4::Matrix<std::uint16_t> inputLODs{{0u, 1u}};
+    rl4::Matrix<std::uint16_t> outputLODs{{2u, 3u}};
+    rl4::Vector<std::uint16_t> cols{0u, 1u, 0u, 1u};
+    // PSD weights {4.0f, 3.0f, 0.5f, 2.0f}
+    rl4::Vector<rl4::PSD> psds{{0u, 2u, 12.0f}, {2u, 2u, 1.0f}};
+    rl4::PSDNet psdNet{std::move(inputLODs),
+                       std::move(outputLODs),
+                       std::move(cols),
+                       std::move(psds),
+                       2u,
+                       3u};
 
     const std::uint16_t guiControlCount = conditionals.getInputCount();
     const std::uint16_t rawControlCount = conditionals.getOutputCount();
-    const std::uint16_t psdControlCount = psds.getDistinctPSDCount();
+    const std::uint16_t psdControlCount = psdNet.getPSDCount();
     auto instanceFactory = ControlsFactory::getInstanceFactory(guiControlCount, rawControlCount, psdControlCount, 0u, 0u);
-    rl4::Vector<float> initialValues(static_cast<std::size_t>(rawControlCount) + static_cast<std::size_t>(psdControlCount));
-    rl4::Controls controls{std::move(conditionals), std::move(psds), std::move(initialValues), instanceFactory};
+    rl4::Vector<rl4::ControlInitializer> initialValues;
+    rl4::Controls controls{std::move(conditionals), std::move(psdNet), std::move(initialValues), instanceFactory};
     auto instance = controls.createInstance(&amr);
     auto buffer = instance->getInputBuffer();
     std::copy(rawControls.begin(), rawControls.end(), buffer.begin());
 
-    controls.calculate(instance.get());
+    controls.calculate(instance.get(), 0u);
 
     ASSERT_EQ(buffer.size(), expected.size());
-    ASSERT_ELEMENTS_EQ(buffer, expected, expected.size());
+    ASSERT_ELEMENTS_NEAR(buffer, expected, expected.size(), 0.0001f);
 }
diff --git a/Engine/Plugins/Animation/RigLogic/Source/RigLogicLibTest/Private/rltests/controls/TestPSDNet.cpp b/Engine/Plugins/Animation/RigLogic/Source/RigLogicLibTest/Private/rltests/controls/TestPSDNet.cpp
new file mode 100644
--- /dev/null
+++ b/Engine/Plugins/Animation/RigLogic/Source/RigLogicLibTest/Private/rltests/controls/TestPSDNet.cpp
@@ -0,0 +1,75 @@
+// Copyright Epic Games, Inc. All Rights Reserved.
+
+#include "rltests/Defs.h"
+
+#include "riglogic/TypeDefs.h"
+#include "riglogic/controls/psdnet/PSDNet.h"
+
+TEST(PSDNetTest, OutputsAreClamped) {
+    pma::AlignedMemoryResource amr;
+
+    rl4::Matrix<std::uint16_t> inputLODs{{0u}};
+    rl4::Matrix<std::uint16_t> outputLODs{{1u}};
+    rl4::Vector<std::uint16_t> cols{0u};
+    // PSD weights {100.0f}
+    rl4::Vector<rl4::PSD> psds{{0u, 1u, 100.0f}};
+    rl4::PSDNet PSDNet{std::move(inputLODs),
+                       std::move(outputLODs),
+                       std::move(cols),
+                       std::move(psds),
+                       1u,
+                       1u};
+
+    float inputs[] = {0.1f, 0.0f};
+    float clampBuffer[] = {0.0f, 0.0f};
+    const float expected[] = {1.0f};
+
+    PSDNet.calculate({inputs, 2ul}, {clampBuffer, 2ul}, 0u);
+    ASSERT_ELEMENTS_EQ((inputs + 1ul), expected, 1);
+}
+
+TEST(PSDNetTest, OutputsKeepExistingProduct) {
+    pma::AlignedMemoryResource amr;
+
+    rl4::Matrix<std::uint16_t> inputLODs{{0u, 1u}};
+    rl4::Matrix<std::uint16_t> outputLODs{{2u}};
+    rl4::Vector<std::uint16_t> cols{0u, 1u};
+    // PSD weights {4.0f, 10.0f}
+    rl4::Vector<rl4::PSD> psds{{0u, 2u, 40.0f}};
+    rl4::PSDNet PSDNet{std::move(inputLODs),
+                       std::move(outputLODs),
+                       std::move(cols),
+                       std::move(psds),
+                       2u,
+                       2u};
+
+    float inputs[] = {0.1f, 0.2f, 0.0f};
+    float clampBuffer[] = {0.0f, 0.0f, 0.0f};
+    const float expected[] = {0.8f};
+
+    PSDNet.calculate({inputs, 3ul}, {clampBuffer, 3ul}, 0ul);
+    ASSERT_ELEMENTS_EQ((inputs + 2ul), expected, 1);
+}
+
+TEST(PSDNetTest, RowsSpecifyDestinationIndex) {
+    pma::AlignedMemoryResource amr;
+
+    rl4::Matrix<std::uint16_t> inputLODs{{0u, 1u}};
+    rl4::Matrix<std::uint16_t> outputLODs{{2u, 3u}};
+    rl4::Vector<std::uint16_t> cols{0u, 1u};
+    // PSD weights {4.0f, 3.0f}
+    rl4::Vector<rl4::PSD> psds{{0u, 1u, 4.0f}, {1u, 1u, 3.0f}};
+    rl4::PSDNet PSDNet{std::move(inputLODs),
+                       std::move(outputLODs),
+                       std::move(cols),
+                       std::move(psds),
+                       2u,
+                       2u};
+
+    float inputs[] = {0.1f, 0.2f, 0.0f, 0.0f};
+    float clampBuffer[] = {0.0f, 0.0f, 0.0f, 0.0f};
+    const float expected[] = {0.4f, 0.6f};
+
+    PSDNet.calculate({inputs, 4ul}, {clampBuffer, 4ul}, 0ul);
+    ASSERT_ELEMENTS_EQ((inputs + 2ul), expected, 2);
+}
diff --git a/Engine/Plugins/Animation/RigLogic/Source/RigLogicLibTest/Private/rltests/dna/DNAFixtures.cpp b/Engine/Plugins/Animation/RigLogic/Source/RigLogicLibTest/Private/rltests/dna/DNAFixtures.cpp
--- a/Engine/Plugins/Animation/RigLogic/Source/RigLogicLibTest/Private/rltests/dna/DNAFixtures.cpp
+++ b/Engine/Plugins/Animation/RigLogic/Source/RigLogicLibTest/Private/rltests/dna/DNAFixtures.cpp
@@ -448,23 +448,23 @@ const unsigned char conditionals[] {
 const unsigned char psds[] {
     // Rows
     0x00, 0x00, 0x00, 0x18,  // Row index count
-    0x00, 0x08,  // Index:  8  C1
-    0x00, 0x08,  // Index:  8  C1
-    0x00, 0x08,  // Index:  8  C1
-    0x00, 0x09,  // Index:  9      C2
-    0x00, 0x09,  // Index:  9      C2
-    0x00, 0x0a,  // Index: 10  C1
-    0x00, 0x0a,  // Index: 10  C1
-    0x00, 0x0a,  // Index: 10  C1
-    0x00, 0x0b,  // Index: 11      C2
+    0x00, 0x09,  // Index:  9  C1
+    0x00, 0x09,  // Index:  9  C1
+    0x00, 0x09,  // Index:  9  C1
+    0x00, 0x0a,  // Index: 10      C2
+    0x00, 0x0a,  // Index: 10      C2
+    0x00, 0x0b,  // Index: 11  C1
+    0x00, 0x0b,  // Index: 11  C1
+    0x00, 0x0b,  // Index: 11  C1
     0x00, 0x0c,  // Index: 12      C2
-    0x00, 0x0d,  // Index: 13  C1
-    0x00, 0x0d,  // Index: 13  C1
-    0x00, 0x0d,  // Index: 13  C1
+    0x00, 0x0d,  // Index: 13      C2
     0x00, 0x0e,  // Index: 14  C1
     0x00, 0x0e,  // Index: 14  C1
+    0x00, 0x0e,  // Index: 14  C1
+    0x00, 0x0f,  // Index: 15  C1
     0x00, 0x0f,  // Index: 15  C1
-    0x00, 0x10,  // Index: 16      C2
+    0x00, 0x10,  // Index: 16  C1
+    0x00, 0x11,  // Index: 17      C2
     0x00, 0x12,  // Index: 18      C2
     0x00, 0x12,  // Index: 18      C2
     0x00, 0x12,  // Index: 18      C2
diff --git a/Engine/Plugins/Animation/RigLogic/Source/RigLogicLibTest/Private/rltests/joints/TestJointsFactory.cpp b/Engine/Plugins/Animation/RigLogic/Source/RigLogicLibTest/Private/rltests/joints/TestJointsFactory.cpp
--- a/Engine/Plugins/Animation/RigLogic/Source/RigLogicLibTest/Private/rltests/joints/TestJointsFactory.cpp
+++ b/Engine/Plugins/Animation/RigLogic/Source/RigLogicLibTest/Private/rltests/joints/TestJointsFactory.cpp
@@ -5,6 +5,7 @@
 
 #include "riglogic/RigLogic.h"
 #include "riglogic/TypeDefs.h"
+#include "riglogic/controls/ControlsFactory.h"
 #include "riglogic/joints/JointsFactory.h"
 #include "riglogic/riglogic/Configuration.h"
 #include "riglogic/riglogic/RigMetrics.h"
@@ -14,7 +15,8 @@ TEST(ScalarJointsFactoryTest, NeutralJointsAreCopied) {
     block4::CanonicalReader reader;
     rl4::Configuration config{};
     config.calculationType = rl4::CalculationType::Scalar;
-    auto joints = rl4::JointsFactory::create(config, &reader, &memRes);
+    auto controls = rl4::ControlsFactory::create(config, &reader, &memRes);
+    auto joints = rl4::JointsFactory::create(config, &reader, controls.get(), &memRes);
     const float expected[] = {
         0.0f, 1.0f, 2.0f, 3.0f, 4.0f, 5.0f, 1.0f, 1.0f, 1.0f,
         6.0f, 7.0f, 8.0f, 9.0f, 10.0f, 11.0f, 1.0f, 1.0f, 1.0f,
diff --git a/Engine/Plugins/Animation/RigLogic/Source/RigLogicLibTest/Private/rltests/joints/bpcm/Assertions.h b/Engine/Plugins/Animation/RigLogic/Source/RigLogicLibTest/Private/rltests/joints/bpcm/Assertions.h
--- a/Engine/Plugins/Animation/RigLogic/Source/RigLogicLibTest/Private/rltests/joints/bpcm/Assertions.h
+++ b/Engine/Plugins/Animation/RigLogic/Source/RigLogicLibTest/Private/rltests/joints/bpcm/Assertions.h
@@ -27,9 +27,9 @@ template<typename TValue>
 struct Evaluator<TValue>::Accessor {
 
     static void assertRawDataEqual(const Evaluator<TValue>& result, const Evaluator<TValue>& expected) {
+        ASSERT_EQ(result.storage.values, expected.storage.values);
         ASSERT_EQ(result.storage.inputIndices, expected.storage.inputIndices);
         ASSERT_EQ(result.storage.outputIndices, expected.storage.outputIndices);
-        ASSERT_EQ(result.storage.values, expected.storage.values);
         ASSERT_EQ(result.storage.outputRotationIndices, expected.storage.outputRotationIndices);
         ASSERT_EQ(result.storage.outputRotationLODs, expected.storage.outputRotationLODs);
     }
@@ -40,15 +40,14 @@ struct Evaluator<TValue>::Accessor {
             const auto& jointGroup = result.storage.jointGroups[jgIdx];
             const auto& expectedJointGroup = expected.storage.jointGroups[jgIdx];
             ASSERT_EQ(jointGroup.inputIndicesOffset, expectedJointGroup.inputIndicesOffset);
-            ASSERT_EQ(jointGroup.inputIndicesSize, expectedJointGroup.inputIndicesSize);
-            ASSERT_EQ(jointGroup.inputIndicesSizeAlignedTo4, expectedJointGroup.inputIndicesSizeAlignedTo4);
-            ASSERT_EQ(jointGroup.inputIndicesSizeAlignedTo8, expectedJointGroup.inputIndicesSizeAlignedTo8);
             ASSERT_EQ(jointGroup.lodsOffset, expectedJointGroup.lodsOffset);
             ASSERT_EQ(jointGroup.outputIndicesOffset, expectedJointGroup.outputIndicesOffset);
             ASSERT_EQ(jointGroup.outputRotationIndicesOffset, expectedJointGroup.outputRotationIndicesOffset);
             ASSERT_EQ(jointGroup.outputRotationLODsOffset, expectedJointGroup.outputRotationLODsOffset);
             ASSERT_EQ(jointGroup.valuesOffset, expectedJointGroup.valuesOffset);
             ASSERT_EQ(jointGroup.valuesSize, expectedJointGroup.valuesSize);
+            ASSERT_EQ(jointGroup.colCount, expectedJointGroup.colCount);
+            ASSERT_EQ(jointGroup.rowCount, expectedJointGroup.rowCount);
         }
     }
 
@@ -57,9 +56,13 @@ struct Evaluator<TValue>::Accessor {
         for (std::size_t lod = 0; lod < expected.storage.lodRegions.size(); ++lod) {
             const auto& lodRegion = result.storage.lodRegions[lod];
             const auto& expectedLodRegion = expected.storage.lodRegions[lod];
-            ASSERT_EQ(lodRegion.size, expectedLodRegion.size);
-            ASSERT_EQ(lodRegion.sizePaddedToLastFullBlock, expectedLodRegion.sizePaddedToLastFullBlock);
-            ASSERT_EQ(lodRegion.sizePaddedToSecondLastFullBlock, expectedLodRegion.sizePaddedToSecondLastFullBlock);
+            ASSERT_EQ(lodRegion.inputLODs.size, expectedLodRegion.inputLODs.size);
+            ASSERT_EQ(lodRegion.inputLODs.sizeAlignedTo4, expectedLodRegion.inputLODs.sizeAlignedTo4);
+            ASSERT_EQ(lodRegion.inputLODs.sizeAlignedTo8, expectedLodRegion.inputLODs.sizeAlignedTo8);
+            ASSERT_EQ(lodRegion.outputLODs.size, expectedLodRegion.outputLODs.size);
+            ASSERT_EQ(lodRegion.outputLODs.sizePaddedToLastFullBlock, expectedLodRegion.outputLODs.sizePaddedToLastFullBlock);
+            ASSERT_EQ(lodRegion.outputLODs.sizePaddedToSecondLastFullBlock,
+                      expectedLodRegion.outputLODs.sizePaddedToSecondLastFullBlock);
         }
     }
 
