Optimize RigLogic for low LOD evaluation

diff --git a/Engine/Plugins/Animation/RigLogic/Source/RigLogicLibTest/Private/rltests/joints/bpcm/BPCMFixturesBlock4.cpp b/Engine/Plugins/Animation/RigLogic/Source/RigLogicLibTest/Private/rltests/joints/bpcm/BPCMFixturesBlock4.cpp
--- a/Engine/Plugins/Animation/RigLogic/Source/RigLogicLibTest/Private/rltests/joints/bpcm/BPCMFixturesBlock4.cpp
+++ b/Engine/Plugins/Animation/RigLogic/Source/RigLogicLibTest/Private/rltests/joints/bpcm/BPCMFixturesBlock4.cpp
@@ -139,7 +139,35 @@ const Matrix<float> values = {
         2.0f, 2.0f, 2.0f, 2.0f, 2.0f,
         3.0f, 3.0f, 3.0f, 3.0f, 3.0f,
         4.0f, 4.0f, 4.0f, 4.0f, 4.0f
-    }
+    },
+    {  // Joint group 12 - Test case for joint group optimizer, order columns by LODs
+        // LOD-3
+        0.0f, 0.0f, 0.0f, 0.0f, 0.0f, 0.0f, 0.0f, 0.0f, 0.0f, 0.0f,
+        0.0f, 1.0f, 0.0f, 0.0f, 0.0f, 4.0f, 5.0f, 0.0f, 0.0f, 8.0f,  // tx
+        0.0f, 1.0f, 0.0f, 0.0f, 0.0f, 4.0f, 5.0f, 0.0f, 0.0f, 8.0f,  // ty
+        0.0f, 1.0f, 0.0f, 0.0f, 0.0f, 4.0f, 5.0f, 0.0f, 0.0f, 8.0f,  // tz
+        // LOD-2
+        0.0f, 1.0f, 0.0f, 0.0f, 0.0f, 4.0f, 5.0f, 0.0f, 7.0f, 0.0f,  // rx
+        0.0f, 1.0f, 0.0f, 0.0f, 0.0f, 4.0f, 5.0f, 0.0f, 7.0f, 0.0f,  // ry
+        0.0f, 1.0f, 0.0f, 0.0f, 0.0f, 4.0f, 5.0f, 0.0f, 7.0f, 0.0f,  // rz
+        0.0f, 1.0f, 0.0f, 0.0f, 0.0f, 4.0f, 5.0f, 0.0f, 7.0f, 0.0f,  // sx
+        0.0f, 1.0f, 0.0f, 0.0f, 0.0f, 4.0f, 5.0f, 0.0f, 7.0f, 0.0f,  // sy
+        0.0f, 1.0f, 0.0f, 0.0f, 0.0f, 4.0f, 5.0f, 0.0f, 7.0f, 0.0f,  // sz
+        0.0f, 0.0f, 0.0f, 0.0f, 0.0f, 0.0f, 0.0f, 0.0f, 0.0f, 0.0f,
+        // LOD-1
+        0.0f, 0.0f, 0.0f, 0.0f, 0.0f, 0.0f, 0.0f, 0.0f, 0.0f, 0.0f,
+        0.0f, 1.0f, 2.0f, 0.0f, 0.0f, 0.0f, 5.0f, 6.0f, 7.0f, 0.0f,  // tx
+        0.0f, 1.0f, 2.0f, 0.0f, 0.0f, 0.0f, 5.0f, 6.0f, 7.0f, 0.0f,  // ty
+        0.0f, 1.0f, 2.0f, 0.0f, 0.0f, 0.0f, 0.0f, 6.0f, 7.0f, 0.0f,  // tz
+        0.0f, 1.0f, 2.0f, 0.0f, 0.0f, 0.0f, 5.0f, 6.0f, 7.0f, 0.0f,  // rx
+        0.0f, 1.0f, 2.0f, 0.0f, 0.0f, 0.0f, 5.0f, 6.0f, 7.0f, 0.0f,  // ry
+        0.0f, 1.0f, 2.0f, 0.0f, 0.0f, 0.0f, 0.0f, 6.0f, 7.0f, 0.0f,  // rz
+        // LOD-0
+        3.0f, 1.0f, 2.0f, 0.0f, 0.0f, 0.0f, 0.0f, 6.0f, 7.0f, 0.0f,  // sx
+        3.0f, 1.0f, 2.0f, 0.0f, 0.0f, 0.0f, 0.0f, 6.0f, 7.0f, 0.0f,  // sy
+        0.0f, 0.0f, 0.0f, 0.0f, 0.0f, 0.0f, 0.0f, 0.0f, 0.0f, 0.0f,
+        0.0f, 0.0f, 0.0f, 0.0f, 0.0f, 0.0f, 0.0f, 0.0f, 0.0f, 0.0f,
+    },
 };
 
 const Matrix<std::uint16_t> inputIndices = {
@@ -179,6 +207,9 @@ const Matrix<std::uint16_t> inputIndices = {
     {  // Joint group 11
         1, 2, 3, 4, 5
     },
+    {  // Joint group 12
+        1, 2, 3, 4, 10, 5, 6, 7, 8, 9
+    },
 };
 const Matrix<std::uint16_t> outputIndices = {
     {  // Joint group 0
@@ -216,6 +247,9 @@ const Matrix<std::uint16_t> outputIndices = {
     },
     {  // Joint group 11
         147, 149, 156, 148
+    },
+    {  // Joint group 12
+        153, 162, 163, 164, 165, 166, 167, 168, 169, 170, 154, 155, 171, 172, 173, 174, 175, 176, 177, 178, 156, 157
     }
 };
 const Matrix<std::uint16_t> lods = {
@@ -254,6 +288,9 @@ const Matrix<std::uint16_t> lods = {
     },
     {  // Joint group 11
         4, 4, 2, 2
+    },
+    {  // Joint group 11
+        22, 18, 11, 4
     }
 };
 
@@ -491,6 +528,28 @@ const AlignedMatrix<float> floatValues = {
         1.0f, 2.0f, 3.0f, 4.0f,
         1.0f, 2.0f, 3.0f, 4.0f,
         1.0f, 2.0f, 3.0f, 4.0f
+    },
+    {  // Joint group 12
+        8.0f, 8.0f, 8.0f, 0.0f, 0.0f, 0.0f, 0.0f, 0.0f,
+        1.0f, 1.0f, 1.0f, 1.0f, 1.0f, 1.0f, 1.0f, 1.0f,
+        5.0f, 5.0f, 5.0f, 5.0f, 5.0f, 5.0f, 5.0f, 5.0f,
+        4.0f, 4.0f, 4.0f, 4.0f, 4.0f, 4.0f, 4.0f, 4.0f,
+        0.0f, 0.0f, 0.0f, 7.0f, 7.0f, 7.0f, 7.0f, 7.0f,
+        0.0f, 0.0f, 0.0f, 0.0f, 0.0f, 0.0f, 0.0f, 0.0f,
+        0.0f, 0.0f, 0.0f, 0.0f, 0.0f, 0.0f, 0.0f, 0.0f,
+        0.0f, 0.0f, 0.0f, 0.0f, 0.0f, 0.0f, 0.0f, 0.0f,
+        0.0f, 0.0f, 0.0f, 0.0f, 0.0f, 0.0f, 0.0f, 0.0f,
+        1.0f, 1.0f, 1.0f, 1.0f, 1.0f, 1.0f, 1.0f, 1.0f,
+        5.0f, 5.0f, 5.0f, 0.0f, 5.0f, 5.0f, 0.0f, 0.0f,
+        4.0f, 0.0f, 0.0f, 0.0f, 0.0f, 0.0f, 0.0f, 0.0f,
+        7.0f, 7.0f, 7.0f, 7.0f, 7.0f, 7.0f, 7.0f, 7.0f,
+        0.0f, 2.0f, 2.0f, 2.0f, 2.0f, 2.0f, 2.0f, 2.0f,
+        0.0f, 6.0f, 6.0f, 6.0f, 6.0f, 6.0f, 6.0f, 6.0f,
+        0.0f, 0.0f, 0.0f, 0.0f, 0.0f, 0.0f, 0.0f, 3.0f,
+        0.0f, 0.0f, 0.0f, 0.0f, 1.0f, 0.0f, 0.0f, 0.0f,
+        0.0f, 0.0f, 0.0f, 0.0f, 0.0f, 0.0f, 0.0f, 0.0f,
+        7.0f, 0.0f, 0.0f, 0.0f, 2.0f, 0.0f, 0.0f, 0.0f,
+        6.0f, 0.0f, 0.0f, 0.0f, 3.0f, 0.0f, 0.0f, 0.0f,
     }
 };
 
@@ -716,6 +775,28 @@ const AlignedMatrix<std::uint16_t> halfFloatValues = {
         15360, 16384, 16896, 17408,
         15360, 16384, 16896, 17408,
     },
+    {  // Joint group 12
+        18432, 18432, 18432, 0, 0, 0, 0, 0,
+        15360, 15360, 15360, 15360, 15360, 15360, 15360, 15360,
+        17664, 17664, 17664, 17664, 17664, 17664, 17664, 17664,
+        17408, 17408, 17408, 17408, 17408, 17408, 17408, 17408,
+        0, 0, 0, 18176, 18176, 18176, 18176, 18176,
+        0, 0, 0, 0, 0, 0, 0, 0,
+        0, 0, 0, 0, 0, 0, 0, 0,
+        0, 0, 0, 0, 0, 0, 0, 0,
+        0, 0, 0, 0, 0, 0, 0, 0,
+        15360, 15360, 15360, 15360, 15360, 15360, 15360, 15360,
+        17664, 17664, 17664, 0, 17664, 17664, 0, 0,
+        17408, 0, 0, 0, 0, 0, 0, 0,
+        18176, 18176, 18176, 18176, 18176, 18176, 18176, 18176,
+        0, 16384, 16384, 16384, 16384, 16384, 16384, 16384,
+        0, 17920, 17920, 17920, 17920, 17920, 17920, 17920,
+        0, 0, 0, 0, 0, 0, 0, 16896,
+        0, 0, 0, 0, 15360, 0, 0, 0,
+        0, 0, 0, 0, 0, 0, 0, 0,
+        18176, 0, 0, 0, 16384, 0, 0, 0,
+        17920, 0, 0, 0, 16896, 0, 0, 0,
+    }
 };
 
 const AlignedMatrix<std::uint16_t> inputIndices = {
@@ -754,6 +835,9 @@ const AlignedMatrix<std::uint16_t> inputIndices = {
     },
     {  // Joint group 11
         1, 2, 3, 4, 5
+    },
+    {  // Joint group 12
+        9, 2, 6, 5, 8, 3, 7, 1
     }
 };
 
@@ -794,6 +878,9 @@ const Vector<Vector<AlignedVector<std::uint16_t> > > outputIndices = {
         },
         {  // Joint group 11
             163, 165, 173, 164
+        },
+        {  // Joint group 12
+            180, 181, 182, 183, 184, 185, 187, 188, 189, 190, 191, 192, 193, 194, 195, 197, 198, 0, 0, 0
         }
     },
     {  // Euler-angle outputs
@@ -832,6 +919,9 @@ const Vector<Vector<AlignedVector<std::uint16_t> > > outputIndices = {
         },
         {  // Joint group 11
             147, 149, 156, 148
+        },
+        {  // Joint group 12
+            162, 163, 164, 165, 166, 167, 168, 169, 170, 171, 172, 173, 174, 175, 176, 177, 178, 0, 0, 0
         }
     }
 };
@@ -870,6 +960,9 @@ const Vector<Vector<AlignedVector<std::uint16_t> > > outputRotationIndices = {
         },
         {  // Joint group 11
             173, 163
+        },
+        {  // Joint group 12
+            183, 193
         }
     },
     {  // Euler-angle outputs
@@ -896,6 +989,8 @@ const Vector<Vector<AlignedVector<std::uint16_t> > > outputRotationIndices = {
         {  // Joint group 10
         },
         {  // Joint group 11
+        },
+        {  // Joint group 11
         }
     }
 };
@@ -937,6 +1032,9 @@ const Vector<Vector<AlignedVector<std::uint16_t> > > outputRotationLODs = {
         },
         {  // Joint group 11
             2, 2, 2, 2
+        },
+        {  // Joint group 12
+            2, 2, 1, 0
         }
     },
     {  // Euler-angle outputs
@@ -963,164 +1061,178 @@ const Vector<Vector<AlignedVector<std::uint16_t> > > outputRotationLODs = {
         {  // Joint group 10
         },
         {  // Joint group 11
+        },
+        {  // Joint group 12
         }
     }
 };
 
 const Vector<Vector<bpcm::JointGroup> > jointGroups = {
     // {valueOffset, inputOffset, outputOffset, lodOffset, outputRotationIndicesOffset, outputRotationLODsOffset, valueSize,
-    // inputSize, inputSizeAligned4, inputSizeAligned8}
+    // colCount, rowCount}
     {  // Quaternion outputs
         {  // Joint group 0
-            0, 0, 0, 0, 0, 0, 52, 13, 12, 8
+            0, 0, 0, 0, 0, 0, 52, 13, 4
         },
         {  // Joint group 1
-            52, 13, 4, 4, 0, 4, 52, 13, 12, 8
+            52, 13, 4, 4, 0, 4, 52, 13, 4
         },
         {  // Joint group 2
-            104, 26, 8, 8, 1, 8, 104, 13, 12, 8
+            104, 26, 8, 8, 1, 8, 104, 13, 8
         },
         {  // Joint group 3
-            208, 39, 16, 12, 2, 12, 104, 13, 12, 8
+            208, 39, 16, 12, 2, 12, 104, 13, 8
         },
         {  // Joint group 4
-            312, 52, 24, 16, 3, 16, 156, 13, 12, 8
+            312, 52, 24, 16, 3, 16, 156, 13, 12
         },
         {  // Joint group 5
-            468, 65, 36, 20, 4, 20, 156, 13, 12, 8
+            468, 65, 36, 20, 4, 20, 156, 13, 12
         },
         {  // Joint group 6
-            624, 78, 48, 24, 5, 24, 208, 13, 12, 8
+            624, 78, 48, 24, 5, 24, 208, 13, 16
         },
         {  // Joint group 7
-            832, 91, 64, 28, 7, 28, 208, 13, 12, 8
+            832, 91, 64, 28, 7, 28, 208, 13, 16
         },
         {  // Joint group 8
-            1040, 104, 80, 32, 9, 32, 208, 13, 12, 8
+            1040, 104, 80, 32, 9, 32, 208, 13, 16
         },
         {  // Joint group 9
-            1248, 117, 96, 36, 11, 36, 20, 5, 4, 0
+            1248, 117, 96, 36, 11, 36, 20, 5, 4
         },
         {  // Joint group 10
-            1268, 122, 100, 40, 11, 40, 20, 5, 4, 0
+            1268, 122, 100, 40, 11, 40, 20, 5, 4
         },
         {  // Joint group 11
-            1288, 127, 104, 44, 11, 44, 20, 5, 4, 0
+            1288, 127, 104, 44, 11, 44, 20, 5, 4
+        },
+        {  // Joint group 12
+            1308, 132, 108, 48, 13, 48, 160, 8, 20
         }
     },
     {  // Euler-angle outputs
         {  // Joint group 0
-            0, 0, 0, 0, 0, 0, 52, 13, 12, 8
+            0, 0, 0, 0, 0, 0, 52, 13, 4
         },
         {  // Joint group 1
-            52, 13, 4, 4, 0, 0, 52, 13, 12, 8
+            52, 13, 4, 4, 0, 0, 52, 13, 4
         },
         {  // Joint group 2
-            104, 26, 8, 8, 0, 0, 104, 13, 12, 8
+            104, 26, 8, 8, 0, 0, 104, 13, 8
         },
         {  // Joint group 3
-            208, 39, 16, 12, 0, 0, 104, 13, 12, 8
+            208, 39, 16, 12, 0, 0, 104, 13, 8
         },
         {  // Joint group 4
-            312, 52, 24, 16, 0, 0, 156, 13, 12, 8
+            312, 52, 24, 16, 0, 0, 156, 13, 12
         },
         {  // Joint group 5
-            468, 65, 36, 20, 0, 0, 156, 13, 12, 8
+            468, 65, 36, 20, 0, 0, 156, 13, 12
         },
         {  // Joint group 6
-            624, 78, 48, 24, 0, 0, 208, 13, 12, 8
+            624, 78, 48, 24, 0, 0, 208, 13, 16
         },
         {  // Joint group 7
-            832, 91, 64, 28, 0, 0, 208, 13, 12, 8
+            832, 91, 64, 28, 0, 0, 208, 13, 16
         },
         {  // Joint group 8
-            1040, 104, 80, 32, 0, 0, 208, 13, 12, 8
+            1040, 104, 80, 32, 0, 0, 208, 13, 16
         },
         {  // Joint group 9
-            1248, 117, 96, 36, 0, 0, 20, 5, 4, 0
+            1248, 117, 96, 36, 0, 0, 20, 5, 4
         },
         {  // Joint group 10
-            1268, 122, 100, 40, 0, 0, 20, 5, 4, 0
+            1268, 122, 100, 40, 0, 0, 20, 5, 4
         },
         {  // Joint group 11
-            1288, 127, 104, 44, 0, 0, 20, 5, 4, 0
+            1288, 127, 104, 44, 0, 0, 20, 5, 4
+        },
+        {  // Joint group 12
+            1308, 132, 108, 48, 0, 0, 160, 8, 20
         }
     }
 };
 
-const Matrix<bpcm::LODRegion> lodRegions = {
+const Matrix<LODRegion> lodRegions = {
     // {unaligned, aligned to last block-8, aligned to second last block-8}
     {  // Joint group 0
-        {1, 0, 0},  // LOD-0
-        {1, 0, 0},  // LOD-1
-        {1, 0, 0},  // LOD-2
-        {1, 0, 0},  // LOD-3
+        {{13, 12, 8}, {1, 0, 0}},  // LOD-0
+        {{13, 12, 8}, {1, 0, 0}},  // LOD-1
+        {{13, 12, 8}, {1, 0, 0}},  // LOD-2
+        {{13, 12, 8}, {1, 0, 0}},  // LOD-3
     },
     {  // Joint group 1
-        {4, 0, 0},  // LOD-0
-        {3, 0, 0},  // LOD-1
-        {2, 0, 0},  // LOD-2
-        {1, 0, 0},  // LOD-3
+        {{13, 12, 8}, {4, 0, 0}},  // LOD-0
+        {{13, 12, 8}, {3, 0, 0}},  // LOD-1
+        {{13, 12, 8}, {2, 0, 0}},  // LOD-2
+        {{13, 12, 8}, {1, 0, 0}},  // LOD-3
     },
     {  // Joint group 2
-        {7, 8, 0},  // LOD-0
-        {5, 8, 0},  // LOD-1
-        {4, 8, 0},  // LOD-2
-        {1, 8, 0},  // LOD-3
+        {{13, 12, 8}, {7, 8, 0}},  // LOD-0
+        {{13, 12, 8}, {5, 8, 0}},  // LOD-1
+        {{13, 12, 8}, {4, 8, 0}},  // LOD-2
+        {{13, 12, 8}, {1, 8, 0}},  // LOD-3
     },
     {  // Joint group 3
-        {8, 8, 8},  // LOD-0
-        {6, 8, 0},  // LOD-1
-        {4, 8, 0},  // LOD-2
-        {3, 8, 0},  // LOD-3
+        {{13, 12, 8}, {8, 8, 8}},  // LOD-0
+        {{13, 12, 8}, {6, 8, 0}},  // LOD-1
+        {{13, 12, 8}, {4, 8, 0}},  // LOD-2
+        {{13, 12, 8}, {3, 8, 0}},  // LOD-3
     },
     {  // Joint group 4
-        {9, 8, 8},  // LOD-0
-        {7, 8, 0},  // LOD-1
-        {4, 8, 0},  // LOD-2
-        {2, 8, 0},  // LOD-3
+        {{13, 12, 8}, {9, 8, 8}},  // LOD-0
+        {{13, 12, 8}, {7, 8, 0}},  // LOD-1
+        {{13, 12, 8}, {4, 8, 0}},  // LOD-2
+        {{13, 12, 8}, {2, 8, 0}},  // LOD-3
     },
     {  // Joint group 5
-        {12, 8, 8},  // LOD-0
-        {9, 8, 8},  // LOD-1
-        {7, 8, 0},  // LOD-2
-        {3, 8, 0},  // LOD-3
+        {{13, 12, 8}, {12, 8, 8}},  // LOD-0
+        {{13, 12, 8}, {9, 8, 8}},  // LOD-1
+        {{13, 12, 8}, {7, 8, 0}},  // LOD-2
+        {{13, 12, 8}, {3, 8, 0}},  // LOD-3
     },
     {  // Joint group 6
-        {14, 16, 8},  // LOD-0
-        {13, 16, 8},  // LOD-1
-        {11, 16, 8},  // LOD-2
-        {9, 16, 8},  // LOD-3
+        {{13, 12, 8}, {14, 16, 8}},  // LOD-0
+        {{13, 12, 8}, {13, 16, 8}},  // LOD-1
+        {{13, 12, 8}, {11, 16, 8}},  // LOD-2
+        {{13, 12, 8}, {9, 16, 8}},  // LOD-3
     },
     {  // Joint group 7
-        {16, 16, 16},  // LOD-0
-        {12, 16, 8},  // LOD-1
-        {8, 8, 8},  // LOD-2
-        {4, 8, 0},  // LOD-3
+        {{13, 12, 8}, {16, 16, 16}},  // LOD-0
+        {{13, 12, 8}, {12, 16, 8}},  // LOD-1
+        {{13, 12, 8}, {8, 8, 8}},  // LOD-2
+        {{13, 12, 8}, {4, 8, 0}},  // LOD-3
     },
     {  // Joint group 8
-        {16, 16, 16},  // LOD-0
-        {12, 16, 8},  // LOD-1
-        {8, 8, 8},  // LOD-2
-        {4, 8, 0},  // LOD-3
+        {{13, 12, 8}, {16, 16, 16}},  // LOD-0
+        {{13, 12, 8}, {12, 16, 8}},  // LOD-1
+        {{13, 12, 8}, {8, 8, 8}},  // LOD-2
+        {{13, 12, 8}, {4, 8, 0}},  // LOD-3
     },
     {  // Joint group 9
-        {2, 0, 0},  // LOD-0
-        {2, 0, 0},  // LOD-1
-        {2, 0, 0},  // LOD-2
-        {2, 0, 0},  // LOD-3
+        {{5, 4, 0}, {2, 0, 0}},  // LOD-0
+        {{5, 4, 0}, {2, 0, 0}},  // LOD-1
+        {{5, 4, 0}, {2, 0, 0}},  // LOD-2
+        {{5, 4, 0}, {2, 0, 0}},  // LOD-3
     },
     {  // Joint group 10
-        {2, 0, 0},  // LOD-0
-        {2, 0, 0},  // LOD-1
-        {2, 0, 0},  // LOD-2
-        {2, 0, 0}  // LOD-3
+        {{5, 4, 0}, {2, 0, 0}},  // LOD-0
+        {{5, 4, 0}, {2, 0, 0}},  // LOD-1
+        {{5, 4, 0}, {2, 0, 0}},  // LOD-2
+        {{5, 4, 0}, {2, 0, 0}},  // LOD-3
     },
     {  // Joint group 11
-        {4, 0, 0},  // LOD-0
-        {4, 0, 0},  // LOD-1
-        {2, 0, 0},  // LOD-2
-        {2, 0, 0}  // LOD-3
+        {{5, 4, 0}, {4, 0, 0}},  // LOD-0
+        {{5, 4, 0}, {4, 0, 0}},  // LOD-1
+        {{5, 4, 0}, {2, 0, 0}},  // LOD-2
+        {{5, 4, 0}, {2, 0, 0}},  // LOD-3
+    },
+    {  // Joint group 12
+        {{8, 8, 8}, {17, 16, 16}},  // LOD-0
+        {{7, 4, 0}, {15, 16, 8}},  // LOD-1
+        {{5, 4, 0}, {9, 16, 8}},  // LOD-2
+        {{4, 4, 0}, {3, 8, 0}},  // LOD-3
     }
 };
 
@@ -1157,7 +1269,9 @@ const Vector<Matrix<float> > valuesPerLOD = {
             15.0f, 30.0f, 0.0f, 0.0f, 0.0f, 0.0f, 1.0f, 0.0f, 0.0f, 0.0f,
             40.0f, 80.0f, 0.0f, 0.0f, 0.0f, 0.0f, 1.0f, 0.0f, 0.0f, 0.0f,
             0.0f, 0.0f, 0.0f, 0.237489700f, 0.449574471f, 0.285265595f, 0.812468469f, 0.0f, 0.0f, 0.0f,
-            0.0f, 0.0f, 0.0f, 0.382683456f, 0.0f, 0.0f, 0.923879504f, 0.0f, 0.0f, 0.0f
+            0.0f, 0.0f, 0.0f, 0.382683456f, 0.0f, 0.0f, 0.923879504f, 0.0f, 0.0f, 0.0f,
+            62.0f, 62.0f, 62.0f, 0.552419543f, -0.174177229f, 0.552419543f, -0.599439502f, 125.0f, 125.0f, 125.0f,
+            157.0f, 157.0f, 122.0f, 0.265585661f, -0.0761554763f, 0.500303566f, -0.820585668f, 128.0f, 128.0f, 0.0f
         }, {
             // LOD-1
             76.0f, 0.0f, 0.0f, 0.0f, 0.0f, 0.0f, 1.0f, 0.0f, 0.0f, 0.0f,
@@ -1177,8 +1291,9 @@ const Vector<Matrix<float> > valuesPerLOD = {
             15.0f, 30.0f, 0.0f, 0.0f, 0.0f, 0.0f, 1.0f, 0.0f, 0.0f, 0.0f,
             40.0f, 80.0f, 0.0f, 0.0f, 0.0f, 0.0f, 1.0f, 0.0f, 0.0f, 0.0f,
             0.0f, 0.0f, 0.0f, 0.237489700f, 0.449574471f, 0.285265595f, 0.812468469f, 0.0f, 0.0f, 0.0f,
-            0.0f, 0.0f, 0.0f, 0.382683456f, 0.0f, 0.0f, 0.923879504f, 0.0f, 0.0f, 0.0f
-
+            0.0f, 0.0f, 0.0f, 0.382683456f, 0.0f, 0.0f, 0.923879504f, 0.0f, 0.0f, 0.0f,
+            62.0f, 62.0f, 62.0f, 0.552419543f, -0.174177229f, 0.552419543f, -0.599439502f, 125.0f, 125.0f, 125.0f,
+            157.0f, 157.0f, 122.0f, 0.265585661f, -0.0761554763f, 0.500303566f, -0.820585668f, 0.0f, 0.0f, 0.0f
         }, {
             // LOD-2
             76.0f, 0.0f, 0.0f, 0.0f, 0.0f, 0.0f, 1.0f, 0.0f, 0.0f, 0.0f,
@@ -1198,6 +1313,8 @@ const Vector<Matrix<float> > valuesPerLOD = {
             15.0f, 30.0f, 0.0f, 0.0f, 0.0f, 0.0f, 1.0f, 0.0f, 0.0f, 0.0f,
             40.0f, 80.0f, 0.0f, 0.0f, 0.0f, 0.0f, 1.0f, 0.0f, 0.0f, 0.0f,
             0.0f, 0.0f, 0.0f, 0.126078621f, -0.0337826647f, 0.256604820f, 0.957662225f, 0.0f, 0.0f, 0.0f,
+            0.0f, 0.0f, 0.0f, 0.0f, 0.0f, 0.0f, 1.0f, 0.0f, 0.0f, 0.0f,
+            62.0f, 62.0f, 62.0f, 0.552419543f, -0.174177229f, 0.552419543f, -0.599439502f, 125.0f, 125.0f, 125.0f,
             0.0f, 0.0f, 0.0f, 0.0f, 0.0f, 0.0f, 1.0f, 0.0f, 0.0f, 0.0f
         }, {
             // LOD-3
@@ -1218,6 +1335,8 @@ const Vector<Matrix<float> > valuesPerLOD = {
             15.0f, 30.0f, 0.0f, 0.0f, 0.0f, 0.0f, 1.0f, 0.0f, 0.0f, 0.0f,
             40.0f, 80.0f, 0.0f, 0.0f, 0.0f, 0.0f, 1.0f, 0.0f, 0.0f, 0.0f,
             0.0f, 0.0f, 0.0f, 0.126078621f, -0.0337826647f, 0.256604820f, 0.957662225f, 0.0f, 0.0f, 0.0f,
+            0.0f, 0.0f, 0.0f, 0.0f, 0.0f, 0.0f, 1.0f, 0.0f, 0.0f, 0.0f,
+            62.0f, 62.0f, 62.0f, 0.0f, 0.0f, 0.0f, 1.0f, 0.0f, 0.0f, 0.0f,
             0.0f, 0.0f, 0.0f, 0.0f, 0.0f, 0.0f, 1.0f, 0.0f, 0.0f, 0.0f
         }
     },
@@ -1250,7 +1369,9 @@ const Vector<Matrix<float> > valuesPerLOD = {
             40.0f, 80.0f,  // Joint group 10
             0.0f, 0.0f, 0.0f, 0.0f, 0.0f, 0.0f, 0.0f,
             0.0f, 0.0f, 0.0f, 15.0f, 60.0f, 30.0f, 0.0f, 0.0f, 0.0f,  // Joint group 11
-            0.0f, 0.0f, 0.0f, 45.0f, 0.0f, 0.0f, 0.0f, 0.0f, 0.0f  // Joint group 11
+            0.0f, 0.0f, 0.0f, 45.0f, 0.0f, 0.0f, 0.0f, 0.0f, 0.0f,  // Joint group 11
+            62.0f, 62.0f, 62.0f, 125.0f, 125.0f, 125.0f, 125.0f, 125.0f, 125.0f,  // Joint group 12
+            157.0f, 157.0f, 122.0f, 157.0f, 157.0f, 122.0f, 128.0f, 128.0f, 0.0f  // Joint group 12
         }, {
             // LOD-1
             76.0f,  // Joint group 0
@@ -1279,7 +1400,9 @@ const Vector<Matrix<float> > valuesPerLOD = {
             40.0f, 80.0f,  // Joint group 10
             0.0f, 0.0f, 0.0f, 0.0f, 0.0f, 0.0f, 0.0f,
             0.0f, 0.0f, 0.0f, 15.0f, 60.0f, 30.0f, 0.0f, 0.0f, 0.0f,  // Joint group 11
-            0.0f, 0.0f, 0.0f, 45.0f, 0.0f, 0.0f, 0.0f, 0.0f, 0.0f  // Joint group 11
+            0.0f, 0.0f, 0.0f, 45.0f, 0.0f, 0.0f, 0.0f, 0.0f, 0.0f,  // Joint group 11
+            62.0f, 62.0f, 62.0f, 125.0f, 125.0f, 125.0f, 125.0f, 125.0f, 125.0f,  // Joint group 12
+            157.0f, 157.0f, 122.0f, 157.0f, 157.0f, 122.0f, 0.0f, 0.0f, 0.0f  // Joint group 12
         }, {
             // LOD-2
             76.0f,  // Joint group 0
@@ -1308,7 +1431,9 @@ const Vector<Matrix<float> > valuesPerLOD = {
             40.0f, 80.0f,  // Joint group 10
             0.0f, 0.0f, 0.0f, 0.0f, 0.0f, 0.0f, 0.0f,
             0.0f, 0.0f, 0.0f, 15.0f, 0.0f, 30.0f, 0.0f, 0.0f, 0.0f,  // Joint group 11
-            0.0f, 0.0f, 0.0f, 0.0f, 0.0f, 0.0f, 0.0f, 0.0f, 0.0f  // Joint group 11
+            0.0f, 0.0f, 0.0f, 0.0f, 0.0f, 0.0f, 0.0f, 0.0f, 0.0f,  // Joint group 11
+            62.0f, 62.0f, 62.0f, 125.0f, 125.0f, 125.0f, 125.0f, 125.0f, 125.0f,  // Joint group 12
+            0.0f, 0.0f, 0.0f, 0.0f, 0.0f, 0.0f, 0.0f, 0.0f, 0.0f  // Joint group 12
         }, {
             // LOD-3
             76.0f,  // Joint group 0
@@ -1337,7 +1462,9 @@ const Vector<Matrix<float> > valuesPerLOD = {
             40.0f, 80.0f,  // Joint group 10
             0.0f, 0.0f, 0.0f, 0.0f, 0.0f, 0.0f, 0.0f,
             0.0f, 0.0f, 0.0f, 15.0f, 0.0f, 30.0f, 0.0f, 0.0f, 0.0f,  // Joint group 11
-            0.0f, 0.0f, 0.0f, 0.0f, 0.0f, 0.0f, 0.0f, 0.0f, 0.0f  // Joint group 11
+            0.0f, 0.0f, 0.0f, 0.0f, 0.0f, 0.0f, 0.0f, 0.0f, 0.0f,  // Joint group 11
+            62.0f, 62.0f, 62.0f, 0.0f, 0.0f, 0.0f, 0.0f, 0.0f, 0.0f,  // Joint group 12
+            0.0f, 0.0f, 0.0f, 0.0f, 0.0f, 0.0f, 0.0f, 0.0f, 0.0f  // Joint group 12
         }
     }
 };
@@ -1440,9 +1567,8 @@ bpcm::Evaluator<TValue> OptimizedStorage<TValue>::create(StrategyPtr strategy,
             0u,  // outputRotationIndicesOffset
             0u,  // outputRotationLODsOffset
             static_cast<std::uint32_t>(storage.values.size()),
-            static_cast<std::uint32_t>(storage.inputIndices.size()),
-            static_cast<std::uint32_t>(storage.inputIndices.size() - (storage.inputIndices.size() % 4ul)),
-            static_cast<std::uint32_t>(storage.inputIndices.size() - (storage.inputIndices.size() % 8ul)),
+            optimized::jointGroups[rotationSelectorIndex][jointGroupIndex].colCount,
+            optimized::jointGroups[rotationSelectorIndex][jointGroupIndex].rowCount
         });
 
     const auto lastJointGroupIndex = optimized::outputIndices[rotationSelectorIndex].size() - 1ul;
diff --git a/Engine/Plugins/Animation/RigLogic/Source/RigLogicLibTest/Private/rltests/joints/bpcm/BPCMFixturesBlock4.h b/Engine/Plugins/Animation/RigLogic/Source/RigLogicLibTest/Private/rltests/joints/bpcm/BPCMFixturesBlock4.h
--- a/Engine/Plugins/Animation/RigLogic/Source/RigLogicLibTest/Private/rltests/joints/bpcm/BPCMFixturesBlock4.h
+++ b/Engine/Plugins/Animation/RigLogic/Source/RigLogicLibTest/Private/rltests/joints/bpcm/BPCMFixturesBlock4.h
@@ -54,7 +54,7 @@ extern const Vector<Vector<AlignedVector<std::uint16_t> > > outputIndices;
 extern const Vector<Vector<AlignedVector<std::uint16_t> > > outputRotationIndices;
 extern const Vector<Vector<AlignedVector<std::uint16_t> > > outputRotationLODs;
 extern const Vector<Vector<bpcm::JointGroup> > jointGroups;
-extern const Matrix<bpcm::LODRegion> lodRegions;
+extern const Matrix<LODRegion> lodRegions;
 
 }  // namespace optimized
 
diff --git a/Engine/Plugins/Animation/RigLogic/Source/RigLogicLibTest/Private/rltests/joints/bpcm/BPCMFixturesBlock8.cpp b/Engine/Plugins/Animation/RigLogic/Source/RigLogicLibTest/Private/rltests/joints/bpcm/BPCMFixturesBlock8.cpp
--- a/Engine/Plugins/Animation/RigLogic/Source/RigLogicLibTest/Private/rltests/joints/bpcm/BPCMFixturesBlock8.cpp
+++ b/Engine/Plugins/Animation/RigLogic/Source/RigLogicLibTest/Private/rltests/joints/bpcm/BPCMFixturesBlock8.cpp
@@ -828,110 +828,110 @@ const Vector<Vector<AlignedVector<std::uint16_t> > > outputRotationLODs = {
 
 const Vector<Vector<bpcm::JointGroup> > jointGroups = {
     // {valueOffset, inputOffset, outputOffset, lodOffset, outputRotationIndicesOffset, outputRotationLODsOffset, valueSize,
-    // inputSize, inputSizeAligned4, inputSizeAligned8}
+    // colCount, rowCount}
     {  // Quaternion outputs
         {  // Joint group 0
-            0, 0, 0, 0, 0, 0, 104, 13, 12, 8
+            0, 0, 0, 0, 0, 0, 104, 13, 8
         },
         {  // Joint group 1
-            104, 13, 8, 4, 1, 4, 104, 13, 12, 8
+            104, 13, 8, 4, 1, 4, 104, 13, 8
         },
         {  // Joint group 2
-            208, 26, 16, 8, 2, 8, 208, 13, 12, 8
+            208, 26, 16, 8, 2, 8, 208, 13, 16
         },
         {  // Joint group 3
-            416, 39, 32, 12, 3, 12, 208, 13, 12, 8
+            416, 39, 32, 12, 3, 12, 208, 13, 16
         },
         {  // Joint group 4
-            624, 52, 48, 16, 5, 16, 312, 13, 12, 8
+            624, 52, 48, 16, 5, 16, 312, 13, 24
         },
         {  // Joint group 5
-            936, 65, 72, 20, 7, 20, 312, 13, 12, 8
+            936, 65, 72, 20, 7, 20, 312, 13, 24
         },
         {  // Joint group 6
-            1248, 78, 96, 24, 10, 24, 416, 13, 12, 8
+            1248, 78, 96, 24, 10, 24, 416, 13, 32
         },
         {  // Joint group 7
-            1664, 91, 128, 28, 13, 28, 416, 13, 12, 8
+            1664, 91, 128, 28, 13, 28, 416, 13, 32
         }
     },
     {  // Euler-angle outputs
         {  // Joint group 0
-            0, 0, 0, 0, 0, 0, 104, 13, 12, 8
+            0, 0, 0, 0, 0, 0, 104, 13, 8
         },
         {  // Joint group 1
-            104, 13, 8, 4, 0, 0, 104, 13, 12, 8
+            104, 13, 8, 4, 0, 0, 104, 13, 8
         },
         {  // Joint group 2
-            208, 26, 16, 8, 0, 0, 208, 13, 12, 8
+            208, 26, 16, 8, 0, 0, 208, 13, 16
         },
         {  // Joint group 3
-            416, 39, 32, 12, 0, 0, 208, 13, 12, 8
+            416, 39, 32, 12, 0, 0, 208, 13, 16
         },
         {  // Joint group 4
-            624, 52, 48, 16, 0, 0, 312, 13, 12, 8
+            624, 52, 48, 16, 0, 0, 312, 13, 24
         },
         {  // Joint group 5
-            936, 65, 72, 20, 0, 0, 312, 13, 12, 8
+            936, 65, 72, 20, 0, 0, 312, 13, 24
         },
         {  // Joint group 6
-            1248, 78, 96, 24, 0, 0, 416, 13, 12, 8
+            1248, 78, 96, 24, 0, 0, 416, 13, 32
         },
         {  // Joint group 7
-            1664, 91, 128, 28, 0, 0, 416, 13, 12, 8
+            1664, 91, 128, 28, 0, 0, 416, 13, 32
         }
     }
 };
 
-const Matrix<bpcm::LODRegion> lodRegions = {
+const Matrix<LODRegion> lodRegions = {
     // {size, sizeAlignedToLastFullBlock, sizeAlignedToSecondLastFullBlock}
     {  // Joint group 0
-        {5, 0, 0},  // LOD-0
-        {4, 0, 0},  // LOD-1
-        {2, 0, 0},  // LOD-2
-        {1, 0, 0},  // LOD-3
+        {{13, 12, 8}, {5, 0, 0}},  // LOD-0
+        {{13, 12, 8}, {4, 0, 0}},  // LOD-1
+        {{13, 12, 8}, {2, 0, 0}},  // LOD-2
+        {{13, 12, 8}, {1, 0, 0}},  // LOD-3
     },
     {  // Joint group 1
-        {8, 0, 0},  // LOD-0
-        {5, 0, 0},  // LOD-1
-        {3, 0, 0},  // LOD-2
-        {1, 0, 0},  // LOD-3
+        {{13, 12, 8}, {8, 0, 0}},  // LOD-0
+        {{13, 12, 8}, {5, 0, 0}},  // LOD-1
+        {{13, 12, 8}, {3, 0, 0}},  // LOD-2
+        {{13, 12, 8}, {1, 0, 0}},  // LOD-3
     },
     {  // Joint group 2
-        {9, 16, 0},  // LOD-0
-        {7, 16, 0},  // LOD-1
-        {4, 16, 0},  // LOD-2
-        {2, 16, 0},  // LOD-3
+        {{13, 12, 8}, {9, 16, 0}},  // LOD-0
+        {{13, 12, 8}, {7, 16, 0}},  // LOD-1
+        {{13, 12, 8}, {4, 16, 0}},  // LOD-2
+        {{13, 12, 8}, {2, 16, 0}},  // LOD-3
     },
     {  // Joint group 3
-        {16, 16, 16},  // LOD-0
-        {13, 16, 0},  // LOD-1
-        {8, 16, 0},  // LOD-2
-        {4, 16, 0},  // LOD-3
+        {{13, 12, 8}, {16, 16, 16}},  // LOD-0
+        {{13, 12, 8}, {13, 16, 0}},  // LOD-1
+        {{13, 12, 8}, {8, 16, 0}},  // LOD-2
+        {{13, 12, 8}, {4, 16, 0}},  // LOD-3
     },
     {  // Joint group 4
-        {17, 16, 16},  // LOD-0
-        {16, 16, 16},  // LOD-1
-        {15, 16, 0},  // LOD-2
-        {12, 16, 0},  // LOD-3
+        {{13, 12, 8}, {17, 16, 16}},  // LOD-0
+        {{13, 12, 8}, {16, 16, 16}},  // LOD-1
+        {{13, 12, 8}, {15, 16, 0}},  // LOD-2
+        {{13, 12, 8}, {12, 16, 0}},  // LOD-3
     },
     {  // Joint group 5
-        {24, 16, 16},  // LOD-0
-        {16, 16, 16},  // LOD-1
-        {15, 16, 0},  // LOD-2
-        {8, 16, 0},  // LOD-3
+        {{13, 12, 8}, {24, 16, 16}},  // LOD-0
+        {{13, 12, 8}, {16, 16, 16}},  // LOD-1
+        {{13, 12, 8}, {15, 16, 0}},  // LOD-2
+        {{13, 12, 8}, {8, 16, 0}},  // LOD-3
     },
     {  // Joint group 6
-        {25, 32, 16},  // LOD-0
-        {17, 32, 16},  // LOD-1
-        {12, 16, 0},  // LOD-2
-        {7, 16, 0},  // LOD-3
+        {{13, 12, 8}, {25, 32, 16}},  // LOD-0
+        {{13, 12, 8}, {17, 32, 16}},  // LOD-1
+        {{13, 12, 8}, {12, 16, 0}},  // LOD-2
+        {{13, 12, 8}, {7, 16, 0}},  // LOD-3
     },
     {  // Joint group 7
-        {32, 32, 32},  // LOD-0
-        {25, 32, 16},  // LOD-1
-        {17, 32, 16},  // LOD-2
-        {9, 16, 0},  // LOD-3
+        {{13, 12, 8}, {32, 32, 32}},  // LOD-0
+        {{13, 12, 8}, {25, 32, 16}},  // LOD-1
+        {{13, 12, 8}, {17, 32, 16}},  // LOD-2
+        {{13, 12, 8}, {9, 16, 0}},  // LOD-3
     }
 };
 
@@ -1234,9 +1234,8 @@ bpcm::Evaluator<TValue> OptimizedStorage<TValue>::create(StrategyPtr strategy,
             0u,  // outputRotationIndicesOffset
             0u,  // outputRotationLODsOffset
             static_cast<std::uint32_t>(storage.values.size()),
-            static_cast<std::uint32_t>(storage.inputIndices.size()),
-            static_cast<std::uint32_t>(storage.inputIndices.size() - (storage.inputIndices.size() % 4ul)),
-            static_cast<std::uint32_t>(storage.inputIndices.size() - (storage.inputIndices.size() % 8ul)),
+            optimized::jointGroups[rotationSelectorIndex][jointGroupIndex].colCount,
+            optimized::jointGroups[rotationSelectorIndex][jointGroupIndex].rowCount
         });
 
     const auto lastJointGroupIndex = optimized::outputIndices[rotationSelectorIndex].size() - 1ul;
diff --git a/Engine/Plugins/Animation/RigLogic/Source/RigLogicLibTest/Private/rltests/joints/bpcm/BPCMFixturesBlock8.h b/Engine/Plugins/Animation/RigLogic/Source/RigLogicLibTest/Private/rltests/joints/bpcm/BPCMFixturesBlock8.h
--- a/Engine/Plugins/Animation/RigLogic/Source/RigLogicLibTest/Private/rltests/joints/bpcm/BPCMFixturesBlock8.h
+++ b/Engine/Plugins/Animation/RigLogic/Source/RigLogicLibTest/Private/rltests/joints/bpcm/BPCMFixturesBlock8.h
@@ -48,7 +48,7 @@ extern const Vector<Vector<AlignedVector<std::uint16_t> > > outputIndices;
 extern const Vector<Vector<AlignedVector<std::uint16_t> > > outputRotationIndices;
 extern const Vector<Vector<AlignedVector<std::uint16_t> > > outputRotationLODs;
 extern const Vector<Vector<bpcm::JointGroup> > jointGroups;
-extern const Matrix<bpcm::LODRegion> lodRegions;
+extern const Matrix<LODRegion> lodRegions;
 
 }  // namespace optimized
 
diff --git a/Engine/Plugins/Animation/RigLogic/Source/RigLogicLibTest/Private/rltests/joints/bpcm/TestBlock4CalculationStrategy.cpp b/Engine/Plugins/Animation/RigLogic/Source/RigLogicLibTest/Private/rltests/joints/bpcm/TestBlock4CalculationStrategy.cpp
--- a/Engine/Plugins/Animation/RigLogic/Source/RigLogicLibTest/Private/rltests/joints/bpcm/TestBlock4CalculationStrategy.cpp
+++ b/Engine/Plugins/Animation/RigLogic/Source/RigLogicLibTest/Private/rltests/joints/bpcm/TestBlock4CalculationStrategy.cpp
@@ -56,7 +56,7 @@ class Block4JointCalculationStrategyTest : public ::testing::TestWithParam<Strat
 
             auto inputInstanceFactory =
                 ControlsFactory::getInstanceFactory(0, static_cast<std::uint16_t>(block4::input::values.size()), 0, 0, 0);
-            rl4::Vector<float> initialValues(block4::input::values.size());
+            rl4::Vector<rl4::ControlInitializer> initialValues;
             auto inputInstance = inputInstanceFactory(initialValues, &memRes);
             auto inputBuffer = inputInstance->getInputBuffer();
             std::copy(block4::input::values.begin(), block4::input::values.end(), inputBuffer.begin());
diff --git a/Engine/Plugins/Animation/RigLogic/Source/RigLogicLibTest/Private/rltests/joints/bpcm/TestBlock8CalculationStrategy.cpp b/Engine/Plugins/Animation/RigLogic/Source/RigLogicLibTest/Private/rltests/joints/bpcm/TestBlock8CalculationStrategy.cpp
--- a/Engine/Plugins/Animation/RigLogic/Source/RigLogicLibTest/Private/rltests/joints/bpcm/TestBlock8CalculationStrategy.cpp
+++ b/Engine/Plugins/Animation/RigLogic/Source/RigLogicLibTest/Private/rltests/joints/bpcm/TestBlock8CalculationStrategy.cpp
@@ -55,7 +55,7 @@ class Block8JointCalculationStrategyTest : public ::testing::TestWithParam<Strat
 
             auto inputInstanceFactory =
                 ControlsFactory::getInstanceFactory(0, static_cast<std::uint16_t>(block8::input::values.size()), 0, 0, 0);
-            rl4::Vector<float> initialValues(block8::input::values.size());
+            rl4::Vector<rl4::ControlInitializer> initialValues;
             auto inputInstance = inputInstanceFactory(initialValues, &memRes);
             auto inputBuffer = inputInstance->getInputBuffer();
             std::copy(block8::input::values.begin(), block8::input::values.end(), inputBuffer.begin());
diff --git a/Engine/Plugins/Animation/RigLogic/Source/RigLogicLibTest/Private/rltests/joints/quaternions/QuaternionFixtures.cpp b/Engine/Plugins/Animation/RigLogic/Source/RigLogicLibTest/Private/rltests/joints/quaternions/QuaternionFixtures.cpp
--- a/Engine/Plugins/Animation/RigLogic/Source/RigLogicLibTest/Private/rltests/joints/quaternions/QuaternionFixtures.cpp
+++ b/Engine/Plugins/Animation/RigLogic/Source/RigLogicLibTest/Private/rltests/joints/quaternions/QuaternionFixtures.cpp
@@ -116,6 +116,43 @@
         10.0f, 40.0f,  // rx
         20.0f, 10.0f,  // ry
         40.0f, 20.0f  // rz
+    },
+    {  // Joint group 7 - test for joint group filter
+        // LOD-3
+        0.0f, 0.0f, 0.0f, 0.0f, 0.0f, 0.0f, 0.0f, 0.0f, 0.0f, 0.0f,  // tx
+        0.0f, 1.0f, 0.0f, 0.0f, 0.0f, 4.0f, 5.0f, 0.0f, 0.0f, 8.0f,  // rx
+        0.0f, 1.0f, 0.0f, 0.0f, 0.0f, 4.0f, 5.0f, 0.0f, 0.0f, 8.0f,  // ry
+        0.0f, 1.0f, 0.0f, 0.0f, 0.0f, 4.0f, 5.0f, 0.0f, 0.0f, 8.0f,  // rz
+        // LOD-2
+        0.0f, 1.0f, 0.0f, 0.0f, 0.0f, 4.0f, 5.0f, 0.0f, 7.0f, 0.0f,  // rx
+        0.0f, 1.0f, 0.0f, 0.0f, 0.0f, 4.0f, 5.0f, 0.0f, 7.0f, 0.0f,  // ry
+        0.0f, 1.0f, 0.0f, 0.0f, 0.0f, 4.0f, 5.0f, 0.0f, 7.0f, 0.0f,  // rz
+        0.0f, 1.0f, 0.0f, 0.0f, 0.0f, 4.0f, 5.0f, 0.0f, 7.0f, 0.0f,  // rx
+        0.0f, 1.0f, 0.0f, 0.0f, 0.0f, 4.0f, 5.0f, 0.0f, 7.0f, 0.0f,  // ry
+        0.0f, 1.0f, 0.0f, 0.0f, 0.0f, 4.0f, 5.0f, 0.0f, 7.0f, 0.0f,  // rz
+        0.0f, 0.0f, 0.0f, 0.0f, 0.0f, 0.0f, 0.0f, 0.0f, 0.0f, 0.0f,  // sy
+        // LOD-1
+        0.0f, 0.0f, 0.0f, 0.0f, 0.0f, 0.0f, 0.0f, 0.0f, 0.0f, 0.0f,  // tz
+        0.0f, 1.0f, 2.0f, 0.0f, 0.0f, 0.0f, 5.0f, 6.0f, 7.0f, 0.0f,  // rx
+        0.0f, 1.0f, 2.0f, 0.0f, 0.0f, 0.0f, 5.0f, 6.0f, 7.0f, 0.0f,  // ry
+        0.0f, 1.0f, 2.0f, 0.0f, 0.0f, 0.0f, 0.0f, 6.0f, 7.0f, 0.0f,  // rz
+        0.0f, 1.0f, 2.0f, 0.0f, 0.0f, 0.0f, 5.0f, 6.0f, 7.0f, 0.0f,  // rx
+        0.0f, 1.0f, 2.0f, 0.0f, 0.0f, 0.0f, 5.0f, 6.0f, 7.0f, 0.0f,  // ry
+        0.0f, 1.0f, 2.0f, 0.0f, 0.0f, 0.0f, 0.0f, 6.0f, 7.0f, 0.0f,  // rz
+        0.0f, 1.0f, 2.0f, 0.0f, 0.0f, 0.0f, 5.0f, 6.0f, 7.0f, 0.0f,  // rx
+        0.0f, 1.0f, 2.0f, 0.0f, 0.0f, 0.0f, 5.0f, 6.0f, 7.0f, 0.0f,  // ry
+        0.0f, 1.0f, 2.0f, 0.0f, 0.0f, 0.0f, 0.0f, 6.0f, 7.0f, 0.0f,  // rz
+        0.0f, 1.0f, 2.0f, 0.0f, 0.0f, 0.0f, 5.0f, 6.0f, 7.0f, 0.0f,  // rx
+        0.0f, 1.0f, 2.0f, 0.0f, 0.0f, 0.0f, 5.0f, 6.0f, 7.0f, 0.0f,  // ry
+        0.0f, 1.0f, 2.0f, 0.0f, 0.0f, 0.0f, 0.0f, 6.0f, 7.0f, 0.0f,  // rz
+        0.0f, 1.0f, 2.0f, 0.0f, 0.0f, 0.0f, 5.0f, 6.0f, 7.0f, 0.0f,  // rx
+        0.0f, 1.0f, 2.0f, 0.0f, 0.0f, 0.0f, 5.0f, 6.0f, 7.0f, 0.0f,  // ry
+        0.0f, 1.0f, 2.0f, 0.0f, 0.0f, 0.0f, 0.0f, 6.0f, 7.0f, 0.0f,  // rz
+        // LOD-0
+        3.0f, 1.0f, 2.0f, 0.0f, 0.0f, 0.0f, 0.0f, 6.0f, 7.0f, 0.0f,  // rx
+        3.0f, 1.0f, 2.0f, 0.0f, 0.0f, 0.0f, 0.0f, 6.0f, 7.0f, 0.0f,  // ry
+        0.0f, 0.0f, 0.0f, 0.0f, 0.0f, 0.0f, 0.0f, 0.0f, 0.0f, 0.0f,  // sy
+        0.0f, 0.0f, 0.0f, 0.0f, 0.0f, 0.0f, 0.0f, 0.0f, 0.0f, 0.0f,  // tx
     }
 };
 const Matrix<std::uint16_t> inputIndices = {
@@ -138,6 +175,9 @@
     },
     {  // Joint group 6
         8, 6
+    },
+    {  // Joint group 7
+        1, 2, 3, 4, 10, 5, 6, 7, 8, 9
     }
 };
 const Matrix<std::uint16_t> outputIndices = {
@@ -180,6 +220,12 @@
     {  // Joint group 6
         219, 220, 221,
         210, 211, 212
+    },
+    {  // Joint group 7
+        225, 228, 229, 230,  // LOD-3
+        237, 238, 239, 246, 247, 248, 250,  // LOD-2
+        254, 255, 256, 257, 264, 265, 266, 273, 274, 275, 282, 283, 284, 291, 292, 293,  // LOD-1
+        300, 301, 304, 306  // LOD-0
     }
 };
 const Matrix<std::uint16_t> lods = {
@@ -204,6 +250,9 @@
     {  // Joint group 6
         6, 3, 0, 0
     },
+    {  // Joint group 7
+        31, 27, 11, 4
+    }
 };
 const Vector<Extent> subMatrices = {
     {24, 2},  // Joint group 0
@@ -213,6 +262,7 @@
     {2, 2},  // Joint group 4
     {0, 0},  // Joint group 5
     {6, 2},  // Joint group 6
+    {31, 10},  // Joint group 7
 };
 const Vector<dna::RotationRepresentation> jointRotationRepresentations = {
     // Joint group 0
@@ -246,6 +296,17 @@
     // Joint group 5 - nothing
     // Joint group 6
     dna::RotationRepresentation::Quaternion,
+    dna::RotationRepresentation::Quaternion,
+    // Joint group 7
+    dna::RotationRepresentation::Quaternion,
+    dna::RotationRepresentation::Quaternion,
+    dna::RotationRepresentation::Quaternion,
+    dna::RotationRepresentation::Quaternion,
+    dna::RotationRepresentation::Quaternion,
+    dna::RotationRepresentation::Quaternion,
+    dna::RotationRepresentation::Quaternion,
+    dna::RotationRepresentation::Quaternion,
+    dna::RotationRepresentation::Quaternion,
     dna::RotationRepresentation::Quaternion
 };
 
@@ -370,6 +431,147 @@
         0.2308131f, 0.0214902f, 0.0f, 0.0f,  // Ys[q0-q1]
         0.1687222f, 0.1919111f, 0.0f, 0.0f,  // Zs[q0-q1]
         0.9372469f, 0.9167188f, 0.0f, 0.0f,  // Ws[q0-q1]
+    },
+    {  // Joint group 7
+        /*
+           x,          y,          z,          w,
+           0.0742712f, 0.0645629f, 0.0742712f, 0.9923705f,  // q0, e0
+           0.0f,       0.0f,       0.0f,       1.0f,        // q1, e0
+           0.0f,       0.0f,       0.0f,       1.0f,        // q2, e0
+           0.0f,       0.0f,       0.0f,       1.0f,        // q3, e0
+           0.0f,       0.0f,       0.0f,       1.0f,        // q4, e0
+           0.0f,       0.0f,       0.0f,       1.0f,        // q5, e0
+           0.0f,       0.0f,       0.0f,       1.0f,        // q6, e0
+           0.0f,       0.0f,       0.0f,       1.0f,        // q7, e0
+           0.008802f,  0.0086497f, 0.008802f,  0.9998851f,  // q0, e1
+           0.008802f,  0.0086497f, 0.008802f,  0.9998851f,  // q1, e1
+           0.008802f,  0.0086497f, 0.008802f,  0.9998851f,  // q2, e1
+           0.008802f,  0.0086497f, 0.008802f,  0.9998851f,  // q3, e1
+           0.008802f,  0.0086497f, 0.008802f,  0.9998851f,  // q4, e1
+           0.008802f,  0.0086497f, 0.008802f,  0.9998851f,  // q5, e1
+           0.008802f,  0.0086497f, 0.008802f,  0.9998851f,  // q6, e1
+           0.008802f,  0.0086497f, 0.008802f,  0.9998851f,  // q7, e1
+           0.0454372f, 0.0416356f, 0.0454372f, 0.9970644f,  // q0, e2
+           0.0454372f, 0.0416356f, 0.0454372f, 0.9970644f,  // q1, e2
+           0.0454372f, 0.0416356f, 0.0454372f, 0.9970644f,  // q2, e2
+           0.0435779f, 0.0435779f, 0.0019027f, 0.9980973f,  // q3, e2
+           0.0435779f, 0.0435779f, 0.0019027f, 0.9980973f,  // q4, e2
+           0.0435779f, 0.0435779f, 0.0019027f, 0.9980973f,  // q5, e2
+           0.0435779f, 0.0435779f, 0.0019027f, 0.9980973f,  // q6, e2
+           0.0435779f, 0.0435779f, 0.0019027f, 0.9980973f,  // q7, e2
+           0.0360742f, 0.0336398f, 0.0360742f, 0.9981311f,  // q0, e3
+           0.0360742f, 0.0336398f, 0.0360742f, 0.9981311f,  // q1, e3
+           0.0360742f, 0.0336398f, 0.0360742f, 0.9981311f,  // q2, e3
+           0.0f,       0.0f,       0.0f,       1.0f,        // q3, e3
+           0.0f,       0.0f,       0.0f,       1.0f,        // q4, e3
+           0.0f,       0.0f,       0.0f,       1.0f,        // q5, e3
+           0.0f,       0.0f,       0.0f,       1.0f,        // q6, e3
+           0.0f,       0.0f,       0.0f,       1.0f,        // q7, e3
+           0.0f,       0.0f,       0.0f,       1.0f,        // q0, e4
+           0.064541f,  0.057101f,  0.064541f,  0.9941873f,  // q1, e4
+           0.064541f,  0.057101f,  0.064541f,  0.9941873f,  // q2, e4
+           0.064541f,  0.057101f,  0.064541f,  0.9941873f,  // q3, e4
+           0.064541f,  0.057101f,  0.064541f,  0.9941873f,  // q4, e4
+           0.064541f,  0.057101f,  0.064541f,  0.9941873f,  // q5, e4
+           0.064541f,  0.057101f,  0.064541f,  0.9941873f,  // q6, e4
+           0.064541f,  0.057101f,  0.064541f,  0.9941873f,  // q7, e4
+           0.0f,       0.0f,       0.0f,       1.0f,        // q0, e5
+           0.0f,       0.0f,       0.0f,       1.0f,        // q1, e5
+           0.0f,       0.0f,       0.0f,       1.0f,        // q2, e5
+           0.0177516f, 0.0171426f, 0.0177516f, 0.9995378f,  // q3, e5
+           0.0177516f, 0.0171426f, 0.0177516f, 0.9995378f,  // q4, e5
+           0.0177516f, 0.0171426f, 0.0177516f, 0.9995378f,  // q5, e5
+           0.0177516f, 0.0171426f, 0.0177516f, 0.9995378f,  // q6, e5
+           0.0177516f, 0.0171426f, 0.0177516f, 0.9995378f,  // q7, e5
+           0.0f,       0.0f,       0.0f,       1.0f,        // q0, e6
+           0.0f,       0.0f,       0.0f,       1.0f,        // q1, e6
+           0.0f,       0.0f,       0.0f,       1.0f,        // q2, e6
+           0.0549279f, 0.0494573f, 0.0549279f, 0.9957509f,  // q3, e6
+           0.0549279f, 0.0494573f, 0.0549279f, 0.9957509f,  // q4, e6
+           0.0549279f, 0.0494573f, 0.0549279f, 0.9957509f,  // q5, e6
+           0.0549279f, 0.0494573f, 0.0549279f, 0.9957509f,  // q6, e6
+           0.0549279f, 0.0494573f, 0.0549279f, 0.9957509f,  // q7, e6
+           0.0f,       0.0f,       0.0f,       1.0f,        // q0, e7
+           0.0f,       0.0f,       0.0f,       1.0f,        // q1, e7
+           0.0f,       0.0f,       0.0f,       1.0f,        // q2, e7
+           0.0f,       0.0f,       0.0f,       1.0f,        // q3, e7
+           0.0f,       0.0f,       0.0f,       1.0f,        // q4, e7
+           0.0f,       0.0f,       0.0f,       1.0f,        // q5, e7
+           0.0f,       0.0f,       0.0f,       1.0f,        // q6, e7
+           0.0f,       0.0f,       0.0f,       1.0f,        // q7, e7
+           0.0f,       0.0f,       0.0f,       1.0f,        // q8, e0
+           0.0087262f, 0.0087262f, 0.0000762f, 0.9999238f,  // q8, e1
+           0.0f,       0.0f,       0.0f,       1.0f,        // q8, e2
+           0.0f,       0.0f,       0.0f,       1.0f,        // q8, e3
+           0.0609347f, 0.0609347f, 0.0037269f, 0.9962731f,  // q8, e4
+           0.0174497f, 0.0174497f, 0.0003046f, 0.9996954f,  // q8, e5
+           0.0522642f, 0.0522642f, 0.0027391f, 0.9972609f,  // q8, e6
+           0.026168f,  0.026168f,  0.0006852f, 0.9993148f,  // q8, e7
+        */
+        0.0742712f, 0.0f, 0.0f, 0.0f, 0.0f, 0.0f, 0.0f, 0.0f,  // Xs[q0-q7]
+        0.0645629f, 0.0f, 0.0f, 0.0f, 0.0f, 0.0f, 0.0f, 0.0f,  // Ys[q0-q7]
+        0.0742712f, 0.0f, 0.0f, 0.0f, 0.0f, 0.0f, 0.0f, 0.0f,  // Zs[q0-q7]
+        0.9923705f, 1.0f, 1.0f, 1.0f, 1.0f, 1.0f, 1.0f, 1.0f,  // Ws[q0-q7]
+        0.008802f, 0.008802f, 0.008802f, 0.008802f, 0.008802f, 0.008802f, 0.008802f, 0.008802f,  // Xs[q0-q7]
+        0.0086497f, 0.0086497f, 0.0086497f, 0.0086497f, 0.0086497f, 0.0086497f, 0.0086497f, 0.0086497f,  // Ys[q0-q7]
+        0.008802f, 0.008802f, 0.008802f, 0.008802f, 0.008802f, 0.008802f, 0.008802f, 0.008802f,  // Zs[q0-q7]
+        0.9998851f, 0.9998851f, 0.9998851f, 0.9998851f, 0.9998851f, 0.9998851f, 0.9998851f, 0.9998851f,  // Ws[q0-q7]
+        0.0454372f, 0.0454372f, 0.0454372f, 0.0435779f, 0.0435779f, 0.0435779f, 0.0435779f, 0.0435779f,  // Xs[q0-q7]
+        0.0416356f, 0.0416356f, 0.0416356f, 0.0435779f, 0.0435779f, 0.0435779f, 0.0435779f, 0.0435779f,  // Ys[q0-q7]
+        0.0454372f, 0.0454372f, 0.0454372f, 0.0019027f, 0.0019027f, 0.0019027f, 0.0019027f, 0.0019027f,  // Zs[q0-q7]
+        0.9970644f, 0.9970644f, 0.9970644f, 0.9980973f, 0.9980973f, 0.9980973f, 0.9980973f, 0.9980973f,  // Ws[q0-q7]
+        0.0360742f, 0.0360742f, 0.0360742f, 0.0f, 0.0f, 0.0f, 0.0f, 0.0f,  // Xs[q0-q7]
+        0.0336398f, 0.0336398f, 0.0336398f, 0.0f, 0.0f, 0.0f, 0.0f, 0.0f,  // Ys[q0-q7]
+        0.0360742f, 0.0360742f, 0.0360742f, 0.0f, 0.0f, 0.0f, 0.0f, 0.0f,  // Zs[q0-q7]
+        0.9981311f, 0.9981311f, 0.9981311f, 1.0f, 1.0f, 1.0f, 1.0f, 1.0f,  // Ws[q0-q7]
+        0.0f, 0.064541f, 0.064541f, 0.064541f, 0.064541f, 0.064541f, 0.064541f, 0.064541f,  // Xs[q0-q7]
+        0.0f, 0.057101f, 0.057101f, 0.057101f, 0.057101f, 0.057101f, 0.057101f, 0.057101f,  // Ys[q0-q7]
+        0.0f, 0.064541f, 0.064541f, 0.064541f, 0.064541f, 0.064541f, 0.064541f, 0.064541f,  // Zs[q0-q7]
+        1.0f, 0.9941873f, 0.9941873f, 0.9941873f, 0.9941873f, 0.9941873f, 0.9941873f, 0.9941873f,  // Ws[q0-q7]
+        0.0f, 0.0f, 0.0f, 0.0177516f, 0.0177516f, 0.0177516f, 0.0177516f, 0.0177516f,  // Xs[q0-q7]
+        0.0f, 0.0f, 0.0f, 0.0171426f, 0.0171426f, 0.0171426f, 0.0171426f, 0.0171426f,  // Ys[q0-q7]
+        0.0f, 0.0f, 0.0f, 0.0177516f, 0.0177516f, 0.0177516f, 0.0177516f, 0.0177516f,  // Zs[q0-q7]
+        1.0f, 1.0f, 1.0f, 0.9995378f, 0.9995378f, 0.9995378f, 0.9995378f, 0.9995378f,  // Ws[q0-q7]
+        0.0f, 0.0f, 0.0f, 0.0549279f, 0.0549279f, 0.0549279f, 0.0549279f, 0.0549279f,  // Xs[q0-q7]
+        0.0f, 0.0f, 0.0f, 0.0494573f, 0.0494573f, 0.0494573f, 0.0494573f, 0.0494573f,  // Ys[q0-q7]
+        0.0f, 0.0f, 0.0f, 0.0549279f, 0.0549279f, 0.0549279f, 0.0549279f, 0.0549279f,  // Zs[q0-q7]
+        1.0f, 1.0f, 1.0f, 0.9957509f, 0.9957509f, 0.9957509f, 0.9957509f, 0.9957509f,  // Ws[q0-q7]
+        0.0f, 0.0f, 0.0f, 0.0f, 0.0f, 0.0f, 0.0f, 0.0f,  // Xs[q0-q7]
+        0.0f, 0.0f, 0.0f, 0.0f, 0.0f, 0.0f, 0.0f, 0.0f,  // Ys[q0-q7]
+        0.0f, 0.0f, 0.0f, 0.0f, 0.0f, 0.0f, 0.0f, 0.0f,  // Zs[q0-q7]
+        1.0f, 1.0f, 1.0f, 1.0f, 1.0f, 1.0f, 1.0f, 1.0f,  // Ws[q0-q7]
+        0.0f, 0.0f, 0.0f, 0.0f,  // Xs[q0]
+        0.0f, 0.0f, 0.0f, 0.0f,  // Ys[q0]
+        0.0f, 0.0f, 0.0f, 0.0f,  // Zs[q0]
+        1.0f, 0.0f, 0.0f, 0.0f,  // Ws[q0]
+        0.0087262f, 0.0f, 0.0f, 0.0f,  // Xs[q0]
+        0.0087262f, 0.0f, 0.0f, 0.0f,  // Ys[q0]
+        0.0000762f, 0.0f, 0.0f, 0.0f,  // Zs[q0]
+        0.9999238f, 0.0f, 0.0f, 0.0f,  // Ws[q0]
+        0.0f, 0.0f, 0.0f, 0.0f,  // Xs[q0]
+        0.0f, 0.0f, 0.0f, 0.0f,  // Ys[q0]
+        0.0f, 0.0f, 0.0f, 0.0f,  // Zs[q0]
+        1.0f, 0.0f, 0.0f, 0.0f,  // Ws[q0]
+        0.0f, 0.0f, 0.0f, 0.0f,  // Xs[q0]
+        0.0f, 0.0f, 0.0f, 0.0f,  // Ys[q0]
+        0.0f, 0.0f, 0.0f, 0.0f,  // Zs[q0]
+        1.0f, 0.0f, 0.0f, 0.0f,  // Ws[q0]
+        0.0609347f, 0.0f, 0.0f, 0.0f,  // Xs[q0]
+        0.0609347f, 0.0f, 0.0f, 0.0f,  // Ys[q0]
+        0.0037269f, 0.0f, 0.0f, 0.0f,  // Zs[q0]
+        0.9962731f, 0.0f, 0.0f, 0.0f,  // Ws[q0]
+        0.0174497f, 0.0f, 0.0f, 0.0f,  // Xs[q0]
+        0.0174497f, 0.0f, 0.0f, 0.0f,  // Ys[q0]
+        0.0003046f, 0.0f, 0.0f, 0.0f,  // Zs[q0]
+        0.9996954f, 0.0f, 0.0f, 0.0f,  // Ws[q0]
+        0.0522642f, 0.0f, 0.0f, 0.0f,  // Xs[q0]
+        0.0522642f, 0.0f, 0.0f, 0.0f,  // Ys[q0]
+        0.0027391f, 0.0f, 0.0f, 0.0f,  // Zs[q0]
+        0.9972609f, 0.0f, 0.0f, 0.0f,  // Ws[q0]
+        0.026168f, 0.0f, 0.0f, 0.0f,  // Xs[q0]
+        0.026168f, 0.0f, 0.0f, 0.0f,  // Ys[q0]
+        0.0006852f, 0.0f, 0.0f, 0.0f,  // Zs[q0]
+        0.9993148f, 0.0f, 0.0f, 0.0f,  // Ws[q0]
     }
 };
 
@@ -427,6 +629,72 @@
         13155, 9600, 0, 0,
         12646, 12836, 0, 0,
         15231, 15189, 0, 0
+    },
+    {  // Joint group 7
+        11457, 0, 0, 0, 0, 0, 0, 0,
+        11298, 0, 0, 0, 0, 0, 0, 0,
+        11457, 0, 0, 0, 0, 0, 0, 0,
+        15344, 15360, 15360, 15360, 15360, 15360, 15360, 15360,
+        8322, 8322, 8322, 8322, 8322, 8322, 8322, 8322,
+        8302, 8302, 8302, 8302, 8302, 8302, 8302, 8302,
+        8322, 8322, 8322, 8322, 8322, 8322, 8322, 8322,
+        15360, 15360, 15360, 15360, 15360, 15360, 15360, 15360,
+        10705, 10705, 10705, 10644, 10644, 10644, 10644, 10644,
+        10580, 10580, 10580, 10644, 10644, 10644, 10644, 10644,
+        10705, 10705, 10705, 6091, 6091, 6091, 6091, 6091,
+        15354, 15354, 15354, 15356, 15356, 15356, 15356, 15356,
+        10398, 10398, 10398, 0, 0, 0, 0, 0,
+        10318, 10318, 10318, 0, 0, 0, 0, 0,
+        10398, 10398, 10398, 0, 0, 0, 0, 0,
+        15356, 15356, 15356, 15360, 15360, 15360, 15360, 15360,
+        0, 11297, 11297, 11297, 11297, 11297, 11297, 11297,
+        0, 11087, 11087, 11087, 11087, 11087, 11087, 11087,
+        0, 11297, 11297, 11297, 11297, 11297, 11297, 11297,
+        15360, 15348, 15348, 15348, 15348, 15348, 15348, 15348,
+        0, 0, 0, 9355, 9355, 9355, 9355, 9355,
+        0, 0, 0, 9315, 9315, 9315, 9315, 9315,
+        0, 0, 0, 9355, 9355, 9355, 9355, 9355,
+        15360, 15360, 15360, 15359, 15359, 15359, 15359, 15359,
+        0, 0, 0, 11016, 11016, 11016, 11016, 11016,
+        0, 0, 0, 10837, 10837, 10837, 10837, 10837,
+        0, 0, 0, 11016, 11016, 11016, 11016, 11016,
+        15360, 15360, 15360, 15351, 15351, 15351, 15351, 15351,
+        0, 0, 0, 0, 0, 0, 0, 0,
+        0, 0, 0, 0, 0, 0, 0, 0,
+        0, 0, 0, 0, 0, 0, 0, 0,
+        15360, 15360, 15360, 15360, 15360, 15360, 15360, 15360,
+        0, 0, 0, 0,
+        0, 0, 0, 0,
+        0, 0, 0, 0,
+        15360, 0, 0, 0,
+        8312, 0, 0, 0,
+        8312, 0, 0, 0,
+        1278, 0, 0, 0,
+        15360, 0, 0, 0,
+        0, 0, 0, 0,
+        0, 0, 0, 0,
+        0, 0, 0, 0,
+        15360, 0, 0, 0,
+        0, 0, 0, 0,
+        0, 0, 0, 0,
+        0, 0, 0, 0,
+        15360, 0, 0, 0,
+        11213, 0, 0, 0,
+        11213, 0, 0, 0,
+        7074, 0, 0, 0,
+        15352, 0, 0, 0,
+        9336, 0, 0, 0,
+        9336, 0, 0, 0,
+        3326, 0, 0, 0,
+        15359, 0, 0, 0,
+        10929, 0, 0, 0,
+        10929, 0, 0, 0,
+        6556, 0, 0, 0,
+        15354, 0, 0, 0,
+        9907, 0, 0, 0,
+        9907, 0, 0, 0,
+        4509, 0, 0, 0,
+        15359, 0, 0, 0,
     }
 };
 
@@ -438,9 +706,10 @@
     {0, 0},  // Joint group 4
     {0, 0},  // Joint group 5
     {16, 2},  // Joint group 6
+    {36, 8},  // Joint group 7
 };
 
-const AlignedMatrix<std::uint16_t> inputIndices = {
+const Matrix<std::uint16_t> inputIndices = {
     {  // Joint group 0
         8, 6
     },
@@ -459,10 +728,13 @@
     },
     {  // Joint group 6
         8, 6
+    },
+    {  // Joint group 7
+        9, 2, 6, 5, 8, 3, 7, 1
     }
 };
 
-const Matrix<AlignedVector<std::uint16_t> > outputIndices = {
+const Matrix<Vector<std::uint16_t> > outputIndices = {
     {  // Quaternion outputs
         {  // Joint group 0
             3, 4, 5, 6,
@@ -502,6 +774,12 @@
             243, 244, 245, 246,
             233, 234, 235, 236
         },
+        {  // Joint group 7
+            253, 254, 255, 256,
+            263, 264, 265, 266, 273, 274, 275, 276,
+            283, 284, 285, 286, 293, 294, 295, 296, 303, 304, 305, 306, 313, 314, 315, 316, 323, 324, 325, 326,
+            333, 334, 335, 336
+        },
     },
     {  // Euler-angle outputs
         {  // Joint group 0
@@ -542,52 +820,64 @@
             219, 220, 221, 0,
             210, 211, 212, 0
         },
+        {  // Joint group 7
+            228, 229, 230, 0,
+            237, 238, 239, 0, 246, 247, 248, 0,
+            255, 256, 257, 0, 264, 265, 266, 0, 273, 274, 275, 0, 282, 283, 284, 0, 291, 292, 293, 0,
+            300, 301, 302, 0
+        },
     }
 };
 
 const Matrix<LODRegion> lodRegions = {
     // {unaligned, aligned to last block-8, aligned to second last block-8}
     {  // Joint group 0
-        {32, 32, 32},  // LOD-0
-        {16, 32, 0},  // LOD-1
-        {16, 32, 0},  // LOD-2
-        {0, 0, 0},  // LOD-3
+        {{2, 0, 0}, {32, 32, 32}},  // LOD-0
+        {{2, 0, 0}, {16, 32, 0}},  // LOD-1
+        {{2, 0, 0}, {16, 32, 0}},  // LOD-2
+        {{0, 0, 0}, {0, 0, 0}},  // LOD-3
     },
     {  // Joint group 1
-        {16, 0, 0},  // LOD-0
-        {8, 0, 0},  // LOD-1
-        {8, 0, 0},  // LOD-2
-        {0, 0, 0},  // LOD-3
+        {{2, 0, 0}, {16, 0, 0}},  // LOD-0
+        {{2, 0, 0}, {8, 0, 0}},  // LOD-1
+        {{2, 0, 0}, {8, 0, 0}},  // LOD-2
+        {{0, 0, 0}, {0, 0, 0}},  // LOD-3
     },
     {  // Joint group 2
-        {28, 32, 0},  // LOD-0
-        {12, 32, 0},  // LOD-1
-        {12, 32, 0},  // LOD-2
-        {0, 0, 0},  // LOD-3
+        {{2, 0, 0}, {28, 32, 0}},  // LOD-0
+        {{2, 0, 0}, {12, 32, 0}},  // LOD-1
+        {{2, 0, 0}, {12, 32, 0}},  // LOD-2
+        {{0, 0, 0}, {0, 0, 0}},  // LOD-3
     },
     {  // Joint group 3
-        {12, 0, 0},  // LOD-0
-        {4, 0, 0},  // LOD-1
-        {4, 0, 0},  // LOD-2
-        {0, 0, 0},  // LOD-3
+        {{2, 0, 0}, {12, 0, 0}},  // LOD-0
+        {{2, 0, 0}, {4, 0, 0}},  // LOD-1
+        {{2, 0, 0}, {4, 0, 0}},  // LOD-2
+        {{0, 0, 0}, {0, 0, 0}},  // LOD-3
     },
     {  // Joint group 4
-        {0, 0, 0},  // LOD-0
-        {0, 0, 0},  // LOD-1
-        {0, 0, 0},  // LOD-2
-        {0, 0, 0},  // LOD-3
+        {{0, 0, 0}, {0, 0, 0}},  // LOD-0
+        {{0, 0, 0}, {0, 0, 0}},  // LOD-1
+        {{0, 0, 0}, {0, 0, 0}},  // LOD-2
+        {{0, 0, 0}, {0, 0, 0}},  // LOD-3
     },
     {  // Joint group 5
-        {0, 0, 0},  // LOD-0
-        {0, 0, 0},  // LOD-1
-        {0, 0, 0},  // LOD-2
-        {0, 0, 0},  // LOD-3
+        {{0, 0, 0}, {0, 0, 0}},  // LOD-0
+        {{0, 0, 0}, {0, 0, 0}},  // LOD-1
+        {{0, 0, 0}, {0, 0, 0}},  // LOD-2
+        {{0, 0, 0}, {0, 0, 0}},  // LOD-3
     },
     {  // Joint group 6
-        {8, 0, 0},  // LOD-0
-        {4, 0, 0},  // LOD-1
-        {0, 0, 0},  // LOD-2
-        {0, 0, 0},  // LOD-3
+        {{2, 0, 0}, {8, 0, 0}},  // LOD-0
+        {{2, 0, 0}, {4, 0, 0}},  // LOD-1
+        {{0, 0, 0}, {0, 0, 0}},  // LOD-2
+        {{0, 0, 0}, {0, 0, 0}},  // LOD-3
+    },
+    {  // Joint group 7
+        {{8, 8, 8}, {36, 32, 32}},  // LOD-0
+        {{7, 4, 0}, {32, 32, 32}},  // LOD-1
+        {{5, 4, 0}, {12, 32, 0}},  // LOD-2
+        {{4, 4, 0}, {4, 32, 0}},  // LOD-3
     },
 };
 
@@ -863,6 +1153,97 @@
             0.850805f,
             0.0f,
             0.0f,
+            0.0f,
+            // Joint group 7 outputs
+            0.0f,
+            0.0f,
+            0.0f,
+            0.0560275689f,
+            0.0518481769f,
+            0.0560221002f,
+            0.994890451f,
+            0.0f,
+            0.0f,
+            0.0f,
+            0.0f,
+            0.0f,
+            0.0f,
+            0.113435827f,
+            0.102763742f,
+            0.113693371f,
+            0.981072128f,
+            0.0f,
+            0.0f,
+            0.0f,
+            0.0f,
+            0.0f,
+            0.0f,
+            0.113435827f,
+            0.102763742f,
+            0.113693371f,
+            0.981072128f,
+            0.0f,
+            0.0f,
+            0.0f,
+            0.0f,
+            0.0f,
+            0.0f,
+            0.138010800f,
+            0.132913858f,
+            0.112703480f,
+            0.974272072f,
+            0.0f,
+            0.0f,
+            0.0f,
+            0.0f,
+            0.0f,
+            0.0f,
+            0.138010800f,
+            0.132913858f,
+            0.112703480f,
+            0.974272072f,
+            0.0f,
+            0.0f,
+            0.0f,
+            0.0f,
+            0.0f,
+            0.0f,
+            0.138010800f,
+            0.132913858f,
+            0.112703480f,
+            0.974272072f,
+            0.0f,
+            0.0f,
+            0.0f,
+            0.0f,
+            0.0f,
+            0.0f,
+            0.138010800f,
+            0.132913858f,
+            0.112703480f,
+            0.974272072f,
+            0.0f,
+            0.0f,
+            0.0f,
+            0.0f,
+            0.0f,
+            0.0f,
+            0.138010800f,
+            0.132913858f,
+            0.112703480f,
+            0.974272072f,
+            0.0f,
+            0.0f,
+            0.0f,
+            0.0f,
+            0.0f,
+            0.0f,
+            0.111141264f,
+            0.111082502f,
+            0.00580846099f,
+            0.986989498f,
+            0.0f,
+            0.0f,
             0.0f
         }, {
             // LOD-1
@@ -1122,6 +1503,97 @@
             0.850805f,
             0.0f,
             0.0f,
+            0.0f,
+            // Joint group 7 outputs
+            0.0f,
+            0.0f,
+            0.0f,
+            0.0560275689f,
+            0.0518481769f,
+            0.0560221002f,
+            0.994890451f,
+            0.0f,
+            0.0f,
+            0.0f,
+            0.0f,
+            0.0f,
+            0.0f,
+            0.113435827f,
+            0.102763742f,
+            0.113693371f,
+            0.981072128f,
+            0.0f,
+            0.0f,
+            0.0f,
+            0.0f,
+            0.0f,
+            0.0f,
+            0.113435827f,
+            0.102763742f,
+            0.113693371f,
+            0.981072128f,
+            0.0f,
+            0.0f,
+            0.0f,
+            0.0f,
+            0.0f,
+            0.0f,
+            0.138010800f,
+            0.132913858f,
+            0.112703480f,
+            0.974272072f,
+            0.0f,
+            0.0f,
+            0.0f,
+            0.0f,
+            0.0f,
+            0.0f,
+            0.138010800f,
+            0.132913858f,
+            0.112703480f,
+            0.974272072f,
+            0.0f,
+            0.0f,
+            0.0f,
+            0.0f,
+            0.0f,
+            0.0f,
+            0.138010800f,
+            0.132913858f,
+            0.112703480f,
+            0.974272072f,
+            0.0f,
+            0.0f,
+            0.0f,
+            0.0f,
+            0.0f,
+            0.0f,
+            0.138010800f,
+            0.132913858f,
+            0.112703480f,
+            0.974272072f,
+            0.0f,
+            0.0f,
+            0.0f,
+            0.0f,
+            0.0f,
+            0.0f,
+            0.138010800f,
+            0.132913858f,
+            0.112703480f,
+            0.974272072f,
+            0.0f,
+            0.0f,
+            0.0f,
+            0.0f,
+            0.0f,
+            0.0f,
+            0.0f,
+            0.0f,
+            0.0f,
+            0.0f,
+            0.0f,
+            0.0f,
             0.0f
         }, {
             // LOD-2
@@ -1381,28 +1853,35 @@
             0.0f,
             0.0f,
             0.0f,
-            0.0f
-        }, {
-            // LOD-3
-            // Joint group 0 outputs
+            0.0f,
+            // Joint group 7 outputs
+            0.0f,
+            0.0f,
+            0.0f,
+            0.0560275689f,
+            0.0518481769f,
+            0.0560221002f,
+            0.994890451f,
             0.0f,
             0.0f,
             0.0f,
             0.0f,
             0.0f,
             0.0f,
+            0.113435827f,
+            0.102763742f,
+            0.113693371f,
+            0.981072128f,
             0.0f,
             0.0f,
             0.0f,
             0.0f,
             0.0f,
             0.0f,
+            0.113435827f,
+            0.102763742f,
+            0.113693371f,
+            0.981072128f,
             0.0f,
             0.0f,
             0.0f,
-            0.0f,
-            0.0f,
-            0.0f,
-            0.0f,
-            0.0f,
-            0.0f,
@@ -1465,8 +1944,10 @@
             0.0f,
             0.0f,
             0.0f,
-            // Joint group 1 outputs
+            0.0f
+        }, {
+            // LOD-3
+            // Joint group 0 outputs
             0.0f,
             0.0f,
             0.0f,
-            0.0f,
@@ -1506,7 +1987,6 @@
             0.0f,
             0.0f,
             0.0f,
-            // Joint group 2 outputs
             0.0f,
             0.0f,
             0.0f,
@@ -1548,6 +2028,7 @@
             0.0f,
             0.0f,
             0.0f,
+            // Joint group 1 outputs
             0.0f,
             0.0f,
             0.0f,
@@ -1577,7 +2058,6 @@
             0.0f,
             0.0f,
             0.0f,
-            // Joint group 3 outputs
             0.0f,
             0.0f,
             0.0f,
@@ -1589,6 +2069,7 @@
             0.0f,
             0.0f,
             0.0f,
+            // Joint group 2 outputs
             0.0f,
             0.0f,
             0.0f,
@@ -1608,7 +2089,6 @@
             0.0f,
             0.0f,
             0.0f,
-            // Joint group 4 outputs
             0.0f,
             0.0f,
             0.0f,
@@ -1619,8 +2099,6 @@
             0.0f,
             0.0f,
             0.0f,
-            // Joint group 5 outputs
-            // Joint group 6 outputs
             0.0f,
             0.0f,
             0.0f,
@@ -1640,215 +2118,147 @@
             0.0f,
             0.0f,
             0.0f,
-            0.0f
-        }
-    },
-    {  // Euler-angle outputs
-        {
-            // LOD-0
-            // Joint group 0 outputs
-            0.0f,
-            0.0f,
-            0.0f,
-            0.763116f,
-            0.636969f,
-            0.326027f,
             0.0f,
             0.0f,
             0.0f,
             0.0f,
             0.0f,
             0.0f,
-            0.580409f,
-            0.440315f,
-            0.891871f,
             0.0f,
             0.0f,
             0.0f,
             0.0f,
             0.0f,
             0.0f,
-            0.544738f,
-            1.05558f,
-            1.19192f,
             0.0f,
             0.0f,
             0.0f,
             0.0f,
             0.0f,
             0.0f,
-            0.471332f,
-            0.733388f,
-            0.317664f,
             0.0f,
             0.0f,
             0.0f,
             0.0f,
-            0.0f,
-            0.0f,
-            0.194852f,
-            0.206271f,
-            0.195788f,
-            0.0f,
-            0.0f,
-            0.0f,
-            0.0f,
-            0.0f,
-            0.0f,
-            0.444944f,
-            0.514331f,
-            0.45186f,
-            0.0f,
-            0.0f,
-            0.0f,
-            0.0f,
-            0.0f,
-            0.0f,
-            0.679014f,
-            0.838617f,
-            0.702706f,
-            0.0f,
-            0.0f,
-            0.0f,
-            0.0f,
-            0.0f,
-            0.0f,
-            0.891595f,
-            1.17419f,
-            0.968903f,
-            0.0f,
-            0.0f,
-            0.0f,
-            // Joint group 1 outputs
-            0.0f,
-            0.0f,
-            0.0f,
-            0.763116f,
-            0.636969f,
-            0.326027f,
-            0.0f,
-            0.0f,
-            0.0f,
-            0.0f,
-            0.0f,
-            0.0f,
-            0.580409f,
-            0.440315f,
-            0.891871f,
-            0.0f,
-            0.0f,
-            0.0f,
-            0.0f,
-            0.0f,
-            0.0f,
-            0.544738f,
-            1.05558f,
-            1.19192f,
-            0.0f,
-            0.0f,
-            0.0f,
-            0.0f,
-            0.0f,
-            0.0f,
-            0.471332f,
-            0.733388f,
-            0.317664f,
-            0.0f,
-            0.0f,
-            0.0f,
-            // Joint group 2 outputs
-            0.0f,
-            0.0f,
-            0.0f,
-            0.763116f,
-            0.636969f,
-            0.326027f,
-            0.0f,
-            0.0f,
-            0.0f,
-            0.0f,
-            0.0f,
-            0.0f,
-            0.580409f,
-            0.440315f,
-            0.891871f,
-            0.0f,
-            0.0f,
-            0.0f,
-            0.0f,
-            0.0f,
-            0.0f,
-            0.544738f,
-            1.05558f,
-            1.19192f,
-            0.0f,
-            0.0f,
-            0.0f,
-            0.0f,
-            0.0f,
-            0.0f,
-            0.471332f,
-            0.733388f,
-            0.317664f,
-            0.0f,
-            0.0f,
-            0.0f,
-            0.0f,
-            0.0f,
-            0.0f,
-            0.194852f,
-            0.206271f,
-            0.195788f,
-            0.0f,
-            0.0f,
-            0.0f,
-            0.0f,
-            0.0f,
-            0.0f,
-            0.444944f,
-            0.514331f,
-            0.45186f,
-            0.0f,
-            0.0f,
-            0.0f,
-            0.0f,
-            0.0f,
-            0.0f,
-            0.679014f,
-            0.838617f,
-            0.702706f,
-            0.0f,
-            0.0f,
-            0.0f,
             // Joint group 3 outputs
             0.0f,
             0.0f,
             0.0f,
-            0.763116f,
-            0.636969f,
-            0.326027f,
             0.0f,
             0.0f,
             0.0f,
             0.0f,
             0.0f,
             0.0f,
-            0.580409f,
-            0.440315f,
-            0.891871f,
             0.0f,
             0.0f,
             0.0f,
             0.0f,
             0.0f,
             0.0f,
-            0.544738f,
-            1.05558f,
-            1.19192f,
             0.0f,
             0.0f,
             0.0f,
+            0.0f,
+            0.0f,
+            0.0f,
+            0.0f,
+            0.0f,
+            0.0f,
+            0.0f,
+            0.0f,
+            0.0f,
+            0.0f,
+            0.0f,
+            0.0f,
             // Joint group 4 outputs
             0.0f,
             0.0f,
             0.0f,
+            0.0f,
+            0.0f,
+            0.0f,
+            0.0f,
+            0.0f,
+            0.0f,
+            0.0f,
+            // Joint group 5 outputs
+            // Joint group 6 outputs
+            0.0f,
+            0.0f,
+            0.0f,
+            0.0f,
+            0.0f,
+            0.0f,
+            0.0f,
+            0.0f,
+            0.0f,
+            0.0f,
+            0.0f,
+            0.0f,
+            0.0f,
+            0.0f,
+            0.0f,
+            0.0f,
+            0.0f,
+            0.0f,
+            0.0f,
+            0.0f,
+            // Joint group 7 outputs
+            0.0f,
+            0.0f,
+            0.0f,
+            0.0560275689f,
+            0.0518481769f,
+            0.0560221002f,
+            0.994890451f,
+            0.0f,
+            0.0f,
+            0.0f,
+            0.0f,
+            0.0f,
+            0.0f,
+            0.0f,
+            0.0f,
+            0.0f,
+            0.0f,
+            0.0f,
+            0.0f,
+            0.0f,
+            0.0f,
+            0.0f,
+            0.0f,
+            0.0f,
+            0.0f,
+            0.0f,
+            0.0f,
+            0.0f,
+            0.0f,
+            0.0f,
+            0.0f,
+            0.0f,
+            0.0f,
+            0.0f,
+            0.0f,
+            0.0f,
+            0.0f,
+            0.0f,
+            0.0f,
+            0.0f,
+            0.0f,
+            0.0f,
+            0.0f,
+            0.0f,
+            0.0f,
+            0.0f,
+            0.0f,
+            0.0f,
+            0.0f,
+            0.0f,
+            0.0f,
+            0.0f,
+            0.0f,
+            0.0f,
+            0.0f,
+            0.0f,
+            0.0f,
@@ -1858,6 +2268,250 @@
             0.0f,
             0.0f,
             0.0f,
+            0.0f,
+            0.0f,
+            0.0f,
+            0.0f,
+            0.0f,
+            0.0f,
+            0.0f,
+            0.0f,
+            0.0f,
+            0.0f,
+            0.0f,
+            0.0f,
+            0.0f,
+            0.0f,
+            0.0f,
+            0.0f,
+            0.0f,
+            0.0f,
+            0.0f,
+            0.0f,
+            0.0f,
+            0.0f,
+            0.0f,
+            0.0f,
+            0.0f,
+            0.0f,
+            0.0f
+        }
+    },
+    {  // Euler-angle outputs
+        {
+            // LOD-0
+            // Joint group 0 outputs
+            0.0f,
+            0.0f,
+            0.0f,
+            0.763116f,
+            0.636969f,
+            0.326027f,
+            0.0f,
+            0.0f,
+            0.0f,
+            0.0f,
+            0.0f,
+            0.0f,
+            0.580409f,
+            0.440315f,
+            0.891871f,
+            0.0f,
+            0.0f,
+            0.0f,
+            0.0f,
+            0.0f,
+            0.0f,
+            0.544738f,
+            1.05558f,
+            1.19192f,
+            0.0f,
+            0.0f,
+            0.0f,
+            0.0f,
+            0.0f,
+            0.0f,
+            0.471332f,
+            0.733388f,
+            0.317664f,
+            0.0f,
+            0.0f,
+            0.0f,
+            0.0f,
+            0.0f,
+            0.0f,
+            0.194852f,
+            0.206271f,
+            0.195788f,
+            0.0f,
+            0.0f,
+            0.0f,
+            0.0f,
+            0.0f,
+            0.0f,
+            0.444944f,
+            0.514331f,
+            0.45186f,
+            0.0f,
+            0.0f,
+            0.0f,
+            0.0f,
+            0.0f,
+            0.0f,
+            0.679014f,
+            0.838617f,
+            0.702706f,
+            0.0f,
+            0.0f,
+            0.0f,
+            0.0f,
+            0.0f,
+            0.0f,
+            0.891595f,
+            1.17419f,
+            0.968903f,
+            0.0f,
+            0.0f,
+            0.0f,
+            // Joint group 1 outputs
+            0.0f,
+            0.0f,
+            0.0f,
+            0.763116f,
+            0.636969f,
+            0.326027f,
+            0.0f,
+            0.0f,
+            0.0f,
+            0.0f,
+            0.0f,
+            0.0f,
+            0.580409f,
+            0.440315f,
+            0.891871f,
+            0.0f,
+            0.0f,
+            0.0f,
+            0.0f,
+            0.0f,
+            0.0f,
+            0.544738f,
+            1.05558f,
+            1.19192f,
+            0.0f,
+            0.0f,
+            0.0f,
+            0.0f,
+            0.0f,
+            0.0f,
+            0.471332f,
+            0.733388f,
+            0.317664f,
+            0.0f,
+            0.0f,
+            0.0f,
+            // Joint group 2 outputs
+            0.0f,
+            0.0f,
+            0.0f,
+            0.763116f,
+            0.636969f,
+            0.326027f,
+            0.0f,
+            0.0f,
+            0.0f,
+            0.0f,
+            0.0f,
+            0.0f,
+            0.580409f,
+            0.440315f,
+            0.891871f,
+            0.0f,
+            0.0f,
+            0.0f,
+            0.0f,
+            0.0f,
+            0.0f,
+            0.544738f,
+            1.05558f,
+            1.19192f,
+            0.0f,
+            0.0f,
+            0.0f,
+            0.0f,
+            0.0f,
+            0.0f,
+            0.471332f,
+            0.733388f,
+            0.317664f,
+            0.0f,
+            0.0f,
+            0.0f,
+            0.0f,
+            0.0f,
+            0.0f,
+            0.194852f,
+            0.206271f,
+            0.195788f,
+            0.0f,
+            0.0f,
+            0.0f,
+            0.0f,
+            0.0f,
+            0.0f,
+            0.444944f,
+            0.514331f,
+            0.45186f,
+            0.0f,
+            0.0f,
+            0.0f,
+            0.0f,
+            0.0f,
+            0.0f,
+            0.679014f,
+            0.838617f,
+            0.702706f,
+            0.0f,
+            0.0f,
+            0.0f,
+            // Joint group 3 outputs
+            0.0f,
+            0.0f,
+            0.0f,
+            0.763116f,
+            0.636969f,
+            0.326027f,
+            0.0f,
+            0.0f,
+            0.0f,
+            0.0f,
+            0.0f,
+            0.0f,
+            0.580409f,
+            0.440315f,
+            0.891871f,
+            0.0f,
+            0.0f,
+            0.0f,
+            0.0f,
+            0.0f,
+            0.0f,
+            0.544738f,
+            1.05558f,
+            1.19192f,
+            0.0f,
+            0.0f,
+            0.0f,
+            // Joint group 4 outputs
+            0.0f,
+            0.0f,
+            0.0f,
+            0.0f,
+            0.0f,
+            0.0f,
+            0.0f,
+            0.0f,
+            0.0f,
             // Joint group 5 outputs
             // Joint group 6 outputs
             0.0f,
@@ -1877,11 +2531,92 @@
             0.326027f,
             0.0f,
             0.0f,
+            0.0f,
+            // Joint group 7 outputs
+            0.0f,
+            0.0f,
+            0.0f,
+            0.106514752f,
+            0.109663844f,
+            0.106503174f,
+            0.0f,
+            0.0f,
+            0.0f,
+            0.0f,
+            0.0f,
+            0.0f,
+            0.206037536f,
+            0.229438782f,
+            0.206623316f,
+            0.0f,
+            0.0f,
+            0.0f,
+            0.0f,
+            0.0f,
+            0.0f,
+            0.206037536f,
+            0.229438782f,
+            0.206623316f,
+            0.0f,
+            0.0f,
+            0.0f,
+            0.0f,
+            0.0f,
+            0.0f,
+            0.252397001f,
+            0.294328332f,
+            0.192341432f,
+            0.0f,
+            0.0f,
+            0.0f,
+            0.0f,
+            0.0f,
+            0.0f,
+            0.252397001f,
+            0.294328332f,
+            0.192341432f,
+            0.0f,
+            0.0f,
+            0.0f,
+            0.0f,
+            0.0f,
+            0.0f,
+            0.252397001f,
+            0.294328332f,
+            0.192341432f,
+            0.0f,
+            0.0f,
+            0.0f,
+            0.0f,
+            0.0f,
+            0.0f,
+            0.252397001f,
+            0.294328332f,
+            0.192341432f,
+            0.0f,
+            0.0f,
+            0.0f,
+            0.0f,
+            0.0f,
+            0.0f,
+            0.252397001f,
+            0.294328332f,
+            0.192341432f,
+            0.0f,
+            0.0f,
+            0.0f,
+            0.0f,
+            0.0f,
+            0.0f,
+            0.225527078f,
+            0.222394347f,
+            -0.0135606825f,
+            0.0f,
+            0.0f,
             0.0f
         }, {
             // LOD-1
             // Joint group 0 outputs
-            // Joint group 0 outputs
             0.0f,
             0.0f,
             0.0f,
@@ -2112,11 +2847,92 @@
             0.326027f,
             0.0f,
             0.0f,
+            0.0f,
+            // Joint group 7 outputs
+            0.0f,
+            0.0f,
+            0.0f,
+            0.106514752f,
+            0.109663844f,
+            0.106503174f,
+            0.0f,
+            0.0f,
+            0.0f,
+            0.0f,
+            0.0f,
+            0.0f,
+            0.206037536f,
+            0.229438782f,
+            0.206623316f,
+            0.0f,
+            0.0f,
+            0.0f,
+            0.0f,
+            0.0f,
+            0.0f,
+            0.206037536f,
+            0.229438782f,
+            0.206623316f,
+            0.0f,
+            0.0f,
+            0.0f,
+            0.0f,
+            0.0f,
+            0.0f,
+            0.252397001f,
+            0.294328332f,
+            0.192341432f,
+            0.0f,
+            0.0f,
+            0.0f,
+            0.0f,
+            0.0f,
+            0.0f,
+            0.252397001f,
+            0.294328332f,
+            0.192341432f,
+            0.0f,
+            0.0f,
+            0.0f,
+            0.0f,
+            0.0f,
+            0.0f,
+            0.252397001f,
+            0.294328332f,
+            0.192341432f,
+            0.0f,
+            0.0f,
+            0.0f,
+            0.0f,
+            0.0f,
+            0.0f,
+            0.252397001f,
+            0.294328332f,
+            0.192341432f,
+            0.0f,
+            0.0f,
+            0.0f,
+            0.0f,
+            0.0f,
+            0.0f,
+            0.252397001f,
+            0.294328332f,
+            0.192341432f,
+            0.0f,
+            0.0f,
+            0.0f,
+            0.0f,
+            0.0f,
+            0.0f,
+            0.0f,
+            0.0f,
+            0.0f,
+            0.0f,
+            0.0f,
             0.0f
         }, {
             // LOD-2
             // Joint group 0 outputs
-            // Joint group 0 outputs
             0.0f,
             0.0f,
             0.0f,
@@ -2347,6 +3163,88 @@
             0.0f,
             0.0f,
             0.0f,
+            0.0f,
+            // Joint group 7 outputs
+            0.0f,
+            0.0f,
+            0.0f,
+            0.106514752f,
+            0.109663844f,
+            0.106503174f,
+            0.0f,
+            0.0f,
+            0.0f,
+            0.0f,
+            0.0f,
+            0.0f,
+            0.206037536f,
+            0.229438782f,
+            0.206623316f,
+            0.0f,
+            0.0f,
+            0.0f,
+            0.0f,
+            0.0f,
+            0.0f,
+            0.206037536f,
+            0.229438782f,
+            0.206623316f,
+            0.0f,
+            0.0f,
+            0.0f,
+            0.0f,
+            0.0f,
+            0.0f,
+            0.0f,
+            0.0f,
+            0.0f,
+            0.0f,
+            0.0f,
+            0.0f,
+            0.0f,
+            0.0f,
+            0.0f,
+            0.0f,
+            0.0f,
+            0.0f,
+            0.0f,
+            0.0f,
+            0.0f,
+            0.0f,
+            0.0f,
+            0.0f,
+            0.0f,
+            0.0f,
+            0.0f,
+            0.0f,
+            0.0f,
+            0.0f,
+            0.0f,
+            0.0f,
+            0.0f,
+            0.0f,
+            0.0f,
+            0.0f,
+            0.0f,
+            0.0f,
+            0.0f,
+            0.0f,
+            0.0f,
+            0.0f,
+            0.0f,
+            0.0f,
+            0.0f,
+            0.0f,
+            0.0f,
+            0.0f,
+            0.0f,
+            0.0f,
+            0.0f,
+            0.0f,
+            0.0f,
+            0.0f,
+            0.0f,
+            0.0f,
             0.0f
         }, {
             // LOD-3
@@ -2581,6 +3479,88 @@
             0.0f,
             0.0f,
             0.0f,
+            0.0f,
+            // Joint group 7 outputs
+            0.0f,
+            0.0f,
+            0.0f,
+            0.106514752f,
+            0.109663844f,
+            0.106503174f,
+            0.0f,
+            0.0f,
+            0.0f,
+            0.0f,
+            0.0f,
+            0.0f,
+            0.0f,
+            0.0f,
+            0.0f,
+            0.0f,
+            0.0f,
+            0.0f,
+            0.0f,
+            0.0f,
+            0.0f,
+            0.0f,
+            0.0f,
+            0.0f,
+            0.0f,
+            0.0f,
+            0.0f,
+            0.0f,
+            0.0f,
+            0.0f,
+            0.0f,
+            0.0f,
+            0.0f,
+            0.0f,
+            0.0f,
+            0.0f,
+            0.0f,
+            0.0f,
+            0.0f,
+            0.0f,
+            0.0f,
+            0.0f,
+            0.0f,
+            0.0f,
+            0.0f,
+            0.0f,
+            0.0f,
+            0.0f,
+            0.0f,
+            0.0f,
+            0.0f,
+            0.0f,
+            0.0f,
+            0.0f,
+            0.0f,
+            0.0f,
+            0.0f,
+            0.0f,
+            0.0f,
+            0.0f,
+            0.0f,
+            0.0f,
+            0.0f,
+            0.0f,
+            0.0f,
+            0.0f,
+            0.0f,
+            0.0f,
+            0.0f,
+            0.0f,
+            0.0f,
+            0.0f,
+            0.0f,
+            0.0f,
+            0.0f,
+            0.0f,
+            0.0f,
+            0.0f,
+            0.0f,
+            0.0f,
             0.0f
         }
     }
diff --git a/Engine/Plugins/Animation/RigLogic/Source/RigLogicLibTest/Private/rltests/joints/quaternions/QuaternionFixtures.h b/Engine/Plugins/Animation/RigLogic/Source/RigLogicLibTest/Private/rltests/joints/quaternions/QuaternionFixtures.h
--- a/Engine/Plugins/Animation/RigLogic/Source/RigLogicLibTest/Private/rltests/joints/quaternions/QuaternionFixtures.h
+++ b/Engine/Plugins/Animation/RigLogic/Source/RigLogicLibTest/Private/rltests/joints/quaternions/QuaternionFixtures.h
@@ -44,8 +44,9 @@ namespace optimized {
 
 extern const AlignedMatrix<float> floatValues;
 extern const AlignedMatrix<std::uint16_t> halfFloatValues;
-extern const AlignedMatrix<std::uint16_t> inputIndices;
-extern const Matrix<AlignedVector<std::uint16_t> > outputIndices;
+extern const Vector<Extent> subMatrices;
+extern const Matrix<std::uint16_t> inputIndices;
+extern const Matrix<Vector<std::uint16_t> > outputIndices;
 extern const Matrix<LODRegion> lodRegions;
 
 template<typename TValue>
diff --git a/Engine/Plugins/Animation/RigLogic/Source/RigLogicLibTest/Private/rltests/joints/quaternions/TestQuaternionEvaluator.cpp b/Engine/Plugins/Animation/RigLogic/Source/RigLogicLibTest/Private/rltests/joints/quaternions/TestQuaternionEvaluator.cpp
--- a/Engine/Plugins/Animation/RigLogic/Source/RigLogicLibTest/Private/rltests/joints/quaternions/TestQuaternionEvaluator.cpp
+++ b/Engine/Plugins/Animation/RigLogic/Source/RigLogicLibTest/Private/rltests/joints/quaternions/TestQuaternionEvaluator.cpp
@@ -48,6 +48,8 @@ class QuaternionEvaluatorTest : public ::testing::Test {
                 jointGroup.inputIndices = rltests::qs::optimized::inputIndices[jgi];
                 jointGroup.outputIndices = rltests::qs::optimized::outputIndices[rotationSelectorIndex][jgi];
                 jointGroup.lods = rltests::qs::optimized::lodRegions[jgi];
+                jointGroup.colCount = rltests::qs::optimized::subMatrices[jgi].cols;
+                jointGroup.rowCount = rltests::qs::optimized::subMatrices[jgi].rows;
             }
             auto factory = rl4::UniqueInstance<QJEvaluator, rl4::JointsEvaluator>::with(&memRes);
             auto strategy = rl4::UniqueInstance<CalculationStrategy, rl4::JointGroupQuaternionCalculationStrategy<T> >::with(
@@ -150,7 +152,7 @@ TYPED_TEST(QuaternionEvaluatorTest, EvaluatePerLOD) {
                                                                     0,
                                                                     0,
                                                                     0);
-    rl4::Vector<float> initialValues(rltests::qs::input::values.size());
+    rl4::Vector<rl4::ControlInitializer> initialValues;
     auto inputInstance = inputInstanceFactory(initialValues, &this->memRes);
     auto inputBuffer = inputInstance->getInputBuffer();
     std::copy(rltests::qs::input::values.begin(), rltests::qs::input::values.end(), inputBuffer.begin());
diff --git a/Engine/Plugins/Animation/RigLogic/Source/RigLogicLibTest/Private/rltests/joints/quaternions/TestQuaternionStorageBuilder.cpp b/Engine/Plugins/Animation/RigLogic/Source/RigLogicLibTest/Private/rltests/joints/quaternions/TestQuaternionStorageBuilder.cpp
--- a/Engine/Plugins/Animation/RigLogic/Source/RigLogicLibTest/Private/rltests/joints/quaternions/TestQuaternionStorageBuilder.cpp
+++ b/Engine/Plugins/Animation/RigLogic/Source/RigLogicLibTest/Private/rltests/joints/quaternions/TestQuaternionStorageBuilder.cpp
@@ -36,9 +36,11 @@ struct QuaternionJointsEvaluator<T>::Accessor {
             for (std::size_t lod = {}; lod < lods[jgi].size(); ++lod) {
                 const auto& resultLOD = jointGroup.lods[lod];
                 const auto& expectedLOD = lods[jgi][lod];
-                ASSERT_EQ(resultLOD.size, expectedLOD.size);
-                ASSERT_EQ(resultLOD.sizePaddedToLastFullBlock, expectedLOD.sizePaddedToLastFullBlock);
-                ASSERT_EQ(resultLOD.sizePaddedToSecondLastFullBlock, expectedLOD.sizePaddedToSecondLastFullBlock);
+                ASSERT_EQ(resultLOD.inputLODs.size, expectedLOD.inputLODs.size);
+                ASSERT_EQ(resultLOD.outputLODs.size, expectedLOD.outputLODs.size);
+                ASSERT_EQ(resultLOD.outputLODs.sizePaddedToLastFullBlock, expectedLOD.outputLODs.sizePaddedToLastFullBlock);
+                ASSERT_EQ(resultLOD.outputLODs.sizePaddedToSecondLastFullBlock,
+                          expectedLOD.outputLODs.sizePaddedToSecondLastFullBlock);
             }
         }
     }
diff --git a/Engine/Plugins/Animation/RigLogic/Source/RigLogicLibTest/Private/rltests/joints/twistswing/TestTwistSwingEvaluator.cpp b/Engine/Plugins/Animation/RigLogic/Source/RigLogicLibTest/Private/rltests/joints/twistswing/TestTwistSwingEvaluator.cpp
--- a/Engine/Plugins/Animation/RigLogic/Source/RigLogicLibTest/Private/rltests/joints/twistswing/TestTwistSwingEvaluator.cpp
+++ b/Engine/Plugins/Animation/RigLogic/Source/RigLogicLibTest/Private/rltests/joints/twistswing/TestTwistSwingEvaluator.cpp
@@ -143,7 +143,7 @@ TYPED_TEST(TwistSwingEvaluatorTest, EvaluatePerLOD) {
                                                                     0,
                                                                     0,
                                                                     0);
-    rl4::Vector<float> initialValues(rltests::twsw::input::values.size());
+    rl4::Vector<rl4::ControlInitializer> initialValues;
     auto inputInstance = inputInstanceFactory(initialValues, &this->memRes);
     auto inputBuffer = inputInstance->getInputBuffer();
     std::copy(rltests::twsw::input::values.begin(), rltests::twsw::input::values.end(), inputBuffer.begin());
diff --git a/Engine/Plugins/Animation/RigLogic/Source/RigLogicLibTest/Private/rltests/joints/utils/TestLODRegion.cpp b/Engine/Plugins/Animation/RigLogic/Source/RigLogicLibTest/Private/rltests/joints/utils/TestLODRegion.cpp
new file mode 100644
--- /dev/null
+++ b/Engine/Plugins/Animation/RigLogic/Source/RigLogicLibTest/Private/rltests/joints/utils/TestLODRegion.cpp
@@ -0,0 +1,116 @@
+// Copyright Epic Games, Inc. All Rights Reserved.
+
+#include "rltests/Defs.h"
+
+#include "riglogic/joints/cpu/utils/LODRegion.h"
+
+#include <cstdint>
+
+namespace {
+
+struct LODTestSetup {
+    std::uint32_t rowCount;
+    std::uint32_t lodEndRow;
+    std::uint32_t blockHeight;
+    std::uint32_t padTo;
+    std::uint32_t lodEndRowPaddedToLastFullBlock;
+    std::uint32_t lodEndRowPaddedToSecondLastFullBlock;
+};
+
+class LODRegionTest : public ::testing::TestWithParam<LODTestSetup> {
+};
+
+}  // namespace
+
+TEST(LODRegionTest, ConstructFromFixtures) {
+    rl4::LODRegion lod{{4, 5, 6}, {1, 2, 3}};
+
+    ASSERT_EQ(lod.inputLODs.size, 4u);
+    ASSERT_EQ(lod.inputLODs.sizeAlignedTo4, 5u);
+    ASSERT_EQ(lod.inputLODs.sizeAlignedTo8, 6u);
+
+    ASSERT_EQ(lod.outputLODs.size, 1u);
+    ASSERT_EQ(lod.outputLODs.sizePaddedToLastFullBlock, 2u);
+    ASSERT_EQ(lod.outputLODs.sizePaddedToSecondLastFullBlock, 3u);
+}
+
+TEST_P(LODRegionTest, ConstructLODRegionFromVaryingParameters) {
+    auto params = GetParam();
+    rl4::LODRegion lod{{0, 0, 0}, {params.lodEndRow, params.rowCount, params.blockHeight, params.padTo}};
+    ASSERT_EQ(lod.outputLODs.size, params.lodEndRow);
+    ASSERT_EQ(lod.outputLODs.sizePaddedToLastFullBlock, params.lodEndRowPaddedToLastFullBlock);
+    ASSERT_EQ(lod.outputLODs.sizePaddedToSecondLastFullBlock, params.lodEndRowPaddedToSecondLastFullBlock);
+}
+
+INSTANTIATE_TEST_SUITE_P(
+    LODTestSuite,
+    LODRegionTest,
+    ::testing::Values(
+        // Matrices passed to LODRegion have already been padded,
+        // so the rowcount must be an integral multiple of 4
+        // {{rows, cols}, lodEndRow, lodEndRowAlignedToLastFullBlock, lodEndRowAlignedToSecondLastFullBlock}
+        // SSE test cases
+        LODTestSetup{4, 2, 8, 4, 0, 0},  // Mask-off last 2 rows from 1st block-4 (handles block-4 loop)
+        LODTestSetup{4, 4, 8, 4, 0, 0},  // No mask-off (handles block-4 loop)
+        LODTestSetup{8, 2, 8, 4, 8, 0},  // Mask-off last 6 rows from 1st block-8 (handles block-8 loop)
+        LODTestSetup{8, 4, 8, 4, 8, 0},  // Mask-off last 4 rows from 1st block-8 (handles block-8 loop)
+        LODTestSetup{8, 6, 8, 4, 8, 0},  // Mask-off last 2 rows from 1st block-8 (handles block-8 loop)
+        LODTestSetup{8, 8, 8, 4, 8, 8},  // No mask-off (handles block-8 loop without mask-off logic)
+        LODTestSetup{12, 2, 8, 4, 8, 0},  // Mask-off last 6 rows from 1st block-8 (handles block-8 loop)
+        LODTestSetup{12, 4, 8, 4, 8, 0},  // Mask-off last 4 rows from 1st block-8 (handles block-8 loop)
+        LODTestSetup{12, 6, 8, 4, 8, 0},  // Mask-off last 2 rows from 1st block-8 (handles block-8 loop)
+        LODTestSetup{12, 8, 8, 4, 8, 8},  // No mask-off (handles block-8 loop without mask-off logic)
+        LODTestSetup{12, 10, 8, 4, 8, 8},  // Mask-off last 2 rows from block-4 (handles block-8 without mask-off, then block-4
+                                           // loop)
+        LODTestSetup{12, 12, 8, 4, 8, 8},  // No mask-off (handles block-8 without mask-off, then block-4 loop)
+        LODTestSetup{32, 2, 8, 4, 8, 0},  // Mask-off last 6 rows from 1st block-8 (handles block-8 loop)
+        LODTestSetup{32, 4, 8, 4, 8, 0},  // Mask-off last 4 rows from 1st block-8 (handles block-8 loop)
+        LODTestSetup{32, 6, 8, 4, 8, 0},  // Mask-off last 2 rows from 1st block-8 (handles block-8 loop)
+        LODTestSetup{32, 8, 8, 4, 8, 8},  // No mask-off (handles block-8 loop without mask-off logic)
+        LODTestSetup{32, 9, 8, 4, 16, 8},  // Mask-off last 7 rows from 2nd block-8 (handles block-8 loop without mask-off, then
+                                           // block-8 loop)
+        LODTestSetup{32, 12, 8, 4, 16, 8},  // Mask-off last 4 rows from 2nd block-8 (handles block-8 loop without mask-off, then
+                                            // block-8 loop)
+        LODTestSetup{32, 16, 8, 4, 16, 16},  // No mask-off (handles block-8 loop without mask-off logic)
+        // AVX test cases (pads to 8 so (rowCount % 4 == 0) is not possible)
+        LODTestSetup{8, 2, 16, 8, 0, 0},
+        LODTestSetup{8, 4, 16, 8, 0, 0},
+        LODTestSetup{8, 6, 16, 8, 0, 0},
+        LODTestSetup{8, 8, 16, 8, 0, 0},
+        LODTestSetup{16, 2, 16, 8, 16, 0},
+        LODTestSetup{16, 4, 16, 8, 16, 0},
+        LODTestSetup{16, 6, 16, 8, 16, 0},
+        LODTestSetup{16, 8, 16, 8, 16, 0},
+        LODTestSetup{16, 10, 16, 8, 16, 0},
+        LODTestSetup{16, 12, 16, 8, 16, 0},
+        LODTestSetup{16, 14, 16, 8, 16, 0},
+        LODTestSetup{16, 16, 16, 8, 16, 16},
+        LODTestSetup{24, 2, 16, 8, 16, 0},
+        LODTestSetup{24, 4, 16, 8, 16, 0},
+        LODTestSetup{24, 6, 16, 8, 16, 0},
+        LODTestSetup{24, 8, 16, 8, 16, 0},
+        LODTestSetup{24, 10, 16, 8, 16, 0},
+        LODTestSetup{24, 12, 16, 8, 16, 0},
+        LODTestSetup{24, 14, 16, 8, 16, 0},
+        LODTestSetup{24, 16, 16, 8, 16, 16},
+        LODTestSetup{24, 18, 16, 8, 16, 16},
+        LODTestSetup{24, 20, 16, 8, 16, 16},
+        LODTestSetup{24, 22, 16, 8, 16, 16},
+        LODTestSetup{24, 24, 16, 8, 16, 16},
+        LODTestSetup{32, 2, 16, 8, 16, 0},
+        LODTestSetup{32, 4, 16, 8, 16, 0},
+        LODTestSetup{32, 6, 16, 8, 16, 0},
+        LODTestSetup{32, 8, 16, 8, 16, 0},
+        LODTestSetup{32, 10, 16, 8, 16, 0},
+        LODTestSetup{32, 12, 16, 8, 16, 0},
+        LODTestSetup{32, 14, 16, 8, 16, 0},
+        LODTestSetup{32, 16, 16, 8, 16, 16},
+        LODTestSetup{32, 18, 16, 8, 32, 16},
+        LODTestSetup{32, 20, 16, 8, 32, 16},
+        LODTestSetup{32, 22, 16, 8, 32, 16},
+        LODTestSetup{32, 24, 16, 8, 32, 16},
+        LODTestSetup{32, 26, 16, 8, 32, 16},
+        LODTestSetup{32, 28, 16, 8, 32, 16},
+        LODTestSetup{32, 30, 16, 8, 32, 16},
+        LODTestSetup{32, 32, 16, 8, 32, 32}
+        ));
diff --git a/Engine/Plugins/Animation/RigLogic/Source/RigLogicLibTest/Private/rltests/ml/cpu/TestInference.cpp b/Engine/Plugins/Animation/RigLogic/Source/RigLogicLibTest/Private/rltests/ml/cpu/TestInference.cpp
--- a/Engine/Plugins/Animation/RigLogic/Source/RigLogicLibTest/Private/rltests/ml/cpu/TestInference.cpp
+++ b/Engine/Plugins/Animation/RigLogic/Source/RigLogicLibTest/Private/rltests/ml/cpu/TestInference.cpp
@@ -96,8 +96,7 @@ TYPED_TEST(MLBInferenceTest, InferencePerLOD) {
                                                                     0,
                                                                     rltests::ml::block4::unoptimized::mlControlCount,
                                                                     0);
-    rl4::Vector<float> initialValues(static_cast<std::size_t>(rltests::ml::block4::unoptimized::rawControlCount) +
-                                     static_cast<std::size_t>(rltests::ml::block4::unoptimized::mlControlCount));
+    rl4::Vector<rl4::ControlInitializer> initialValues;
     auto inputInstance = inputInstanceFactory(initialValues, &this->memRes);
     auto inputBuffer = inputInstance->getInputBuffer();
     auto outputBuffer = inputBuffer.subview(rltests::ml::block4::unoptimized::rawControlCount,
diff --git a/Engine/Plugins/Animation/RigLogic/Source/RigLogicLibTest/Private/rltests/rbf/cpu/TestRBFEvaluator.cpp b/Engine/Plugins/Animation/RigLogic/Source/RigLogicLibTest/Private/rltests/rbf/cpu/TestRBFEvaluator.cpp
--- a/Engine/Plugins/Animation/RigLogic/Source/RigLogicLibTest/Private/rltests/rbf/cpu/TestRBFEvaluator.cpp
+++ b/Engine/Plugins/Animation/RigLogic/Source/RigLogicLibTest/Private/rltests/rbf/cpu/TestRBFEvaluator.cpp
@@ -94,8 +94,7 @@ TYPED_TEST(RBFBehaviorTest, SolverPerLOD) {
                                                                     0,
                                                                     0,
                                                                     rltests::rbf::unoptimized::poseControlCount);
-    rl4::Vector<float> initialValues(static_cast<std::size_t>(rltests::rbf::unoptimized::rawControlCount) +
-                                     static_cast<std::size_t>(rltests::rbf::unoptimized::poseControlCount));
+    rl4::Vector<rl4::ControlInitializer> initialValues;
     auto inputInstance = inputInstanceFactory(initialValues, &this->memRes);
     auto inputBuffer = inputInstance->getInputBuffer();
     auto outputBuffer =
