Optimize RigLogic for low LOD evaluation

diff --git a/Engine/Plugins/Animation/RigLogic/Source/RigLogicModule/Private/AnimNode_RigLogic.cpp b/Engine/Plugins/Animation/RigLogic/Source/RigLogicModule/Private/AnimNode_RigLogic.cpp
--- a/Engine/Plugins/Animation/RigLogic/Source/RigLogicModule/Private/AnimNode_RigLogic.cpp
+++ b/Engine/Plugins/Animation/RigLogic/Source/RigLogicModule/Private/AnimNode_RigLogic.cpp
@@ -14,8 +14,6 @@
 #include "HAL/LowLevelMemTracker.h"
 #include "Animation/MorphTarget.h"
 
-#include <tdm/TDM.h>
-
 #include UE_INLINE_GENERATED_CPP_BY_NAME(AnimNode_RigLogic)
 
 LLM_DEFINE_TAG(Animation_RigLogic);
@@ -40,15 +38,98 @@ void FAnimNode_RigLogic::Initialize_AnyThread(const FAnimationInitializeContext&
 {
 	DECLARE_SCOPE_HIERARCHICAL_COUNTER_FUNC()
 	QUICK_SCOPE_CYCLE_COUNTER(STAT_AnimNode_RigLogic_Initialize_AnyThread);
+	TRACE_CPUPROFILER_EVENT_SCOPE_STR("FAnimNode_RigLogic::Initialize_AnyThread");
 	LLM_SCOPE_BYNAME(TEXT("Animation/RigLogic"));
 
 	AnimSequence.Initialize(Context);
 }
 
+void FAnimNode_RigLogic::CacheVariableJointAttributes(const FBoneContainer& RequiredBones)
+{
+	DECLARE_SCOPE_HIERARCHICAL_COUNTER_FUNC()
+	QUICK_SCOPE_CYCLE_COUNTER(STAT_AnimNode_RigLogic_CacheVariableJointAttributes);
+	TRACE_CPUPROFILER_EVENT_SCOPE_STR("FAnimNode_RigLogic::CacheVariableJointAttributes");
+	LLM_SCOPE_BYNAME(TEXT("Animation/RigLogic"));
+
+	// Populate mapping of DNA joint indices to CompactPoseBoneIndex entries (used in updating joints with results from RigLogic)
+	const uint16 CurrentLOD = RigInstance->GetLOD();
+	TArrayView<const uint16> VariableJointIndices = LocalRigRuntimeContext->VariableJointIndicesPerLOD[CurrentLOD].Values;
+	FCachedJointMapping& CurrentLODJointMapping = LocalJointMappingsPerLOD[CurrentLOD];
+	auto& JointsMapDNAIndicesToCompactPoseBoneIndices = CurrentLODJointMapping.JointsMapDNAIndicesToCompactPoseBoneIndices;
+	JointsMapDNAIndicesToCompactPoseBoneIndices.Empty();
+	JointsMapDNAIndicesToCompactPoseBoneIndices.Reserve(VariableJointIndices.Num());
+	for (const uint16 JointIndex : VariableJointIndices)
+	{
+		const FMeshPoseBoneIndex MeshPoseBoneIndex = LocalDNAIndexMapping->JointsMapDNAIndicesToMeshPoseBoneIndices[JointIndex];
+		const FCompactPoseBoneIndex CompactPoseBoneIndex = RequiredBones.MakeCompactPoseIndex(MeshPoseBoneIndex);
+		if (CompactPoseBoneIndex != INDEX_NONE)
+		{
+			JointsMapDNAIndicesToCompactPoseBoneIndices.Add({JointIndex, CompactPoseBoneIndex});
+		}
+	}
+}
+
+void FAnimNode_RigLogic::CacheDriverJoints(const FBoneContainer& RequiredBones)
+{
+	DECLARE_SCOPE_HIERARCHICAL_COUNTER_FUNC()
+	QUICK_SCOPE_CYCLE_COUNTER(STAT_AnimNode_RigLogic_CacheDriverJoints);
+	TRACE_CPUPROFILER_EVENT_SCOPE_STR("FAnimNode_RigLogic::CacheDriverJoints");
+	LLM_SCOPE_BYNAME(TEXT("Animation/RigLogic"));
+
+	const uint16 CurrentLOD = RigInstance->GetLOD();
+	FCachedJointMapping& CurrentLODJointMapping = LocalJointMappingsPerLOD[CurrentLOD];
+	auto& SparseDriverJointsToControlAttributesMap = CurrentLODJointMapping.SparseDriverJointsToControlAttributesMap;
+	auto& DenseDriverJointsToControlAttributesMap = CurrentLODJointMapping.DenseDriverJointsToControlAttributesMap;
+	// Populate driver joint to raw control attribute mapping (used to feed RigLogic with inputs from the joint hierarchy)
+	SparseDriverJointsToControlAttributesMap.Empty();
+	DenseDriverJointsToControlAttributesMap.Empty();
+	DenseDriverJointsToControlAttributesMap.Reserve(LocalDNAIndexMapping->DriverJointsToControlAttributesMap.Num());
+	// Sparse mapping will likely remain empty so no reservation happens
+	for (const auto& Mapping : LocalDNAIndexMapping->DriverJointsToControlAttributesMap)
+	{
+		const FCompactPoseBoneIndex CompactPoseBoneIndex = RequiredBones.MakeCompactPoseIndex(Mapping.MeshPoseBoneIndex);
+		if (CompactPoseBoneIndex != INDEX_NONE)
+		{
+			if ((Mapping.RotationX != INDEX_NONE) && (Mapping.RotationY != INDEX_NONE) && (Mapping.RotationZ != INDEX_NONE) && (Mapping.RotationW != INDEX_NONE))
+			{
+				DenseDriverJointsToControlAttributesMap.Add({CompactPoseBoneIndex, Mapping.DNAJointIndex, Mapping.RotationX, Mapping.RotationY, Mapping.RotationZ, Mapping.RotationW});
+			}
+			else
+			{
+				SparseDriverJointsToControlAttributesMap.Add({CompactPoseBoneIndex, Mapping.DNAJointIndex, Mapping.RotationX, Mapping.RotationY, Mapping.RotationZ, Mapping.RotationW});
+			}
+		}
+	}
+}
+
+void FAnimNode_RigLogic::CachePoseCurvesToRigLogicControlsMap(const FPoseContext& InputContext, const FCachedIndexedCurve& IndexedCurves, TArray<int32>& Indices)
+{
+	DECLARE_SCOPE_HIERARCHICAL_COUNTER_FUNC()
+	QUICK_SCOPE_CYCLE_COUNTER(STAT_AnimNode_RigLogic_CachePoseCurvesToRigLogicControlsMap);
+	TRACE_CPUPROFILER_EVENT_SCOPE_STR("FAnimNode_RigLogic::CachePoseCurvesToRigLogicControlsMap");
+	LLM_SCOPE_BYNAME(TEXT("Animation/RigLogic"));
+
+	Indices.Init(INDEX_NONE, InputContext.Curve.Num());
+
+	int32 CurveIndex = 0;
+	InputContext.Curve.ForEachElement([&](const UE::Anim::FCurveElement& InCurveElement)
+	{
+		IndexedCurves.ForEachElement([&](const UE::Anim::FCurveElementIndexed& InControlAttributeCurveElement)
+		{
+			if (InCurveElement.Name == InControlAttributeCurveElement.Name)
+			{
+				Indices[CurveIndex] = InControlAttributeCurveElement.Index;
+			}
+		});
+		++CurveIndex;
+	});
+}
+
 void FAnimNode_RigLogic::CacheBones_AnyThread(const FAnimationCacheBonesContext& Context)
 {
 	DECLARE_SCOPE_HIERARCHICAL_COUNTER_FUNC()
 	QUICK_SCOPE_CYCLE_COUNTER(STAT_AnimNode_RigLogic_CacheBones_AnyThread);
+	TRACE_CPUPROFILER_EVENT_SCOPE_STR("FAnimNode_RigLogic::CacheBones_AnyThread");
 	LLM_SCOPE_BYNAME(TEXT("Animation/RigLogic"));
 
 	AnimSequence.CacheBones(Context);
@@ -98,9 +179,12 @@ void FAnimNode_RigLogic::CacheBones_AnyThread(const FAnimationCacheBonesContext&
 	TSharedPtr<FDNAIndexMapping> SharedDNAIndexMapping = DNAAsset->GetDNAIndexMapping(Skeleton, SkeletalMesh);
 	if (LocalDNAIndexMapping != SharedDNAIndexMapping)
 	{
+		const uint16 LODCount = LocalRigRuntimeContext->RigLogic->GetLODCount();
 		LocalDNAIndexMapping = SharedDNAIndexMapping;
 		LocalJointMappingsPerLOD.Empty();
-		LocalJointMappingsPerLOD.SetNum(LocalRigRuntimeContext->RigLogic->GetLODCount());
+		LocalJointMappingsPerLOD.SetNum(LODCount);
+		PoseCurvesToRigLogicControlsMap.Reset(LODCount);
+		PoseCurvesToRigLogicControlsMap.AddDefaulted(LODCount);
 	}
 
 	// CacheBones is called on LOD switches as well, in which case compact pose bone indices must be remapped
@@ -108,48 +192,19 @@ void FAnimNode_RigLogic::CacheBones_AnyThread(const FAnimationCacheBonesContext&
 	if (RequiredBones.IsValid())
 	{
 		const uint16 CurrentLOD = RigInstance->GetLOD();
-		auto& CurrentLODJointMapping = LocalJointMappingsPerLOD[CurrentLOD];
-
+		FCachedJointMapping& CurrentLODJointMapping = LocalJointMappingsPerLOD[CurrentLOD];
 		// Lazily initialize and cache mappings for each LOD as they are requested
 		const int32 BoneCountForLOD = RequiredBones.GetCompactPoseNumBones();
 		if (CurrentLODJointMapping.BoneCount != BoneCountForLOD)
 		{
-			// Populate mapping of DNA joint indices to CompactPoseBoneIndex entries (used in updating joints with results from RigLogic)
-			const TArray<uint16>& VariableJointIndices = LocalRigRuntimeContext->VariableJointIndicesPerLOD[CurrentLOD].Values;
-			auto& JointsMapDNAIndicesToCompactPoseBoneIndices = CurrentLODJointMapping.JointsMapDNAIndicesToCompactPoseBoneIndices;
-			JointsMapDNAIndicesToCompactPoseBoneIndices.Empty();
-			JointsMapDNAIndicesToCompactPoseBoneIndices.Reserve(VariableJointIndices.Num());
-			for (const uint16 JointIndex : VariableJointIndices)
+			const FRigLogicConfiguration& RigLogicConfig = LocalRigRuntimeContext->RigLogic->GetConfiguration();
+			if (RigLogicConfig.LoadJoints)
 			{
-				const FMeshPoseBoneIndex MeshPoseBoneIndex = LocalDNAIndexMapping->JointsMapDNAIndicesToMeshPoseBoneIndices[JointIndex];
-				const FCompactPoseBoneIndex CompactPoseBoneIndex = RequiredBones.MakeCompactPoseIndex(MeshPoseBoneIndex);
-				if (CompactPoseBoneIndex != INDEX_NONE)
-				{
-					JointsMapDNAIndicesToCompactPoseBoneIndices.Add({JointIndex, CompactPoseBoneIndex});
-				}
+				CacheVariableJointAttributes(RequiredBones);
 			}
-
-			auto& SparseDriverJointsToControlAttributesMap = CurrentLODJointMapping.SparseDriverJointsToControlAttributesMap;
-			auto& DenseDriverJointsToControlAttributesMap = CurrentLODJointMapping.DenseDriverJointsToControlAttributesMap;
-			// Populate driver joint to raw control attribute mapping (used to feed RigLogic with inputs from the joint hierarchy)
-			SparseDriverJointsToControlAttributesMap.Empty();
-			DenseDriverJointsToControlAttributesMap.Empty();
-			DenseDriverJointsToControlAttributesMap.Reserve(LocalDNAIndexMapping->DriverJointsToControlAttributesMap.Num());
-			// Sparse mapping will likely remain empty so no reservation happens
-			for (const auto& Mapping : LocalDNAIndexMapping->DriverJointsToControlAttributesMap)
+			if (RigLogicConfig.LoadTwistSwingBehavior || RigLogicConfig.LoadRBFBehavior)
 			{
-				const FCompactPoseBoneIndex CompactPoseBoneIndex = RequiredBones.MakeCompactPoseIndex(Mapping.MeshPoseBoneIndex);
-				if (CompactPoseBoneIndex != INDEX_NONE)
-				{
-					if ((Mapping.RotationX != INDEX_NONE) && (Mapping.RotationY != INDEX_NONE) && (Mapping.RotationZ != INDEX_NONE) && (Mapping.RotationW != INDEX_NONE))
-					{
-						DenseDriverJointsToControlAttributesMap.Add({CompactPoseBoneIndex, Mapping.DNAJointIndex, Mapping.RotationX, Mapping.RotationY, Mapping.RotationZ, Mapping.RotationW});
-					}
-					else
-					{
-						SparseDriverJointsToControlAttributesMap.Add({CompactPoseBoneIndex, Mapping.DNAJointIndex, Mapping.RotationX, Mapping.RotationY, Mapping.RotationZ, Mapping.RotationW});
-					}
-				}
+				CacheDriverJoints(RequiredBones);
 			}
 			CurrentLODJointMapping.BoneCount = BoneCountForLOD;
 		}
@@ -160,6 +215,7 @@ void FAnimNode_RigLogic::Update_AnyThread(const FAnimationUpdateContext& Context
 {
 	DECLARE_SCOPE_HIERARCHICAL_COUNTER_FUNC()
 	QUICK_SCOPE_CYCLE_COUNTER(STAT_AnimNode_RigLogic_Update_AnyThread);
+	TRACE_CPUPROFILER_EVENT_SCOPE_STR("FAnimNode_RigLogic::Update_AnyThread");
 
 	GetEvaluateGraphExposedInputs().Execute(Context);
 	AnimSequence.Update(Context);
@@ -169,6 +225,7 @@ void FAnimNode_RigLogic::Evaluate_AnyThread(FPoseContext& OutputContext)
 {
 	DECLARE_SCOPE_HIERARCHICAL_COUNTER_FUNC()
 	QUICK_SCOPE_CYCLE_COUNTER(STAT_AnimNode_RigLogic_Evaluate_AnyThread);
+	TRACE_CPUPROFILER_EVENT_SCOPE_STR("FAnimNode_RigLogic::Evaluate_AnyThread");
 
 	AnimSequence.Evaluate(OutputContext);
 
@@ -177,14 +234,32 @@ void FAnimNode_RigLogic::Evaluate_AnyThread(FPoseContext& OutputContext)
 		return;
 	}
 
-	if (IsLODEnabled(OutputContext.AnimInstanceProxy))
+	if (!IsLODEnabled(OutputContext.AnimInstanceProxy))
+	{
+		return;
+	}
+
+	const FRigLogicConfiguration& RigLogicConfig = LocalRigRuntimeContext->RigLogic->GetConfiguration();
+
+	UpdateControlCurves(OutputContext);
+	CalculateRigLogic();
+
+	if (RigLogicConfig.LoadJoints)
+	{
+		UpdateJoints(OutputContext);
+	}
+	if (RigLogicConfig.LoadBlendShapes)
 	{
-		UpdateControlCurves(OutputContext, LocalDNAIndexMapping.Get(), LocalRigRuntimeContext->RigLogic->GetNeutralJointValues());
-		CalculateRigLogic(LocalRigRuntimeContext->RigLogic.Get());
-		UpdateJoints(LocalRigRuntimeContext->RigLogic->GetNeutralJointValues(), RigInstance->GetJointOutputs(), OutputContext);
-		UpdateBlendShapeCurves(LocalDNAIndexMapping.Get(), RigInstance->GetBlendShapeOutputs(), OutputContext);
-		UpdateAnimMapCurves(LocalDNAIndexMapping.Get(), RigInstance->GetAnimatedMapOutputs(), OutputContext);
+		UpdateBlendShapeCurves(OutputContext);
 	}
+	if (RigLogicConfig.LoadAnimatedMaps)
+	{
+		UpdateAnimMapCurves(OutputContext);
+	}
+
+#if STATS
+	LocalRigRuntimeContext->RigLogic->CollectCalculationStats(RigInstance);
+#endif  // STATS
 }
 
 void FAnimNode_RigLogic::GatherDebugData(FNodeDebugData& DebugData)
@@ -192,101 +267,219 @@ void FAnimNode_RigLogic::GatherDebugData(FNodeDebugData& DebugData)
 	AnimSequence.GatherDebugData(DebugData);
 }
 
-void FAnimNode_RigLogic::UpdateControlCurves(const FPoseContext& InputContext, const FDNAIndexMapping* DNAIndexMapping, TArrayView<const float> NeutralJointValues)
+void FAnimNode_RigLogic::UpdateRawControls(const FPoseContext& InputContext)
 {
 	DECLARE_SCOPE_HIERARCHICAL_COUNTER_FUNC()
-	QUICK_SCOPE_CYCLE_COUNTER(STAT_AnimNode_RigLogic_UpdateControlCurves);
-
-	const uint16 LOD = RigInstance->GetLOD();
+	QUICK_SCOPE_CYCLE_COUNTER(STAT_AnimNode_RigLogic_UpdateRawControls);
+	TRACE_CPUPROFILER_EVENT_SCOPE_STR("FAnimNode_RigLogic::UpdateRawControls");
 
 	// Combine control attribute curve with input curve to get indexed curve to apply to rig
 	// Curve elements that dont have a control mapping will have INDEX_NONE as their index
-	UE::Anim::FNamedValueArrayUtils::Union(InputContext.Curve, DNAIndexMapping->ControlAttributeCurves,
+	UE::Anim::FNamedValueArrayUtils::Union(InputContext.Curve, LocalDNAIndexMapping->ControlAttributeCurves,
 		[this](const UE::Anim::FCurveElement& InCurveElement, const UE::Anim::FCurveElementIndexed& InControlAttributeCurveElement, UE::Anim::ENamedValueUnionFlags InFlags)
 		{
 			if (InControlAttributeCurveElement.Index != INDEX_NONE)
- 			{
+			{
 				RigInstance->SetRawControl(InControlAttributeCurveElement.Index, FMath::Clamp(InCurveElement.Value, 0.0, 1.0));
 			}
 		});
+}
+
+void FAnimNode_RigLogic::UpdateRawControlsCached(const FPoseContext& InputContext)
+{
+	DECLARE_SCOPE_HIERARCHICAL_COUNTER_FUNC()
+	QUICK_SCOPE_CYCLE_COUNTER(STAT_AnimNode_RigLogic_UpdateRawControlsCached);
+	TRACE_CPUPROFILER_EVENT_SCOPE_STR("FAnimNode_RigLogic::UpdateRawControlsCached");
+
+	const uint16 CurrentLOD = RigInstance->GetLOD();
+	auto& RawControlIndices = PoseCurvesToRigLogicControlsMap[CurrentLOD].RawControlIndices;
+	if (RawControlIndices.Num() != InputContext.Curve.Num())
+	{
+		CachePoseCurvesToRigLogicControlsMap(InputContext, LocalDNAIndexMapping->ControlAttributeCurves, RawControlIndices);
+	}
+
+	int32 CurveIndex = 0;
+	InputContext.Curve.ForEachElement([&](const UE::Anim::FCurveElement& InCurveElement)
+		{
+			const int32 ControlIndex = RawControlIndices[CurveIndex];
+			if (ControlIndex != INDEX_NONE)
+			{
+				RigInstance->SetRawControl(ControlIndex, FMath::Clamp(InCurveElement.Value, 0.0, 1.0));
+			}
+			++CurveIndex;
+		});
+}
+
+void FAnimNode_RigLogic::UpdateSparseDriverJointDrivenControlCurves(const FPoseContext& InputContext)
+{
+	DECLARE_SCOPE_HIERARCHICAL_COUNTER_FUNC()
+	QUICK_SCOPE_CYCLE_COUNTER(STAT_AnimNode_RigLogic_UpdateSparseDriverJointDrivenControlCurves);
+	TRACE_CPUPROFILER_EVENT_SCOPE_STR("FAnimNode_RigLogic::UpdateSparseDriverJointDrivenControlCurves");
 
+	const uint16 CurrentLOD = RigInstance->GetLOD();
+	const auto& SparseDriverJointsToControlAttributesMap = LocalJointMappingsPerLOD[CurrentLOD].SparseDriverJointsToControlAttributesMap;
+	TArrayView<const FQuat> InverseNeutralJointRotations = LocalRigRuntimeContext->InverseNeutralJointRotations;
 	// The sparse mapping is NOT guaranteed to supply all quaternion attributes, so checks for each attribute mapping are present
-	const auto& SparseDriverJointsToControlAttributesMap = LocalJointMappingsPerLOD[LOD].SparseDriverJointsToControlAttributesMap;
 	for (int32 MappingIndex = 0; MappingIndex < SparseDriverJointsToControlAttributesMap.Num(); ++MappingIndex)
 	{
 		const FCompactPoseBoneControlAttributeMapping& Mapping = SparseDriverJointsToControlAttributesMap[MappingIndex];
 		const FTransform& CompactPose = InputContext.Pose[Mapping.CompactPoseBoneIndex];
 		// Translation and Scale is currently not used here, so to avoid the overhead of checking them, they are simply ignored.
 		// Should the need arise to use them as well, this code will need adjustment.
-		const FQuat Rotation = CompactPose.GetRotation();
-		const int32 AttrIndex = Mapping.DNAJointIndex * ATTR_COUNT_PER_JOINT;
-		const tdm::fquat NeutralRotation{NeutralJointValues[AttrIndex + 3], NeutralJointValues[AttrIndex + 4], NeutralJointValues[AttrIndex + 5], NeutralJointValues[AttrIndex + 6]};
-		const tdm::fquat AbsPoseRotation{static_cast<float>(Rotation.X), static_cast<float>(Rotation.Y), static_cast<float>(Rotation.Z), static_cast<float>(Rotation.W)};
-		const tdm::fquat DeltaPoseRotation = tdm::inverse(NeutralRotation) * AbsPoseRotation;
+		const FQuat AbsPoseRotation = CompactPose.GetRotation();
+		const FQuat InvNeutralRotation = InverseNeutralJointRotations[Mapping.DNAJointIndex];
+		const FQuat DeltaPoseRotation = InvNeutralRotation * AbsPoseRotation;
 
 		if (Mapping.RotationX != INDEX_NONE)
 		{
-			RigInstance->SetRawControl(Mapping.RotationX, DeltaPoseRotation.x);
+			RigInstance->SetRawControl(Mapping.RotationX, DeltaPoseRotation.X);
 		}
 		if (Mapping.RotationY != INDEX_NONE)
 		{
-			RigInstance->SetRawControl(Mapping.RotationY, DeltaPoseRotation.y);
+			RigInstance->SetRawControl(Mapping.RotationY, DeltaPoseRotation.Y);
 		}
 		if (Mapping.RotationZ != INDEX_NONE)
 		{
-			RigInstance->SetRawControl(Mapping.RotationZ, DeltaPoseRotation.z);
+			RigInstance->SetRawControl(Mapping.RotationZ, DeltaPoseRotation.Z);
 		}
 		if (Mapping.RotationW != INDEX_NONE)
 		{
-			RigInstance->SetRawControl(Mapping.RotationW, DeltaPoseRotation.w);
+			RigInstance->SetRawControl(Mapping.RotationW, DeltaPoseRotation.W);
 		}
 	}
+}
+
+void FAnimNode_RigLogic::UpdateDenseDriverJointDrivenControlCurves(const FPoseContext& InputContext)
+{
+	DECLARE_SCOPE_HIERARCHICAL_COUNTER_FUNC()
+	QUICK_SCOPE_CYCLE_COUNTER(STAT_AnimNode_RigLogic_UpdateDenseDriverJointDrivenControlCurves);
+	TRACE_CPUPROFILER_EVENT_SCOPE_STR("FAnimNode_RigLogic::UpdateDenseDriverJointDrivenControlCurves");
 
+	const uint16 CurrentLOD = RigInstance->GetLOD();
+	const auto& DenseDriverJointsToControlAttributesMap = LocalJointMappingsPerLOD[CurrentLOD].DenseDriverJointsToControlAttributesMap;
+	TArrayView<const FQuat> InverseNeutralJointRotations = LocalRigRuntimeContext->InverseNeutralJointRotations;
 	// The dense mapping is guaranteed to supply all quaternion attributes, so NO checks for each attribute mapping are present
-	const auto& DenseDriverJointsToControlAttributesMap = LocalJointMappingsPerLOD[LOD].DenseDriverJointsToControlAttributesMap;
 	for (int32 MappingIndex = 0; MappingIndex < DenseDriverJointsToControlAttributesMap.Num(); ++MappingIndex)
 	{
 		const FCompactPoseBoneControlAttributeMapping& Mapping = DenseDriverJointsToControlAttributesMap[MappingIndex];
 		const FTransform& CompactPose = InputContext.Pose[Mapping.CompactPoseBoneIndex];
 		// Translation and Scale is currently not used here, so to avoid the overhead of checking them, they are simply ignored.
 		// Should the need arise to use them as well, this code will need adjustment.
-		const FQuat Rotation = CompactPose.GetRotation();
-		const int32 AttrIndex = Mapping.DNAJointIndex * ATTR_COUNT_PER_JOINT;
-		const tdm::fquat NeutralRotation{ NeutralJointValues[AttrIndex + 3], NeutralJointValues[AttrIndex + 4], NeutralJointValues[AttrIndex + 5], NeutralJointValues[AttrIndex + 6] };
-		const tdm::fquat AbsPoseRotation{ static_cast<float>(Rotation.X), static_cast<float>(Rotation.Y), static_cast<float>(Rotation.Z), static_cast<float>(Rotation.W) };
-		const tdm::fquat DeltaPoseRotation = tdm::inverse(NeutralRotation) * AbsPoseRotation;
-		RigInstance->SetRawControl(Mapping.RotationX, DeltaPoseRotation.x);
-		RigInstance->SetRawControl(Mapping.RotationY, DeltaPoseRotation.y);
-		RigInstance->SetRawControl(Mapping.RotationZ, DeltaPoseRotation.z);
-		RigInstance->SetRawControl(Mapping.RotationW, DeltaPoseRotation.w);
+		const FQuat AbsPoseRotation = CompactPose.GetRotation();
+		const FQuat InvNeutralRotation = InverseNeutralJointRotations[Mapping.DNAJointIndex];
+		const FQuat DeltaPoseRotation = InvNeutralRotation * AbsPoseRotation;
+		RigInstance->SetRawControl(Mapping.RotationX, DeltaPoseRotation.X);
+		RigInstance->SetRawControl(Mapping.RotationY, DeltaPoseRotation.Y);
+		RigInstance->SetRawControl(Mapping.RotationZ, DeltaPoseRotation.Z);
+		RigInstance->SetRawControl(Mapping.RotationW, DeltaPoseRotation.W);
 	}
+}
+
+void FAnimNode_RigLogic::UpdateNeuralNetworkMaskCurves(const FPoseContext& InputContext)
+{
+	DECLARE_SCOPE_HIERARCHICAL_COUNTER_FUNC()
+	QUICK_SCOPE_CYCLE_COUNTER(STAT_AnimNode_RigLogic_UpdateNeuralNetworkMaskCurves);
+	TRACE_CPUPROFILER_EVENT_SCOPE_STR("FAnimNode_RigLogic::UpdateNeuralNetworkMaskCurves");
 
 	if (RigInstance->GetNeuralNetworkCount() != 0)
 	{
-		UE::Anim::FNamedValueArrayUtils::Union(InputContext.Curve, DNAIndexMapping->NeuralNetworkMaskCurves,
+		UE::Anim::FNamedValueArrayUtils::Union(InputContext.Curve, LocalDNAIndexMapping->NeuralNetworkMaskCurves,
 			[this](const UE::Anim::FCurveElement& InCurveElement, const UE::Anim::FCurveElementIndexed& InControlAttributeCurveElement, UE::Anim::ENamedValueUnionFlags InFlags)
 			{
 				if (InControlAttributeCurveElement.Index != INDEX_NONE)
- 				{
+				{
 					RigInstance->SetNeuralNetworkMask(InControlAttributeCurveElement.Index, InCurveElement.Value);
- 				}
+				}
+			});
+	}
+}
+
+void FAnimNode_RigLogic::UpdateNeuralNetworkMaskCurvesCached(const FPoseContext& InputContext)
+{
+	DECLARE_SCOPE_HIERARCHICAL_COUNTER_FUNC()
+	QUICK_SCOPE_CYCLE_COUNTER(STAT_AnimNode_RigLogic_UpdateNeuralNetworkMaskCurvesCached);
+	TRACE_CPUPROFILER_EVENT_SCOPE_STR("FAnimNode_RigLogic::UpdateNeuralNetworkMaskCurvesCached");
+
+	if (RigInstance->GetNeuralNetworkCount() != 0)
+	{
+		const uint16 CurrentLOD = RigInstance->GetLOD();
+		auto& NeuralNetworkMaskIndices = PoseCurvesToRigLogicControlsMap[CurrentLOD].NeuralNetworkMaskIndices;
+		if (NeuralNetworkMaskIndices.Num() != InputContext.Curve.Num())
+		{
+			CachePoseCurvesToRigLogicControlsMap(InputContext, LocalDNAIndexMapping->NeuralNetworkMaskCurves, NeuralNetworkMaskIndices);
+		}
+
+		int32 CurveIndex = 0;
+		InputContext.Curve.ForEachElement([&](const UE::Anim::FCurveElement& InCurveElement)
+			{
+				const int32 ControlIndex = NeuralNetworkMaskIndices[CurveIndex];
+				if (ControlIndex != INDEX_NONE)
+				{
+					RigInstance->SetNeuralNetworkMask(ControlIndex, InCurveElement.Value);
+				}
+				++CurveIndex;
 			});
 	}
 }
 
-void FAnimNode_RigLogic::CalculateRigLogic(FRigLogic* RigLogic)
+void FAnimNode_RigLogic::UpdateControlCurves(const FPoseContext& InputContext)
+{
+	DECLARE_SCOPE_HIERARCHICAL_COUNTER_FUNC()
+	QUICK_SCOPE_CYCLE_COUNTER(STAT_AnimNode_RigLogic_UpdateControlCurves);
+	TRACE_CPUPROFILER_EVENT_SCOPE_STR("FAnimNode_RigLogic::UpdateControlCurves");
+
+	const FRigLogicConfiguration& RigLogicConfig = LocalRigRuntimeContext->RigLogic->GetConfiguration();
+
+	if (CacheAnimCurveNames)
+	{
+		UpdateRawControlsCached(InputContext);
+		if (RigLogicConfig.LoadMachineLearnedBehavior)
+		{
+			UpdateNeuralNetworkMaskCurvesCached(InputContext);
+		}
+	}
+	else
+	{
+		UpdateRawControls(InputContext);
+		if (RigLogicConfig.LoadMachineLearnedBehavior)
+		{
+			UpdateNeuralNetworkMaskCurves(InputContext);
+		}
+	}
+
+	if (RigLogicConfig.LoadRBFBehavior || RigLogicConfig.LoadTwistSwingBehavior)
+	{
+		UpdateSparseDriverJointDrivenControlCurves(InputContext);
+		UpdateDenseDriverJointDrivenControlCurves(InputContext);
+	}
+}
+
+void FAnimNode_RigLogic::CalculateRigLogic()
 {
 	DECLARE_SCOPE_HIERARCHICAL_COUNTER_FUNC()
 	QUICK_SCOPE_CYCLE_COUNTER(STAT_AnimNode_RigLogic_CalculateRigLogic);
-	RigLogic->Calculate(RigInstance);
+	TRACE_CPUPROFILER_EVENT_SCOPE_STR("FAnimNode_RigLogic::CalculateRigLogic");
+
+	FRigLogic* RigLogic = LocalRigRuntimeContext->RigLogic.Get();
+	// RigLogic has Null evaluators for each class of computations, so no explicit checks are necessary here
+	// based on the chosen configuration, as no extra work will be performed if not needed.
+	RigLogic->CalculateMachineLearnedBehaviorControls(RigInstance);
+	RigLogic->CalculateRBFControls(RigInstance);
+	RigLogic->CalculateControls(RigInstance);
+	RigLogic->CalculateJoints(RigInstance);
+	RigLogic->CalculateBlendShapes(RigInstance);
+	RigLogic->CalculateAnimatedMaps(RigInstance);
 }
 
-void FAnimNode_RigLogic::UpdateJoints(TArrayView<const float> NeutralJointValues, TArrayView<const float> DeltaJointValues, FPoseContext& OutputContext)
+void FAnimNode_RigLogic::UpdateJoints(FPoseContext& OutputContext)
 {
 	DECLARE_SCOPE_HIERARCHICAL_COUNTER_FUNC()
 	QUICK_SCOPE_CYCLE_COUNTER(STAT_AnimNode_RigLogic_UpdateJoints);
+	TRACE_CPUPROFILER_EVENT_SCOPE_STR("FAnimNode_RigLogic::UpdateJoints");
 
 	const uint16 LOD = RigInstance->GetLOD();
+	TArrayView<const float> DeltaJointValues = RigInstance->GetJointOutputs();
+	TArrayView<const float> NeutralJointValues = LocalRigRuntimeContext->RigLogic->GetNeutralJointValues();
 	const float* N = NeutralJointValues.GetData();
 	const float* D = DeltaJointValues.GetData();
 	const auto& JointsMapDNAIndicesToCompactPoseBoneIndices = LocalJointMappingsPerLOD[LOD].JointsMapDNAIndicesToCompactPoseBoneIndices;
@@ -300,13 +493,15 @@ void FAnimNode_RigLogic::UpdateJoints(TArrayView<const float> NeutralJointValues
 	}
 }
 
-void FAnimNode_RigLogic::UpdateBlendShapeCurves(const FDNAIndexMapping* DNAIndexMapping, TArrayView<const float> BlendShapeValues, FPoseContext& OutputContext)
+void FAnimNode_RigLogic::UpdateBlendShapeCurves(FPoseContext& OutputContext)
 {
 	DECLARE_SCOPE_HIERARCHICAL_COUNTER_FUNC()
 	QUICK_SCOPE_CYCLE_COUNTER(STAT_AnimNode_RigLogic_UpdateBlendShapeCurves);
-	
+	TRACE_CPUPROFILER_EVENT_SCOPE_STR("FAnimNode_RigLogic::UpdateBlendShapeCurves");
+
 	const uint16 LOD = RigInstance->GetLOD();
-	const FDNAIndexMapping::FCachedIndexedCurve& MorphTargetCurve = DNAIndexMapping->MorphTargetCurvesPerLOD[LOD];
+	TArrayView<const float> BlendShapeValues = RigInstance->GetBlendShapeOutputs();
+	const FDNAIndexMapping::FCachedIndexedCurve& MorphTargetCurve = LocalDNAIndexMapping->MorphTargetCurvesPerLOD[LOD];
 	UE::Anim::FNamedValueArrayUtils::Union(OutputContext.Curve, MorphTargetCurve, [&BlendShapeValues](UE::Anim::FCurveElement& InOutResult, const UE::Anim::FCurveElementIndexed& InSource, UE::Anim::ENamedValueUnionFlags InFlags)
 	{
 		if (BlendShapeValues.IsValidIndex(InSource.Index))
@@ -317,13 +512,15 @@ void FAnimNode_RigLogic::UpdateBlendShapeCurves(const FDNAIndexMapping* DNAIndex
 	});
 }
 
-void FAnimNode_RigLogic::UpdateAnimMapCurves(const FDNAIndexMapping* DNAIndexMapping, TArrayView<const float> AnimMapOutputs, FPoseContext& OutputContext)
+void FAnimNode_RigLogic::UpdateAnimMapCurves(FPoseContext& OutputContext)
 {
 	DECLARE_SCOPE_HIERARCHICAL_COUNTER_FUNC()
 	QUICK_SCOPE_CYCLE_COUNTER(STAT_AnimNode_RigLogic_UpdateAnimMapCurves);
+	TRACE_CPUPROFILER_EVENT_SCOPE_STR("FAnimNode_RigLogic::UpdateAnimMapCurves");
 
 	const uint16 LOD = RigInstance->GetLOD();
-	const FDNAIndexMapping::FCachedIndexedCurve& MaskMultiplierCurve = DNAIndexMapping->MaskMultiplierCurvesPerLOD[LOD];
+	TArrayView<const float> AnimMapOutputs = RigInstance->GetAnimatedMapOutputs();
+	const FDNAIndexMapping::FCachedIndexedCurve& MaskMultiplierCurve = LocalDNAIndexMapping->MaskMultiplierCurvesPerLOD[LOD];
 	UE::Anim::FNamedValueArrayUtils::Union(OutputContext.Curve, MaskMultiplierCurve, [&AnimMapOutputs](UE::Anim::FCurveElement& InOutResult, const UE::Anim::FCurveElementIndexed& InSource, UE::Anim::ENamedValueUnionFlags InFlags)
 	{
 		if (AnimMapOutputs.IsValidIndex(InSource.Index))
diff --git a/Engine/Plugins/Animation/RigLogic/Source/RigLogicModule/Private/DNAAsset.cpp b/Engine/Plugins/Animation/RigLogic/Source/RigLogicModule/Private/DNAAsset.cpp
--- a/Engine/Plugins/Animation/RigLogic/Source/RigLogicModule/Private/DNAAsset.cpp
+++ b/Engine/Plugins/Animation/RigLogic/Source/RigLogicModule/Private/DNAAsset.cpp
@@ -190,8 +190,13 @@ void UDNAAsset::InitializeRigRuntimeContext()
 		// Convert behavior data from DNA to UE space on the fly as RigLogic accesses it
 		RigLogicDNAReader BehaviorReaderInUESpace{BehaviorReader->Unwrap()};
 		FDNAReader<RigLogicDNAReader> BehaviorReaderInUESpaceWrapper{&BehaviorReaderInUESpace};
-		NewContext->RigLogic = MakeShared<FRigLogic>(&BehaviorReaderInUESpaceWrapper);
+		NewContext->RigLogic = MakeShared<FRigLogic>(&BehaviorReaderInUESpaceWrapper, RigLogicConfiguration);
 		NewContext->CacheVariableJointIndices();
+		if ((BehaviorReader->GetRBFSolverCount() != 0) || (BehaviorReader->GetSwingCount() != 0) || (BehaviorReader->GetTwistCount() != 0))
+		{
+			NewContext->CacheInverseNeutralJointRotations();
+		}
+
 		{
 			FWriteScopeLock ContextScopeLock{RigRuntimeContextUpdateLock};
 			RigRuntimeContext = NewContext;
@@ -380,3 +385,21 @@ void UDNAAsset::Serialize(FArchive& Ar)
 	}
 }
 
+#if WITH_EDITOR
+void UDNAAsset::PostEditChangeProperty(FPropertyChangedEvent& PropertyChangedEvent)
+{
+	Super::PostEditChangeProperty(PropertyChangedEvent);
+
+	if (PropertyChangedEvent.MemberProperty == nullptr)
+	{
+		return;
+	}
+
+	const FName& MemberPropertyName = PropertyChangedEvent.MemberProperty->GetFName();
+	if (MemberPropertyName == GET_MEMBER_NAME_CHECKED(UDNAAsset, RigLogicConfiguration))
+	{
+		InvalidateRigRuntimeContext();
+		InitializeRigRuntimeContext();
+	}
+}
+#endif
diff --git a/Engine/Plugins/Animation/RigLogic/Source/RigLogicModule/Private/RigLogic.cpp b/Engine/Plugins/Animation/RigLogic/Source/RigLogicModule/Private/RigLogic.cpp
--- a/Engine/Plugins/Animation/RigLogic/Source/RigLogicModule/Private/RigLogic.cpp
+++ b/Engine/Plugins/Animation/RigLogic/Source/RigLogicModule/Private/RigLogic.cpp
@@ -5,15 +5,41 @@
 #include "DNAReader.h"
 #include "FMemoryResource.h"
 #include "RigInstance.h"
+#include "RigLogicModule.h"
+#include "Stats/Stats.h"
+#include "Stats/StatsHierarchical.h"
 
 #include "riglogic/RigLogic.h"
 
+#if STATS
+
+DEFINE_STAT(STAT_RigLogic_CalculationType);
+DEFINE_STAT(STAT_RigLogic_FloatingPointType);
+DEFINE_STAT(STAT_RigLogic_LOD);
+DEFINE_STAT(STAT_RigLogic_RBFSolverCount);
+DEFINE_STAT(STAT_RigLogic_NeuralNetworkCount);
+DEFINE_STAT(STAT_RigLogic_PSDCount);
+DEFINE_STAT(STAT_RigLogic_JointCount);
+DEFINE_STAT(STAT_RigLogic_JointDeltaValueCount);
+DEFINE_STAT(STAT_RigLogic_BlendShapeChannelCount);
+DEFINE_STAT(STAT_RigLogic_AnimatedMapCount);
+DEFINE_STAT(STAT_FRigLogic_MapGUIToRawControlsTime);
+DEFINE_STAT(STAT_FRigLogic_MapRawToGUIControlsTime);
+DEFINE_STAT(STAT_RigLogic_CalculateRBFControlsTime);
+DEFINE_STAT(STAT_RigLogic_CalculateMLControlsTime);
+DEFINE_STAT(STAT_RigLogic_CalculateControlsTime);
+DEFINE_STAT(STAT_RigLogic_CalculateJointsTime);
+DEFINE_STAT(STAT_RigLogic_CalculateBlendShapesTime);
+DEFINE_STAT(STAT_RigLogic_CalculateAnimatedMapsTime);
+
+#endif  // STATS
+
 void FRigLogic::FRigLogicDeleter::operator()(rl4::RigLogic* Pointer)
 {
 	rl4::RigLogic::destroy(Pointer);
 }
 
-static rl4::Configuration AdaptRigLogicConfig(const FRigLogicConfiguration& Config)
+static rl4::Configuration AdaptToRigLogicConfig(const FRigLogicConfiguration& Config)
 {
 	rl4::Configuration Copy = {};
 	Copy.calculationType = static_cast<rl4::CalculationType>(Config.CalculationType);
@@ -27,22 +53,81 @@ static rl4::Configuration AdaptRigLogicConfig(const FRigLogicConfiguration& Conf
 	Copy.rotationType = static_cast<rl4::RotationType>(Config.RotationType);
 	Copy.rotationOrder = static_cast<rl4::RotationOrder>(Config.RotationOrder);
 	Copy.scaleType = static_cast<rl4::ScaleType>(Config.ScaleType);
+	Copy.translationPruningThreshold = Config.TranslationPruningThreshold;
+	Copy.rotationPruningThreshold = Config.RotationPruningThreshold;
+	Copy.scalePruningThreshold = Config.ScalePruningThreshold;
 	return Copy;
 }
 
-FRigLogic::FRigLogic(const IDNAReader* Reader, FRigLogicConfiguration Config) :
+static FRigLogicConfiguration AdaptFromRigLogicConfig(const rl4::Configuration& Config)
+{
+	FRigLogicConfiguration Copy = {};
+	Copy.CalculationType = static_cast<ERigLogicCalculationType>(Config.calculationType);
+	Copy.LoadJoints = Config.loadJoints;
+	Copy.LoadBlendShapes = Config.loadBlendShapes;
+	Copy.LoadAnimatedMaps = Config.loadAnimatedMaps;
+	Copy.LoadMachineLearnedBehavior = Config.loadMachineLearnedBehavior;
+	Copy.LoadRBFBehavior = Config.loadRBFBehavior;
+	Copy.LoadTwistSwingBehavior = Config.loadTwistSwingBehavior;
+	Copy.TranslationType = static_cast<ERigLogicTranslationType>(Config.translationType);
+	Copy.RotationType = static_cast<ERigLogicRotationType>(Config.rotationType);
+	Copy.RotationOrder = static_cast<ERigLogicRotationOrder>(Config.rotationOrder);
+	Copy.ScaleType = static_cast<ERigLogicScaleType>(Config.scaleType);
+	Copy.TranslationPruningThreshold = Config.translationPruningThreshold;
+	Copy.RotationPruningThreshold = Config.rotationPruningThreshold;
+	Copy.ScalePruningThreshold = Config.scalePruningThreshold;
+	return Copy;
+}
+
+FRigLogic::FRigLogic(const IDNAReader* Reader, const FRigLogicConfiguration& Config) :
 	MemoryResource{FMemoryResource::SharedInstance()},
-	RigLogic{rl4::RigLogic::create(Reader->Unwrap(), AdaptRigLogicConfig(Config), FMemoryResource::Instance())}
+	RigLogic{rl4::RigLogic::create(Reader->Unwrap(), AdaptToRigLogicConfig(Config), FMemoryResource::Instance())},
+	Configuration{AdaptFromRigLogicConfig(RigLogic->getConfiguration())}
 {
 }
 
 FRigLogic::~FRigLogic() = default;
 
+const FRigLogicConfiguration& FRigLogic::GetConfiguration() const
+{
+	return Configuration;
+}
+
 uint16 FRigLogic::GetLODCount() const
 {
 	return RigLogic->getLODCount();
 }
 
+TArrayView<const uint16> FRigLogic::GetRBFSolverIndicesForLOD(uint16_t LOD) const
+{
+	rl4::ConstArrayView<uint16> Indices = RigLogic->getRBFSolverIndicesForLOD(LOD);
+	return TArrayView<const uint16>{Indices.data(), static_cast<int32>(Indices.size())};
+}
+
+TArrayView<const uint32> FRigLogic::GetNeuralNetworkIndicesForLOD(uint16_t LOD) const
+{
+	rl4::ConstArrayView<uint32> Indices = RigLogic->getNeuralNetworkIndicesForLOD(LOD);
+	return TArrayView<const uint32>{Indices.data(), static_cast<int32>(Indices.size())};
+}
+
+TArrayView<const uint16> FRigLogic::GetBlendShapeChannelIndicesForLOD(uint16_t LOD) const
+{
+	rl4::ConstArrayView<uint16> Indices = RigLogic->getBlendShapeChannelIndicesForLOD(LOD);
+	return TArrayView<const uint16>{Indices.data(), static_cast<int32>(Indices.size())};
+}
+
+TArrayView<const uint16> FRigLogic::GetAnimatedMapIndicesForLOD(uint16_t LOD) const
+{
+	rl4::ConstArrayView<uint16> Indices = RigLogic->getAnimatedMapIndicesForLOD(LOD);
+	return TArrayView<const uint16>{Indices.data(), static_cast<int32>(Indices.size())};
+}
+
+TArrayView<const uint16> FRigLogic::GetJointIndicesForLOD(uint16_t LOD) const
+{
+	rl4::ConstArrayView<uint16> Indices = RigLogic->getJointIndicesForLOD(LOD);
+	return TArrayView<const uint16>{Indices.data(), static_cast<int32>(Indices.size())};
+}
+
 TArrayView<const float> FRigLogic::GetNeutralJointValues() const
 {
 	rl4::ConstArrayView<float> Values = RigLogic->getNeutralJointValues();
@@ -88,64 +173,130 @@ TArrayView<const uint16> FRigLogic::GetNeuralNetworkIndices(uint16 MeshIndex, ui
 
 void FRigLogic::MapGUIToRawControls(FRigInstance* Instance) const
 {
+	DECLARE_SCOPE_HIERARCHICAL_COUNTER_FUNC()
+	SCOPE_CYCLE_COUNTER(STAT_FRigLogic_MapGUIToRawControlsTime);
+	TRACE_CPUPROFILER_EVENT_SCOPE_STR("FRigLogic::MapGUIToRawControls");
+
 	RigLogic->mapGUIToRawControls(Instance->Unwrap());
 }
 
 void FRigLogic::MapRawToGUIControls(FRigInstance* Instance) const
 {
+	DECLARE_SCOPE_HIERARCHICAL_COUNTER_FUNC()
+	SCOPE_CYCLE_COUNTER(STAT_FRigLogic_MapRawToGUIControlsTime);
+	TRACE_CPUPROFILER_EVENT_SCOPE_STR("FRigLogic::MapRawToGUIControls");
+
 	RigLogic->mapRawToGUIControls(Instance->Unwrap());
 }
 
 void FRigLogic::CalculateControls(FRigInstance* Instance) const
 {
+	DECLARE_SCOPE_HIERARCHICAL_COUNTER_FUNC()
+	SCOPE_CYCLE_COUNTER(STAT_RigLogic_CalculateControlsTime);
+	TRACE_CPUPROFILER_EVENT_SCOPE_STR("FRigLogic::CalculateControls");
+
 	RigLogic->calculateControls(Instance->Unwrap());
 }
 
 void FRigLogic::CalculateMachineLearnedBehaviorControls(FRigInstance* Instance) const
 {
+	DECLARE_SCOPE_HIERARCHICAL_COUNTER_FUNC()
+	SCOPE_CYCLE_COUNTER(STAT_RigLogic_CalculateMLControlsTime);
+	TRACE_CPUPROFILER_EVENT_SCOPE_STR("FRigLogic::CalculateMachineLearnedBehaviorControls");
+
 	RigLogic->calculateMachineLearnedBehaviorControls(Instance->Unwrap());
 }
 
 void FRigLogic::CalculateMachineLearnedBehaviorControls(FRigInstance* Instance, uint16 NeuralNetIndex) const
 {
+	DECLARE_SCOPE_HIERARCHICAL_COUNTER_FUNC()
+	SCOPE_CYCLE_COUNTER(STAT_RigLogic_CalculateMLControlsTime);
+	TRACE_CPUPROFILER_EVENT_SCOPE_STR("FRigLogic::CalculateMachineLearnedBehaviorControls");
+
 	RigLogic->calculateMachineLearnedBehaviorControls(Instance->Unwrap(), NeuralNetIndex);
 }
 
 void FRigLogic::CalculateRBFControls(FRigInstance* Instance) const
 {
+	DECLARE_SCOPE_HIERARCHICAL_COUNTER_FUNC()
+	SCOPE_CYCLE_COUNTER(STAT_RigLogic_CalculateRBFControlsTime);
+	TRACE_CPUPROFILER_EVENT_SCOPE_STR("FRigLogic::CalculateRBFControls");
+
 	RigLogic->calculateRBFControls(Instance->Unwrap());
 }
 
 void FRigLogic::CalculateRBFControls(FRigInstance* Instance, uint16 SolverIndex) const
 {
+	DECLARE_SCOPE_HIERARCHICAL_COUNTER_FUNC()
+	SCOPE_CYCLE_COUNTER(STAT_RigLogic_CalculateRBFControlsTime);
+	TRACE_CPUPROFILER_EVENT_SCOPE_STR("FRigLogic::CalculateRBFControls");
+
 	RigLogic->calculateRBFControls(Instance->Unwrap(), SolverIndex);
 }
 
 void FRigLogic::CalculateJoints(FRigInstance* Instance) const
 {
+	DECLARE_SCOPE_HIERARCHICAL_COUNTER_FUNC()
+	SCOPE_CYCLE_COUNTER(STAT_RigLogic_CalculateJointsTime);
+	TRACE_CPUPROFILER_EVENT_SCOPE_STR("FRigLogic::CalculateJoints");
+
 	RigLogic->calculateJoints(Instance->Unwrap());
 }
 
 void FRigLogic::CalculateJoints(FRigInstance* Instance, uint16 JointGroupIndex) const
 {
+	DECLARE_SCOPE_HIERARCHICAL_COUNTER_FUNC()
+	SCOPE_CYCLE_COUNTER(STAT_RigLogic_CalculateJointsTime);
+	TRACE_CPUPROFILER_EVENT_SCOPE_STR("FRigLogic::CalculateJoints");
+
 	RigLogic->calculateJoints(Instance->Unwrap(), JointGroupIndex);
 }
 
 void FRigLogic::CalculateBlendShapes(FRigInstance* Instance) const
 {
+	DECLARE_SCOPE_HIERARCHICAL_COUNTER_FUNC()
+	SCOPE_CYCLE_COUNTER(STAT_RigLogic_CalculateBlendShapesTime);
+	TRACE_CPUPROFILER_EVENT_SCOPE_STR("FRigLogic::CalculateBlendShapes");
+
 	RigLogic->calculateBlendShapes(Instance->Unwrap());
 }
 
 void FRigLogic::CalculateAnimatedMaps(FRigInstance* Instance) const
 {
+	DECLARE_SCOPE_HIERARCHICAL_COUNTER_FUNC()
+	SCOPE_CYCLE_COUNTER(STAT_RigLogic_CalculateAnimatedMapsTime);
+	TRACE_CPUPROFILER_EVENT_SCOPE_STR("FRigLogic::CalculateAnimatedMaps");
+
 	RigLogic->calculateAnimatedMaps(Instance->Unwrap());
 }
 
 void FRigLogic::Calculate(FRigInstance* Instance) const
 {
+	DECLARE_SCOPE_HIERARCHICAL_COUNTER_FUNC()
+	QUICK_SCOPE_CYCLE_COUNTER(STAT_FRigLogic_Calculate);
+	TRACE_CPUPROFILER_EVENT_SCOPE_STR("FRigLogic::Calculate");
+
 	RigLogic->calculate(Instance->Unwrap());
 }
 
+void FRigLogic::CollectCalculationStats(FRigInstance* Instance) const
+{
+#if STATS
+	rl4::Stats Stats = {};
+	RigLogic->collectCalculationStats(Instance->Unwrap(), &Stats);
+	SET_DWORD_STAT(STAT_RigLogic_CalculationType, Stats.calculationType);
+	SET_DWORD_STAT(STAT_RigLogic_FloatingPointType, Stats.floatingPointType);
+	SET_DWORD_STAT(STAT_RigLogic_LOD, Instance->GetLOD());
+	SET_DWORD_STAT(STAT_RigLogic_RBFSolverCount, Stats.rbfSolverCount);
+	SET_DWORD_STAT(STAT_RigLogic_NeuralNetworkCount, Stats.neuralNetworkCount);
+	SET_DWORD_STAT(STAT_RigLogic_PSDCount, Stats.psdCount);
+	SET_DWORD_STAT(STAT_RigLogic_JointCount, Stats.jointCount);
+	SET_DWORD_STAT(STAT_RigLogic_JointDeltaValueCount, Stats.jointDeltaValueCount);
+	SET_DWORD_STAT(STAT_RigLogic_BlendShapeChannelCount, Stats.blendShapeChannelCount);
+	SET_DWORD_STAT(STAT_RigLogic_AnimatedMapCount, Stats.animatedMapCount);
+#endif  // STATS
+}
+
 rl4::RigLogic* FRigLogic::Unwrap() const
 {
 	return RigLogic.Get();
diff --git a/Engine/Plugins/Animation/RigLogic/Source/RigLogicModule/Private/RigUnit_RigLogic.cpp b/Engine/Plugins/Animation/RigLogic/Source/RigLogicModule/Private/RigUnit_RigLogic.cpp
--- a/Engine/Plugins/Animation/RigLogic/Source/RigLogicModule/Private/RigUnit_RigLogic.cpp
+++ b/Engine/Plugins/Animation/RigLogic/Source/RigLogicModule/Private/RigUnit_RigLogic.cpp
@@ -107,11 +107,31 @@ void FRigUnit_RigLogic_Data::InitializeRigLogic(const URigHierarchy* InHierarchy
 		RigInstance->SetLOD(CurrentLOD);
 		CurrentLOD = RigInstance->GetLOD();
 
-		MapJoints(InHierarchy);
+		const FRigLogicConfiguration& RigLogicConfig = LocalRigRuntimeContext->RigLogic->GetConfiguration();
+
+		if (RigLogicConfig.LoadJoints)
+		{
+			MapJoints(InHierarchy);
+		}
+		if (RigLogicConfig.LoadTwistSwingBehavior || RigLogicConfig.LoadRBFBehavior)
+		{
+			MapDriverJoints(InHierarchy);
+		}
+
 		MapInputCurveIndices(InHierarchy);
-		MapNeuralNetMaskCurveIndices(InHierarchy);
-		MapMorphTargets(InHierarchy);
-		MapMaskMultipliers(InHierarchy);
+
+		if (RigLogicConfig.LoadMachineLearnedBehavior)
+		{
+			MapNeuralNetMaskCurveIndices(InHierarchy);
+		}
+		if (RigLogicConfig.LoadBlendShapes)
+		{
+			MapMorphTargets(InHierarchy);
+		}
+		if (RigLogicConfig.LoadAnimatedMaps)
+		{
+			MapMaskMultipliers(InHierarchy);
+		}
 	}
 }
 
@@ -181,10 +201,19 @@ void FRigUnit_RigLogic_Data::MapJoints(const URigHierarchy* Hierarchy)
 		const int32 BoneIndex = Hierarchy->GetIndex(FRigElementKey(JointFName, ERigElementType::Bone));
 		HierarchyBoneIndices.Add(BoneIndex);
 	}
+}
+
+void FRigUnit_RigLogic_Data::MapDriverJoints(const URigHierarchy* Hierarchy)
+{
+	DECLARE_SCOPE_HIERARCHICAL_COUNTER_FUNC()
+	QUICK_SCOPE_CYCLE_COUNTER(STAT_RigUnit_RigLogic_MapDriverJoints);
+
+	const IDNAReader* DNABehavior = LocalRigRuntimeContext->BehaviorReader.Get();
 
 	auto FindJointIndex = [DNABehavior](const FString& JointName)
 		{
-			for (uint16 JointIndex = 0; JointIndex < DNABehavior->GetJointCount(); ++JointIndex)
+			const uint16 JointCount = DNABehavior->GetJointCount();
+			for (uint16 JointIndex = 0; JointIndex < JointCount; ++JointIndex)
 			{
 				const FString& RLJointName = DNABehavior->GetJointName(JointIndex);
 				if (RLJointName == JointName)
@@ -336,44 +365,51 @@ void FRigUnit_RigLogic_Data::CalculateRigLogic(const URigHierarchy* InHierarchy,
 		RigInstance->SetRawControl(ControlIndex, FMath::Clamp(Value, 0.0, 1.0));
 	}
 
-	for (int32 MappingIndex = 0; MappingIndex < DriverJointsToControlAttributesMap.Num(); ++MappingIndex)
+	const FRigLogicConfiguration& RigLogicConfig = LocalRigRuntimeContext->RigLogic->GetConfiguration();
+	if (RigLogicConfig.LoadRBFBehavior || RigLogicConfig.LoadTwistSwingBehavior)
 	{
-		const FBoneIndexControlAttributeMapping& Mapping = DriverJointsToControlAttributesMap[MappingIndex];
-		const FTransform& PoseTransform = InHierarchy->GetLocalTransform(Mapping.BoneIndex);
-		// Translation and Scale is currently not used here, so to avoid the overhead of checking them, they are simply ignored.
-		// Should the need arise to use them as well, this code will need adjustment.
-		const FQuat Rotation = PoseTransform.GetRotation();
-		const int32 AttrIndex = Mapping.DNAJointIndex * MAX_ATTRS_PER_JOINT;
-		const tdm::fquat NeutralRotation{NeutralJointValues[AttrIndex + 3], NeutralJointValues[AttrIndex + 4], NeutralJointValues[AttrIndex + 5], NeutralJointValues[AttrIndex + 6]};
-		const tdm::fquat AbsPoseRotation{static_cast<float>(Rotation.X), static_cast<float>(Rotation.Y), static_cast<float>(Rotation.Z), static_cast<float>(Rotation.W)};
-		const tdm::fquat DeltaPoseRotation = tdm::inverse(NeutralRotation) * AbsPoseRotation;
-
-		if (Mapping.RotationX != INDEX_NONE)
-		{
-			RigInstance->SetRawControl(Mapping.RotationX, DeltaPoseRotation.x);
-		}
-		if (Mapping.RotationY != INDEX_NONE)
+		for (int32 MappingIndex = 0; MappingIndex < DriverJointsToControlAttributesMap.Num(); ++MappingIndex)
 		{
-			RigInstance->SetRawControl(Mapping.RotationY, DeltaPoseRotation.y);
-		}
-		if (Mapping.RotationZ != INDEX_NONE)
-		{
-			RigInstance->SetRawControl(Mapping.RotationZ, DeltaPoseRotation.z);
-		}
-		if (Mapping.RotationW != INDEX_NONE)
-		{
-			RigInstance->SetRawControl(Mapping.RotationW, DeltaPoseRotation.w);
+			const FBoneIndexControlAttributeMapping& Mapping = DriverJointsToControlAttributesMap[MappingIndex];
+			const FTransform& PoseTransform = InHierarchy->GetLocalTransform(Mapping.BoneIndex);
+			// Translation and Scale is currently not used here, so to avoid the overhead of checking them, they are simply ignored.
+			// Should the need arise to use them as well, this code will need adjustment.
+			const FQuat Rotation = PoseTransform.GetRotation();
+			const int32 AttrIndex = Mapping.DNAJointIndex * MAX_ATTRS_PER_JOINT;
+			const tdm::fquat NeutralRotation{ NeutralJointValues[AttrIndex + 3], NeutralJointValues[AttrIndex + 4], NeutralJointValues[AttrIndex + 5], NeutralJointValues[AttrIndex + 6] };
+			const tdm::fquat AbsPoseRotation{ static_cast<float>(Rotation.X), static_cast<float>(Rotation.Y), static_cast<float>(Rotation.Z), static_cast<float>(Rotation.W) };
+			const tdm::fquat DeltaPoseRotation = tdm::inverse(NeutralRotation) * AbsPoseRotation;
+
+			if (Mapping.RotationX != INDEX_NONE)
+			{
+				RigInstance->SetRawControl(Mapping.RotationX, DeltaPoseRotation.x);
+			}
+			if (Mapping.RotationY != INDEX_NONE)
+			{
+				RigInstance->SetRawControl(Mapping.RotationY, DeltaPoseRotation.y);
+			}
+			if (Mapping.RotationZ != INDEX_NONE)
+			{
+				RigInstance->SetRawControl(Mapping.RotationZ, DeltaPoseRotation.z);
+			}
+			if (Mapping.RotationW != INDEX_NONE)
+			{
+				RigInstance->SetRawControl(Mapping.RotationW, DeltaPoseRotation.w);
+			}
 		}
 	}
 
-	const int32 NeuralNetworkCount = RigInstance->GetNeuralNetworkCount();
-	for (int32 NeuralNetworkIndex = 0; NeuralNetworkIndex < NeuralNetworkCount; ++NeuralNetworkIndex)
+	if (RigLogicConfig.LoadMachineLearnedBehavior)
 	{
-		const uint32 CurveIndex = NeuralNetMaskCurveIndices[NeuralNetworkIndex];
-		if (InHierarchy->IsCurveValueSetByIndex(CurveIndex))
+		const int32 NeuralNetworkCount = RigInstance->GetNeuralNetworkCount();
+		for (int32 NeuralNetworkIndex = 0; NeuralNetworkIndex < NeuralNetworkCount; ++NeuralNetworkIndex)
 		{
-			const float Value = InHierarchy->GetCurveValue(CurveIndex);
-			RigInstance->SetNeuralNetworkMask(NeuralNetworkIndex, Value);
+			const uint32 CurveIndex = NeuralNetMaskCurveIndices[NeuralNetworkIndex];
+			if (InHierarchy->IsCurveValueSetByIndex(CurveIndex))
+			{
+				const float Value = InHierarchy->GetCurveValue(CurveIndex);
+				RigInstance->SetNeuralNetworkMask(NeuralNetworkIndex, Value);
+			}
 		}
 	}
 
@@ -517,18 +553,29 @@ FRigUnit_RigLogic_Execute()
 		}
 		Data.CurrentLOD = Data.SkelMeshComponent->GetPredictedLODLevel();
 
+		const FRigLogicConfiguration& RigLogicConfig = Data.LocalRigRuntimeContext->RigLogic->GetConfiguration();
+
 		TArrayView<const float> NeutralJointValues = Data.LocalRigRuntimeContext->RigLogic->GetNeutralJointValues();
 		TArrayView<const float> DeltaJointValues = Data.RigInstance->GetJointOutputs();
 
 		Data.CalculateRigLogic(Hierarchy, NeutralJointValues);
 
-		Data.UpdateJoints(Hierarchy, NeutralJointValues, DeltaJointValues);
+		if (RigLogicConfig.LoadJoints)
+		{
+			Data.UpdateJoints(Hierarchy, NeutralJointValues, DeltaJointValues);
+		}
 
-		TArrayView<const float> BlendShapeValues = Data.RigInstance->GetBlendShapeOutputs();
-		Data.UpdateBlendShapeCurves(Hierarchy, BlendShapeValues);
+		if (RigLogicConfig.LoadBlendShapes)
+		{
+			TArrayView<const float> BlendShapeValues = Data.RigInstance->GetBlendShapeOutputs();
+			Data.UpdateBlendShapeCurves(Hierarchy, BlendShapeValues);
+		}
 
-		TArrayView<const float> AnimMapOutputs = Data.RigInstance->GetAnimatedMapOutputs();
-		Data.UpdateAnimMapCurves(Hierarchy, AnimMapOutputs);
+		if (RigLogicConfig.LoadAnimatedMaps)
+		{
+			TArrayView<const float> AnimMapOutputs = Data.RigInstance->GetAnimatedMapOutputs();
+			Data.UpdateAnimMapCurves(Hierarchy, AnimMapOutputs);
+		}
 
 		//const double delta = FPlatformTime::Seconds() - startTime;
 		//UE_LOG(LogRigLogicUnit, Warning, TEXT("RigLogic::Update execution time: %f"), delta);
diff --git a/Engine/Plugins/Animation/RigLogic/Source/RigLogicModule/Private/SharedRigRuntimeContext.cpp b/Engine/Plugins/Animation/RigLogic/Source/RigLogicModule/Private/SharedRigRuntimeContext.cpp
--- a/Engine/Plugins/Animation/RigLogic/Source/RigLogicModule/Private/SharedRigRuntimeContext.cpp
+++ b/Engine/Plugins/Animation/RigLogic/Source/RigLogicModule/Private/SharedRigRuntimeContext.cpp
@@ -32,3 +32,22 @@ void FSharedRigRuntimeContext::CacheVariableJointIndices()
 		VariableJointIndicesPerLOD[LODIndex].Values = DistinctVariableJointIndices.Array();
 	}
 }
+
+void FSharedRigRuntimeContext::CacheInverseNeutralJointRotations()
+{
+	LLM_SCOPE_BYNAME(TEXT("Animation/RigLogic"));
+
+	static constexpr uint32 JointAttributeCount = 10;
+	TArrayView<const float> NeutralJointValues = RigLogic->GetNeutralJointValues();
+	ensure(NeutralJointValues.Num() % JointAttributeCount == 0);
+
+	const int32 JointCount = NeutralJointValues.Num() / JointAttributeCount;
+	InverseNeutralJointRotations.Reset(JointCount);
+	for (int32 JointIndex = 0; JointIndex < JointCount; ++JointIndex)
+	{
+		const int32 AttrIndex = JointIndex * JointAttributeCount;
+		const tdm::fquat NeutralRotation{NeutralJointValues[AttrIndex + 3], NeutralJointValues[AttrIndex + 4], NeutralJointValues[AttrIndex + 5], NeutralJointValues[AttrIndex + 6]};
+		const tdm::fquat InverseNeutralRotation = tdm::inverse(NeutralRotation);
+		InverseNeutralJointRotations.Add(FQuat(InverseNeutralRotation.x, InverseNeutralRotation.y, InverseNeutralRotation.z, InverseNeutralRotation.w));
+	}
+}
diff --git a/Engine/Plugins/Animation/RigLogic/Source/RigLogicModule/Public/AnimNode_RigLogic.h b/Engine/Plugins/Animation/RigLogic/Source/RigLogicModule/Public/AnimNode_RigLogic.h
--- a/Engine/Plugins/Animation/RigLogic/Source/RigLogicModule/Public/AnimNode_RigLogic.h
+++ b/Engine/Plugins/Animation/RigLogic/Source/RigLogicModule/Public/AnimNode_RigLogic.h
@@ -36,12 +36,33 @@ public:
 	UPROPERTY(EditAnywhere, BlueprintReadWrite, Category = Links)
 	FPoseLink AnimSequence;
 
+	/*
+	 * Since the order of anim curves may change even on a frame by frame basis, it is not safe to cache and
+	 * rely on cached indices by default, but if the blueprints feeding anim curves into RigLogic are set up
+	 * in a controlled manner, such that no such runtime changes are expected to the order or number of anim
+	 * curves, caching may improve the performance of the node, especially in low-LOD evaluations.
+	 */
+	UPROPERTY(EditAnywhere, BlueprintReadWrite, Category = RigLogic)
+	bool CacheAnimCurveNames = false;
+
+private:
+	using FCachedIndexedCurve = TBaseBlendedCurve<FDefaultAllocator, UE::Anim::FCurveElementIndexed>;
+
 private:
-	void UpdateControlCurves(const FPoseContext& OutputContext, const FDNAIndexMapping* DNAIndexMapping, TArrayView<const float> NeutralJointValues);
-	void CalculateRigLogic(FRigLogic* RigLogic);
-	void UpdateJoints(TArrayView<const float> NeutralJointValues, TArrayView<const float> DeltaJointValues, FPoseContext& OutputContext);
-	void UpdateBlendShapeCurves(const FDNAIndexMapping* DNAIndexMapping, TArrayView<const float> BlendShapeValues, FPoseContext& OutputContext);
-	void UpdateAnimMapCurves(const FDNAIndexMapping* DNAIndexMapping, TArrayView<const float> AnimMapOutputs, FPoseContext& OutputContext);
+	void CacheVariableJointAttributes(const FBoneContainer& RequiredBones);
+	void CacheDriverJoints(const FBoneContainer& RequiredBones);
+	void CachePoseCurvesToRigLogicControlsMap(const FPoseContext& InputContext, const FCachedIndexedCurve& IndexedCurves, TArray<int32>& Indices);
+	void UpdateRawControls(const FPoseContext& InputContext);
+	void UpdateRawControlsCached(const FPoseContext& InputContext);
+	void UpdateSparseDriverJointDrivenControlCurves(const FPoseContext& InputContext);
+	void UpdateDenseDriverJointDrivenControlCurves(const FPoseContext& InputContext);
+	void UpdateNeuralNetworkMaskCurves(const FPoseContext& InputContext);
+	void UpdateNeuralNetworkMaskCurvesCached(const FPoseContext& InputContext);
+	void UpdateControlCurves(const FPoseContext& InputContext);
+	void CalculateRigLogic();
+	void UpdateJoints(FPoseContext& OutputContext);
+	void UpdateBlendShapeCurves(FPoseContext& OutputContext);
+	void UpdateAnimMapCurves(FPoseContext& OutputContext);
 
 private:
 	struct FJointCompactPoseBoneMapping
@@ -68,6 +89,12 @@ private:
 		int32 BoneCount = -1;
 	};
 
+	struct FCurveElementControlAttributeMapping
+	{
+		TArray<int32> RawControlIndices;
+		TArray<int32> NeuralNetworkMaskIndices;
+	};
+
 private:
 	/*
 	 * Max LOD level that post-process AnimBPs are evaluated.
@@ -82,4 +109,5 @@ private:
 	TSharedPtr<FDNAIndexMapping> LocalDNAIndexMapping;
 	FRigInstance* RigInstance;
 	TArray<FCachedJointMapping> LocalJointMappingsPerLOD;
+	TArray<FCurveElementControlAttributeMapping> PoseCurvesToRigLogicControlsMap;
 };
diff --git a/Engine/Plugins/Animation/RigLogic/Source/RigLogicModule/Public/DNAAsset.h b/Engine/Plugins/Animation/RigLogic/Source/RigLogicModule/Public/DNAAsset.h
--- a/Engine/Plugins/Animation/RigLogic/Source/RigLogicModule/Public/DNAAsset.h
+++ b/Engine/Plugins/Animation/RigLogic/Source/RigLogicModule/Public/DNAAsset.h
@@ -69,6 +69,13 @@ public:
 	UPROPERTY(EditAnywhere, BlueprintReadWrite, Category = RuntimeSettings)
 	TMap<FString, FString> MetaData;
 
+	UPROPERTY(EditAnywhere, BlueprintReadWrite, Category = RuntimeSettings)
+	FRigLogicConfiguration RigLogicConfiguration;
+
+#if WITH_EDITOR
+	void PostEditChangeProperty(FPropertyChangedEvent& PropertyChangedEvent) override;
+#endif
+
 	bool Init(const FString& Filename);
 	void Serialize(FArchive& Ar) override;
 
diff --git a/Engine/Plugins/Animation/RigLogic/Source/RigLogicModule/Public/RigLogic.h b/Engine/Plugins/Animation/RigLogic/Source/RigLogicModule/Public/RigLogic.h
--- a/Engine/Plugins/Animation/RigLogic/Source/RigLogicModule/Public/RigLogic.h
+++ b/Engine/Plugins/Animation/RigLogic/Source/RigLogicModule/Public/RigLogic.h
@@ -77,7 +77,10 @@ struct FRigLogicConfiguration
 		TranslationType(ERigLogicTranslationType::Vector),
 		RotationType(ERigLogicRotationType::Quaternions),
 		RotationOrder(ERigLogicRotationOrder::ZYX),
-		ScaleType(ERigLogicScaleType::Vector)
+		ScaleType(ERigLogicScaleType::Vector),
+		TranslationPruningThreshold(0.0f),
+		RotationPruningThreshold(0.0f),
+		ScalePruningThreshold(0.0f)
 	{
 	}
 
@@ -91,7 +94,10 @@ struct FRigLogicConfiguration
 							ERigLogicTranslationType TranslationType,
 							ERigLogicRotationType RotationType,
 							ERigLogicRotationOrder RotationOrder,
-							ERigLogicScaleType ScaleType) :
+							ERigLogicScaleType ScaleType,
+							float TranslationPruningThreshold,
+							float RotationPruningThreshold,
+							float ScalePruningThreshold) :
 		CalculationType(CalculationType),
 		LoadJoints(LoadJoints),
 		LoadBlendShapes(LoadBlendShapes),
@@ -102,30 +108,32 @@ struct FRigLogicConfiguration
 		TranslationType(TranslationType),
 		RotationType(RotationType),
 		RotationOrder(RotationOrder),
-		ScaleType(ScaleType
-	)
+		ScaleType(ScaleType),
+		TranslationPruningThreshold(TranslationPruningThreshold),
+		RotationPruningThreshold(RotationPruningThreshold),
+		ScalePruningThreshold(ScalePruningThreshold)
 	{
 	}
 
 	UPROPERTY(BlueprintReadOnly, Category = "RigLogic")
 	ERigLogicCalculationType CalculationType;
 
-	UPROPERTY(BlueprintReadOnly, Category = "RigLogic")
+	UPROPERTY(EditAnywhere, BlueprintReadWrite, Category = "RigLogic")
 	bool LoadJoints;
 
-	UPROPERTY(BlueprintReadOnly, Category = "RigLogic")
+	UPROPERTY(EditAnywhere, BlueprintReadWrite, Category = "RigLogic")
 	bool LoadBlendShapes;
 
-	UPROPERTY(BlueprintReadOnly, Category = "RigLogic")
+	UPROPERTY(EditAnywhere, BlueprintReadWrite, Category = "RigLogic")
 	bool LoadAnimatedMaps;
 
-	UPROPERTY(BlueprintReadOnly, Category = "RigLogic")
+	UPROPERTY(EditAnywhere, BlueprintReadWrite, Category = "RigLogic")
 	bool LoadMachineLearnedBehavior;
 	
-	UPROPERTY(BlueprintReadOnly, Category = "RigLogic")
+	UPROPERTY(EditAnywhere, BlueprintReadWrite, Category = "RigLogic")
 	bool LoadRBFBehavior;
 	
-	UPROPERTY(BlueprintReadOnly, Category = "RigLogic")
+	UPROPERTY(EditAnywhere, BlueprintReadWrite, Category = "RigLogic")
 	bool LoadTwistSwingBehavior;
 	
 	UPROPERTY(BlueprintReadOnly, Category = "RigLogic")
@@ -139,12 +147,39 @@ struct FRigLogicConfiguration
 
 	UPROPERTY(BlueprintReadOnly, Category = "RigLogic")
 	ERigLogicScaleType ScaleType;
+
+	/** The joint translation pruning threshold is used to eliminate joint translation deltas below
+	  * the specified threshold from the joint matrix when the RigLogic instance is initialized.
+	  * Use it with caution, as while it may reduce the amount of compute to be done, it may also erase
+	  * important deltas that could introduce artifacts into the rig.
+	  * A reasonably safe starting value to try translation pruning would be 0.0001f
+	 **/
+	UPROPERTY(EditAnywhere, BlueprintReadWrite, Category = "RigLogic")
+	float TranslationPruningThreshold;
+
+	/** The joint rotation pruning threshold is used to eliminate joint rotation deltas below
+	  * the specified threshold from the joint matrix when the RigLogic instance is initialized.
+	  * Use it with caution, as while it may reduce the amount of compute to be done, it may also erase
+	  * important deltas that could introduce artifacts into the rig.
+	  * A reasonably safe starting value to try rotation pruning would be 0.1f
+	 **/
+	UPROPERTY(EditAnywhere, BlueprintReadWrite, Category = "RigLogic")
+	float RotationPruningThreshold;
+
+	/** The joint scale pruning threshold is used to eliminate joint scale deltas below
+	  * the specified threshold from the joint matrix when the RigLogic instance is initialized.
+	  * Use it with caution, as while it may reduce the amount of compute to be done, it may also erase
+	  * important deltas that could introduce artifacts into the rig.
+	  * A reasonably safe starting value to try scale pruning would be 0.001f
+	 **/
+	UPROPERTY(EditAnywhere, BlueprintReadWrite, Category = "RigLogic")
+	float ScalePruningThreshold;
 };
 
 class RIGLOGICMODULE_API FRigLogic
 {
 public:
-	explicit FRigLogic(const IDNAReader* Reader, FRigLogicConfiguration Config = FRigLogicConfiguration());
+	explicit FRigLogic(const IDNAReader* Reader, const FRigLogicConfiguration& Config = FRigLogicConfiguration());
 	~FRigLogic();
 
 	FRigLogic(const FRigLogic&) = delete;
@@ -153,7 +188,13 @@ public:
 	FRigLogic(FRigLogic&&) = default;
 	FRigLogic& operator=(FRigLogic&&) = default;
 
+	const FRigLogicConfiguration& GetConfiguration() const;
 	uint16 GetLODCount() const;
+	TArrayView<const uint16> GetRBFSolverIndicesForLOD(uint16_t LOD) const;
+	TArrayView<const uint32> GetNeuralNetworkIndicesForLOD(uint16_t LOD) const;
+	TArrayView<const uint16> GetBlendShapeChannelIndicesForLOD(uint16_t LOD) const;
+	TArrayView<const uint16> GetAnimatedMapIndicesForLOD(uint16_t LOD) const;
+	TArrayView<const uint16> GetJointIndicesForLOD(uint16_t LOD) const;
 	TArrayView<const float> GetNeutralJointValues() const;
 	TArrayView<const uint16> GetJointVariableAttributeIndices(uint16 LOD) const;
 	uint16 GetJointGroupCount() const;
@@ -175,6 +216,7 @@ public:
 	void CalculateBlendShapes(FRigInstance* Instance) const;
 	void CalculateAnimatedMaps(FRigInstance* Instance) const;
 	void Calculate(FRigInstance* Instance) const;
+	void CollectCalculationStats(FRigInstance* Instance) const;
 
 private:
 	friend FRigInstance;
@@ -188,5 +230,6 @@ private:
 		void operator()(rl4::RigLogic* Pointer);
 	};
 	TUniquePtr<rl4::RigLogic, FRigLogicDeleter> RigLogic;
+	FRigLogicConfiguration Configuration;
 };
 
diff --git a/Engine/Plugins/Animation/RigLogic/Source/RigLogicModule/Public/RigLogicModule.h b/Engine/Plugins/Animation/RigLogic/Source/RigLogicModule/Public/RigLogicModule.h
--- a/Engine/Plugins/Animation/RigLogic/Source/RigLogicModule/Public/RigLogicModule.h
+++ b/Engine/Plugins/Animation/RigLogic/Source/RigLogicModule/Public/RigLogicModule.h
@@ -4,9 +4,34 @@
 
 #include "CoreMinimal.h"
 #include "Modules/ModuleInterface.h"
+#include "Stats/Stats.h"
 
 DECLARE_LOG_CATEGORY_EXTERN(LogRigLogic, Log, All);
 
+#if STATS
+DECLARE_STATS_GROUP(TEXT("RigLogic"), STATGROUP_RigLogic, STATCAT_Advanced)
+
+DECLARE_DWORD_COUNTER_STAT_EXTERN(TEXT("Calculation Backend [ 0 = Scalar, 1 = SSE, 2 = AVX, 3 = NEON ]"), STAT_RigLogic_CalculationType, STATGROUP_RigLogic, RIGLOGICMODULE_API);
+DECLARE_DWORD_COUNTER_STAT_EXTERN(TEXT("Floating Point Type [ 0 = Float, 1 = Half-Float ]"), STAT_RigLogic_FloatingPointType, STATGROUP_RigLogic, RIGLOGICMODULE_API);
+DECLARE_DWORD_COUNTER_STAT_EXTERN(TEXT("Current LOD"), STAT_RigLogic_LOD, STATGROUP_RigLogic, RIGLOGICMODULE_API);
+DECLARE_DWORD_COUNTER_STAT_EXTERN(TEXT("RBF Solver Count"), STAT_RigLogic_RBFSolverCount, STATGROUP_RigLogic, RIGLOGICMODULE_API);
+DECLARE_DWORD_COUNTER_STAT_EXTERN(TEXT("Neural Network Count"), STAT_RigLogic_NeuralNetworkCount, STATGROUP_RigLogic, RIGLOGICMODULE_API);
+DECLARE_DWORD_COUNTER_STAT_EXTERN(TEXT("PSD Expression Count"), STAT_RigLogic_PSDCount, STATGROUP_RigLogic, RIGLOGICMODULE_API);
+DECLARE_DWORD_COUNTER_STAT_EXTERN(TEXT("Joint Count"), STAT_RigLogic_JointCount, STATGROUP_RigLogic, RIGLOGICMODULE_API);
+DECLARE_DWORD_COUNTER_STAT_EXTERN(TEXT("Joint Delta Value Count"), STAT_RigLogic_JointDeltaValueCount, STATGROUP_RigLogic, RIGLOGICMODULE_API);
+DECLARE_DWORD_COUNTER_STAT_EXTERN(TEXT("Blend Shape Channel Count"), STAT_RigLogic_BlendShapeChannelCount, STATGROUP_RigLogic, RIGLOGICMODULE_API);
+DECLARE_DWORD_COUNTER_STAT_EXTERN(TEXT("Animated Map Count"), STAT_RigLogic_AnimatedMapCount, STATGROUP_RigLogic, RIGLOGICMODULE_API);
+
+DECLARE_CYCLE_STAT_EXTERN(TEXT("Map GUI To Raw Controls Time"), STAT_FRigLogic_MapGUIToRawControlsTime, STATGROUP_RigLogic, RIGLOGICMODULE_API);
+DECLARE_CYCLE_STAT_EXTERN(TEXT("Map Raw To GUI Controls Time"), STAT_FRigLogic_MapRawToGUIControlsTime, STATGROUP_RigLogic, RIGLOGICMODULE_API);
+DECLARE_CYCLE_STAT_EXTERN(TEXT("Calculate RBF Controls Time"), STAT_RigLogic_CalculateRBFControlsTime, STATGROUP_RigLogic, RIGLOGICMODULE_API);
+DECLARE_CYCLE_STAT_EXTERN(TEXT("Calculate ML Controls Time"), STAT_RigLogic_CalculateMLControlsTime, STATGROUP_RigLogic, RIGLOGICMODULE_API);
+DECLARE_CYCLE_STAT_EXTERN(TEXT("Calculate PSD Controls Time"), STAT_RigLogic_CalculateControlsTime, STATGROUP_RigLogic, RIGLOGICMODULE_API);
+DECLARE_CYCLE_STAT_EXTERN(TEXT("Calculate Joints Time"), STAT_RigLogic_CalculateJointsTime, STATGROUP_RigLogic, RIGLOGICMODULE_API);
+DECLARE_CYCLE_STAT_EXTERN(TEXT("Calculate Blend Shapes Time"), STAT_RigLogic_CalculateBlendShapesTime, STATGROUP_RigLogic, RIGLOGICMODULE_API);
+DECLARE_CYCLE_STAT_EXTERN(TEXT("Calculate Animated Maps Time"), STAT_RigLogic_CalculateAnimatedMapsTime, STATGROUP_RigLogic, RIGLOGICMODULE_API);
+#endif  // STATS
+
 class RIGLOGICMODULE_API FRigLogicModule: public IModuleInterface
 {
 public:
diff --git a/Engine/Plugins/Animation/RigLogic/Source/RigLogicModule/Public/RigUnit_RigLogic.h b/Engine/Plugins/Animation/RigLogic/Source/RigLogicModule/Public/RigUnit_RigLogic.h
--- a/Engine/Plugins/Animation/RigLogic/Source/RigLogicModule/Public/RigUnit_RigLogic.h
+++ b/Engine/Plugins/Animation/RigLogic/Source/RigLogicModule/Public/RigUnit_RigLogic.h
@@ -116,6 +116,7 @@ struct FRigUnit_RigLogic_Data
 	void MapNeuralNetMaskCurveIndices(const URigHierarchy* InHierarchy);
 	/** Uses names to map joint indices from DNA file to the indices of bones in control rig hierarchy **/
 	void MapJoints(const URigHierarchy* Hierarchy);
+	void MapDriverJoints(const URigHierarchy* Hierarchy);
 	/** Uses names of blend shapes and meshes in DNA file, for all LODs, to map their indices to the indices of
 	  * morph target curves in the curve container; curve name format is <mesh>__<blendshape> **/
 	void MapMorphTargets(const URigHierarchy* InHierarchy);
diff --git a/Engine/Plugins/Animation/RigLogic/Source/RigLogicModule/Public/SharedRigRuntimeContext.h b/Engine/Plugins/Animation/RigLogic/Source/RigLogicModule/Public/SharedRigRuntimeContext.h
--- a/Engine/Plugins/Animation/RigLogic/Source/RigLogicModule/Public/SharedRigRuntimeContext.h
+++ b/Engine/Plugins/Animation/RigLogic/Source/RigLogicModule/Public/SharedRigRuntimeContext.h
@@ -25,5 +25,9 @@ struct FSharedRigRuntimeContext
 	/** Cached joint indices that need to be updated for each LOD **/
 	TArray<TNestedArray<uint16>> VariableJointIndicesPerLOD;
 
+	TArray<FQuat> InverseNeutralJointRotations;
+
 	void CacheVariableJointIndices();
+	void CacheInverseNeutralJointRotations();
+
 };
