Fix RigLogic incremental build error due to missing preprocessor guards for stats

diff --git a/Engine/Plugins/Animation/RigLogic/Source/RigLogicModule/Private/AnimNode_RigLogic.cpp b/Engine/Plugins/Animation/RigLogic/Source/RigLogicModule/Private/AnimNode_RigLogic.cpp
--- a/Engine/Plugins/Animation/RigLogic/Source/RigLogicModule/Private/AnimNode_RigLogic.cpp
+++ b/Engine/Plugins/Animation/RigLogic/Source/RigLogicModule/Private/AnimNode_RigLogic.cpp
@@ -36,9 +36,13 @@ FAnimNode_RigLogic::~FAnimNode_RigLogic()
 
 void FAnimNode_RigLogic::Initialize_AnyThread(const FAnimationInitializeContext& Context)
 {
+#if STATS
 	DECLARE_SCOPE_HIERARCHICAL_COUNTER_FUNC()
 	QUICK_SCOPE_CYCLE_COUNTER(STAT_AnimNode_RigLogic_Initialize_AnyThread);
+#endif  // STATS
+#if CPUPROFILERTRACE_ENABLED
 	TRACE_CPUPROFILER_EVENT_SCOPE_STR("FAnimNode_RigLogic::Initialize_AnyThread");
+#endif  // CPUPROFILERTRACE_ENABLED
 	LLM_SCOPE_BYNAME(TEXT("Animation/RigLogic"));
 
 	AnimSequence.Initialize(Context);
@@ -46,9 +50,13 @@ void FAnimNode_RigLogic::Initialize_AnyThread(const FAnimationInitializeContext&
 
 void FAnimNode_RigLogic::CacheVariableJointAttributes(const FBoneContainer& RequiredBones)
 {
+#if STATS
 	DECLARE_SCOPE_HIERARCHICAL_COUNTER_FUNC()
 	QUICK_SCOPE_CYCLE_COUNTER(STAT_AnimNode_RigLogic_CacheVariableJointAttributes);
+#endif  // STATS
+#if CPUPROFILERTRACE_ENABLED
 	TRACE_CPUPROFILER_EVENT_SCOPE_STR("FAnimNode_RigLogic::CacheVariableJointAttributes");
+#endif  // CPUPROFILERTRACE_ENABLED
 	LLM_SCOPE_BYNAME(TEXT("Animation/RigLogic"));
 
 	// Populate mapping of DNA joint indices to CompactPoseBoneIndex entries (used in updating joints with results from RigLogic)
@@ -71,9 +79,13 @@ void FAnimNode_RigLogic::CacheVariableJointAttributes(const FBoneContainer& Requ
 
 void FAnimNode_RigLogic::CacheDriverJoints(const FBoneContainer& RequiredBones)
 {
+#if STATS
 	DECLARE_SCOPE_HIERARCHICAL_COUNTER_FUNC()
 	QUICK_SCOPE_CYCLE_COUNTER(STAT_AnimNode_RigLogic_CacheDriverJoints);
+#endif  // STATS
+#if CPUPROFILERTRACE_ENABLED
 	TRACE_CPUPROFILER_EVENT_SCOPE_STR("FAnimNode_RigLogic::CacheDriverJoints");
+#endif  // CPUPROFILERTRACE_ENABLED
 	LLM_SCOPE_BYNAME(TEXT("Animation/RigLogic"));
 
 	const uint16 CurrentLOD = RigInstance->GetLOD();
@@ -104,9 +116,13 @@ void FAnimNode_RigLogic::CacheDriverJoints(const FBoneContainer& RequiredBones)
 
 void FAnimNode_RigLogic::CachePoseCurvesToRigLogicControlsMap(const FPoseContext& InputContext, const FCachedIndexedCurve& IndexedCurves, TArray<int32>& Indices)
 {
+#if STATS
 	DECLARE_SCOPE_HIERARCHICAL_COUNTER_FUNC()
 	QUICK_SCOPE_CYCLE_COUNTER(STAT_AnimNode_RigLogic_CachePoseCurvesToRigLogicControlsMap);
+#endif  // STATS
+#if CPUPROFILERTRACE_ENABLED
 	TRACE_CPUPROFILER_EVENT_SCOPE_STR("FAnimNode_RigLogic::CachePoseCurvesToRigLogicControlsMap");
+#endif  // CPUPROFILERTRACE_ENABLED
 	LLM_SCOPE_BYNAME(TEXT("Animation/RigLogic"));
 
 	Indices.Init(INDEX_NONE, InputContext.Curve.Num());
@@ -127,9 +143,13 @@ void FAnimNode_RigLogic::CachePoseCurvesToRigLogicControlsMap(const FPoseContext
 
 void FAnimNode_RigLogic::CacheBones_AnyThread(const FAnimationCacheBonesContext& Context)
 {
+#if STATS
 	DECLARE_SCOPE_HIERARCHICAL_COUNTER_FUNC()
 	QUICK_SCOPE_CYCLE_COUNTER(STAT_AnimNode_RigLogic_CacheBones_AnyThread);
+#endif  // STATS
+#if CPUPROFILERTRACE_ENABLED
 	TRACE_CPUPROFILER_EVENT_SCOPE_STR("FAnimNode_RigLogic::CacheBones_AnyThread");
+#endif  // CPUPROFILERTRACE_ENABLED
 	LLM_SCOPE_BYNAME(TEXT("Animation/RigLogic"));
 
 	AnimSequence.CacheBones(Context);
@@ -213,9 +233,13 @@ void FAnimNode_RigLogic::CacheBones_AnyThread(const FAnimationCacheBonesContext&
 
 void FAnimNode_RigLogic::Update_AnyThread(const FAnimationUpdateContext& Context)
 {
+#if STATS
 	DECLARE_SCOPE_HIERARCHICAL_COUNTER_FUNC()
 	QUICK_SCOPE_CYCLE_COUNTER(STAT_AnimNode_RigLogic_Update_AnyThread);
+#endif  // STATS
+#if CPUPROFILERTRACE_ENABLED
 	TRACE_CPUPROFILER_EVENT_SCOPE_STR("FAnimNode_RigLogic::Update_AnyThread");
+#endif  // CPUPROFILERTRACE_ENABLED
 
 	GetEvaluateGraphExposedInputs().Execute(Context);
 	AnimSequence.Update(Context);
@@ -223,9 +247,13 @@ void FAnimNode_RigLogic::Update_AnyThread(const FAnimationUpdateContext& Context
 
 void FAnimNode_RigLogic::Evaluate_AnyThread(FPoseContext& OutputContext)
 {
+#if STATS
 	DECLARE_SCOPE_HIERARCHICAL_COUNTER_FUNC()
 	QUICK_SCOPE_CYCLE_COUNTER(STAT_AnimNode_RigLogic_Evaluate_AnyThread);
+#endif  // STATS
+#if CPUPROFILERTRACE_ENABLED
 	TRACE_CPUPROFILER_EVENT_SCOPE_STR("FAnimNode_RigLogic::Evaluate_AnyThread");
+#endif  // CPUPROFILERTRACE_ENABLED
 
 	AnimSequence.Evaluate(OutputContext);
 
@@ -269,9 +297,13 @@ void FAnimNode_RigLogic::GatherDebugData(FNodeDebugData& DebugData)
 
 void FAnimNode_RigLogic::UpdateRawControls(const FPoseContext& InputContext)
 {
+#if STATS
 	DECLARE_SCOPE_HIERARCHICAL_COUNTER_FUNC()
 	QUICK_SCOPE_CYCLE_COUNTER(STAT_AnimNode_RigLogic_UpdateRawControls);
+#endif  // STATS
+#if CPUPROFILERTRACE_ENABLED
 	TRACE_CPUPROFILER_EVENT_SCOPE_STR("FAnimNode_RigLogic::UpdateRawControls");
+#endif  // CPUPROFILERTRACE_ENABLED
 
 	// Combine control attribute curve with input curve to get indexed curve to apply to rig
 	// Curve elements that dont have a control mapping will have INDEX_NONE as their index
@@ -287,9 +319,13 @@ void FAnimNode_RigLogic::UpdateRawControls(const FPoseContext& InputContext)
 
 void FAnimNode_RigLogic::UpdateRawControlsCached(const FPoseContext& InputContext)
 {
+#if STATS
 	DECLARE_SCOPE_HIERARCHICAL_COUNTER_FUNC()
 	QUICK_SCOPE_CYCLE_COUNTER(STAT_AnimNode_RigLogic_UpdateRawControlsCached);
+#endif  // STATS
+#if CPUPROFILERTRACE_ENABLED
 	TRACE_CPUPROFILER_EVENT_SCOPE_STR("FAnimNode_RigLogic::UpdateRawControlsCached");
+#endif  // CPUPROFILERTRACE_ENABLED
 
 	const uint16 CurrentLOD = RigInstance->GetLOD();
 	auto& RawControlIndices = PoseCurvesToRigLogicControlsMap[CurrentLOD].RawControlIndices;
@@ -312,9 +348,13 @@ void FAnimNode_RigLogic::UpdateRawControlsCached(const FPoseContext& InputContex
 
 void FAnimNode_RigLogic::UpdateSparseDriverJointDrivenControlCurves(const FPoseContext& InputContext)
 {
+#if STATS
 	DECLARE_SCOPE_HIERARCHICAL_COUNTER_FUNC()
 	QUICK_SCOPE_CYCLE_COUNTER(STAT_AnimNode_RigLogic_UpdateSparseDriverJointDrivenControlCurves);
+#endif  // STATS
+#if CPUPROFILERTRACE_ENABLED
 	TRACE_CPUPROFILER_EVENT_SCOPE_STR("FAnimNode_RigLogic::UpdateSparseDriverJointDrivenControlCurves");
+#endif  // CPUPROFILERTRACE_ENABLED
 
 	const uint16 CurrentLOD = RigInstance->GetLOD();
 	const auto& SparseDriverJointsToControlAttributesMap = LocalJointMappingsPerLOD[CurrentLOD].SparseDriverJointsToControlAttributesMap;
@@ -351,9 +391,13 @@ void FAnimNode_RigLogic::UpdateSparseDriverJointDrivenControlCurves(const FPoseC
 
 void FAnimNode_RigLogic::UpdateDenseDriverJointDrivenControlCurves(const FPoseContext& InputContext)
 {
+#if STATS
 	DECLARE_SCOPE_HIERARCHICAL_COUNTER_FUNC()
 	QUICK_SCOPE_CYCLE_COUNTER(STAT_AnimNode_RigLogic_UpdateDenseDriverJointDrivenControlCurves);
+#endif  // STATS
+#if CPUPROFILERTRACE_ENABLED
 	TRACE_CPUPROFILER_EVENT_SCOPE_STR("FAnimNode_RigLogic::UpdateDenseDriverJointDrivenControlCurves");
+#endif  // CPUPROFILERTRACE_ENABLED
 
 	const uint16 CurrentLOD = RigInstance->GetLOD();
 	const auto& DenseDriverJointsToControlAttributesMap = LocalJointMappingsPerLOD[CurrentLOD].DenseDriverJointsToControlAttributesMap;
@@ -377,9 +421,13 @@ void FAnimNode_RigLogic::UpdateDenseDriverJointDrivenControlCurves(const FPoseCo
 
 void FAnimNode_RigLogic::UpdateNeuralNetworkMaskCurves(const FPoseContext& InputContext)
 {
+#if STATS
 	DECLARE_SCOPE_HIERARCHICAL_COUNTER_FUNC()
 	QUICK_SCOPE_CYCLE_COUNTER(STAT_AnimNode_RigLogic_UpdateNeuralNetworkMaskCurves);
+#endif  // STATS
+#if CPUPROFILERTRACE_ENABLED
 	TRACE_CPUPROFILER_EVENT_SCOPE_STR("FAnimNode_RigLogic::UpdateNeuralNetworkMaskCurves");
+#endif  // CPUPROFILERTRACE_ENABLED
 
 	if (RigInstance->GetNeuralNetworkCount() != 0)
 	{
@@ -396,9 +444,13 @@ void FAnimNode_RigLogic::UpdateNeuralNetworkMaskCurves(const FPoseContext& Input
 
 void FAnimNode_RigLogic::UpdateNeuralNetworkMaskCurvesCached(const FPoseContext& InputContext)
 {
+#if STATS
 	DECLARE_SCOPE_HIERARCHICAL_COUNTER_FUNC()
 	QUICK_SCOPE_CYCLE_COUNTER(STAT_AnimNode_RigLogic_UpdateNeuralNetworkMaskCurvesCached);
+#endif  // STATS
+#if CPUPROFILERTRACE_ENABLED
 	TRACE_CPUPROFILER_EVENT_SCOPE_STR("FAnimNode_RigLogic::UpdateNeuralNetworkMaskCurvesCached");
+#endif  // CPUPROFILERTRACE_ENABLED
 
 	if (RigInstance->GetNeuralNetworkCount() != 0)
 	{
@@ -424,9 +476,13 @@ void FAnimNode_RigLogic::UpdateNeuralNetworkMaskCurvesCached(const FPoseContext&
 
 void FAnimNode_RigLogic::UpdateControlCurves(const FPoseContext& InputContext)
 {
+#if STATS
 	DECLARE_SCOPE_HIERARCHICAL_COUNTER_FUNC()
 	QUICK_SCOPE_CYCLE_COUNTER(STAT_AnimNode_RigLogic_UpdateControlCurves);
+#endif  // STATS
+#if CPUPROFILERTRACE_ENABLED
 	TRACE_CPUPROFILER_EVENT_SCOPE_STR("FAnimNode_RigLogic::UpdateControlCurves");
+#endif  // CPUPROFILERTRACE_ENABLED
 
 	const FRigLogicConfiguration& RigLogicConfig = LocalRigRuntimeContext->RigLogic->GetConfiguration();
 
@@ -456,9 +512,13 @@ void FAnimNode_RigLogic::UpdateControlCurves(const FPoseContext& InputContext)
 
 void FAnimNode_RigLogic::CalculateRigLogic()
 {
+#if STATS
 	DECLARE_SCOPE_HIERARCHICAL_COUNTER_FUNC()
 	QUICK_SCOPE_CYCLE_COUNTER(STAT_AnimNode_RigLogic_CalculateRigLogic);
+#endif  // STATS
+#if CPUPROFILERTRACE_ENABLED
 	TRACE_CPUPROFILER_EVENT_SCOPE_STR("FAnimNode_RigLogic::CalculateRigLogic");
+#endif  // CPUPROFILERTRACE_ENABLED
 
 	FRigLogic* RigLogic = LocalRigRuntimeContext->RigLogic.Get();
 	// RigLogic has Null evaluators for each class of computations, so no explicit checks are necessary here
@@ -473,9 +533,13 @@ void FAnimNode_RigLogic::CalculateRigLogic()
 
 void FAnimNode_RigLogic::UpdateJoints(FPoseContext& OutputContext)
 {
+#if STATS
 	DECLARE_SCOPE_HIERARCHICAL_COUNTER_FUNC()
 	QUICK_SCOPE_CYCLE_COUNTER(STAT_AnimNode_RigLogic_UpdateJoints);
+#endif  // STATS
+#if CPUPROFILERTRACE_ENABLED
 	TRACE_CPUPROFILER_EVENT_SCOPE_STR("FAnimNode_RigLogic::UpdateJoints");
+#endif  // CPUPROFILERTRACE_ENABLED
 
 	const uint16 LOD = RigInstance->GetLOD();
 	TArrayView<const float> DeltaJointValues = RigInstance->GetJointOutputs();
@@ -495,9 +559,13 @@ void FAnimNode_RigLogic::UpdateJoints(FPoseContext& OutputContext)
 
 void FAnimNode_RigLogic::UpdateBlendShapeCurves(FPoseContext& OutputContext)
 {
+#if STATS
 	DECLARE_SCOPE_HIERARCHICAL_COUNTER_FUNC()
 	QUICK_SCOPE_CYCLE_COUNTER(STAT_AnimNode_RigLogic_UpdateBlendShapeCurves);
+#endif  // STATS
+#if CPUPROFILERTRACE_ENABLED
 	TRACE_CPUPROFILER_EVENT_SCOPE_STR("FAnimNode_RigLogic::UpdateBlendShapeCurves");
+#endif  // CPUPROFILERTRACE_ENABLED
 
 	const uint16 LOD = RigInstance->GetLOD();
 	TArrayView<const float> BlendShapeValues = RigInstance->GetBlendShapeOutputs();
@@ -514,9 +582,13 @@ void FAnimNode_RigLogic::UpdateBlendShapeCurves(FPoseContext& OutputContext)
 
 void FAnimNode_RigLogic::UpdateAnimMapCurves(FPoseContext& OutputContext)
 {
+#if STATS
 	DECLARE_SCOPE_HIERARCHICAL_COUNTER_FUNC()
 	QUICK_SCOPE_CYCLE_COUNTER(STAT_AnimNode_RigLogic_UpdateAnimMapCurves);
+#endif  // STATS
+#if CPUPROFILERTRACE_ENABLED
 	TRACE_CPUPROFILER_EVENT_SCOPE_STR("FAnimNode_RigLogic::UpdateAnimMapCurves");
+#endif  // CPUPROFILERTRACE_ENABLED
 
 	const uint16 LOD = RigInstance->GetLOD();
 	TArrayView<const float> AnimMapOutputs = RigInstance->GetAnimatedMapOutputs();
diff --git a/Engine/Plugins/Animation/RigLogic/Source/RigLogicModule/Private/RigLogic.cpp b/Engine/Plugins/Animation/RigLogic/Source/RigLogicModule/Private/RigLogic.cpp
--- a/Engine/Plugins/Animation/RigLogic/Source/RigLogicModule/Private/RigLogic.cpp
+++ b/Engine/Plugins/Animation/RigLogic/Source/RigLogicModule/Private/RigLogic.cpp
@@ -23,8 +23,8 @@ DEFINE_STAT(STAT_RigLogic_JointCount);
 DEFINE_STAT(STAT_RigLogic_JointDeltaValueCount);
 DEFINE_STAT(STAT_RigLogic_BlendShapeChannelCount);
 DEFINE_STAT(STAT_RigLogic_AnimatedMapCount);
-DEFINE_STAT(STAT_FRigLogic_MapGUIToRawControlsTime);
-DEFINE_STAT(STAT_FRigLogic_MapRawToGUIControlsTime);
+DEFINE_STAT(STAT_RigLogic_MapGUIToRawControlsTime);
+DEFINE_STAT(STAT_RigLogic_MapRawToGUIControlsTime);
 DEFINE_STAT(STAT_RigLogic_CalculateRBFControlsTime);
 DEFINE_STAT(STAT_RigLogic_CalculateMLControlsTime);
 DEFINE_STAT(STAT_RigLogic_CalculateControlsTime);
@@ -173,108 +173,155 @@ TArrayView<const uint16> FRigLogic::GetNeuralNetworkIndices(uint16 MeshIndex, ui
 
 void FRigLogic::MapGUIToRawControls(FRigInstance* Instance) const
 {
+#if STATS
 	DECLARE_SCOPE_HIERARCHICAL_COUNTER_FUNC()
-	SCOPE_CYCLE_COUNTER(STAT_FRigLogic_MapGUIToRawControlsTime);
+	SCOPE_CYCLE_COUNTER(STAT_RigLogic_MapGUIToRawControlsTime);
+#endif  // STATS
+#if CPUPROFILERTRACE_ENABLED
 	TRACE_CPUPROFILER_EVENT_SCOPE_STR("FRigLogic::MapGUIToRawControls");
-
+#endif  // CPUPROFILERTRACE_ENABLED
 	RigLogic->mapGUIToRawControls(Instance->Unwrap());
 }
 
 void FRigLogic::MapRawToGUIControls(FRigInstance* Instance) const
 {
+#if STATS
 	DECLARE_SCOPE_HIERARCHICAL_COUNTER_FUNC()
-	SCOPE_CYCLE_COUNTER(STAT_FRigLogic_MapRawToGUIControlsTime);
+	SCOPE_CYCLE_COUNTER(STAT_RigLogic_MapRawToGUIControlsTime);
+#endif  // STATS
+#if CPUPROFILERTRACE_ENABLED
 	TRACE_CPUPROFILER_EVENT_SCOPE_STR("FRigLogic::MapRawToGUIControls");
+#endif  // CPUPROFILERTRACE_ENABLED
 
 	RigLogic->mapRawToGUIControls(Instance->Unwrap());
 }
 
 void FRigLogic::CalculateControls(FRigInstance* Instance) const
 {
+#if STATS
 	DECLARE_SCOPE_HIERARCHICAL_COUNTER_FUNC()
 	SCOPE_CYCLE_COUNTER(STAT_RigLogic_CalculateControlsTime);
+#endif  // STATS
+#if CPUPROFILERTRACE_ENABLED
 	TRACE_CPUPROFILER_EVENT_SCOPE_STR("FRigLogic::CalculateControls");
+#endif  // CPUPROFILERTRACE_ENABLED
 
 	RigLogic->calculateControls(Instance->Unwrap());
 }
 
 void FRigLogic::CalculateMachineLearnedBehaviorControls(FRigInstance* Instance) const
 {
+#if STATS
 	DECLARE_SCOPE_HIERARCHICAL_COUNTER_FUNC()
 	SCOPE_CYCLE_COUNTER(STAT_RigLogic_CalculateMLControlsTime);
+#endif  // STATS
+#if CPUPROFILERTRACE_ENABLED
 	TRACE_CPUPROFILER_EVENT_SCOPE_STR("FRigLogic::CalculateMachineLearnedBehaviorControls");
+#endif  // CPUPROFILERTRACE_ENABLED
 
 	RigLogic->calculateMachineLearnedBehaviorControls(Instance->Unwrap());
 }
 
 void FRigLogic::CalculateMachineLearnedBehaviorControls(FRigInstance* Instance, uint16 NeuralNetIndex) const
 {
+#if STATS
 	DECLARE_SCOPE_HIERARCHICAL_COUNTER_FUNC()
 	SCOPE_CYCLE_COUNTER(STAT_RigLogic_CalculateMLControlsTime);
+#endif  // STATS
+#if CPUPROFILERTRACE_ENABLED
 	TRACE_CPUPROFILER_EVENT_SCOPE_STR("FRigLogic::CalculateMachineLearnedBehaviorControls");
+#endif  // CPUPROFILERTRACE_ENABLED
 
 	RigLogic->calculateMachineLearnedBehaviorControls(Instance->Unwrap(), NeuralNetIndex);
 }
 
 void FRigLogic::CalculateRBFControls(FRigInstance* Instance) const
 {
+#if STATS
 	DECLARE_SCOPE_HIERARCHICAL_COUNTER_FUNC()
 	SCOPE_CYCLE_COUNTER(STAT_RigLogic_CalculateRBFControlsTime);
+#endif  // STATS
+#if CPUPROFILERTRACE_ENABLED
 	TRACE_CPUPROFILER_EVENT_SCOPE_STR("FRigLogic::CalculateRBFControls");
+#endif  // CPUPROFILERTRACE_ENABLED
 
 	RigLogic->calculateRBFControls(Instance->Unwrap());
 }
 
 void FRigLogic::CalculateRBFControls(FRigInstance* Instance, uint16 SolverIndex) const
 {
+#if STATS
 	DECLARE_SCOPE_HIERARCHICAL_COUNTER_FUNC()
 	SCOPE_CYCLE_COUNTER(STAT_RigLogic_CalculateRBFControlsTime);
+#endif  // STATS
+#if CPUPROFILERTRACE_ENABLED
 	TRACE_CPUPROFILER_EVENT_SCOPE_STR("FRigLogic::CalculateRBFControls");
+#endif  // CPUPROFILERTRACE_ENABLED
 
 	RigLogic->calculateRBFControls(Instance->Unwrap(), SolverIndex);
 }
 
 void FRigLogic::CalculateJoints(FRigInstance* Instance) const
 {
+#if STATS
 	DECLARE_SCOPE_HIERARCHICAL_COUNTER_FUNC()
 	SCOPE_CYCLE_COUNTER(STAT_RigLogic_CalculateJointsTime);
+#endif  // STATS
+#if CPUPROFILERTRACE_ENABLED
 	TRACE_CPUPROFILER_EVENT_SCOPE_STR("FRigLogic::CalculateJoints");
+#endif  // CPUPROFILERTRACE_ENABLED
 
 	RigLogic->calculateJoints(Instance->Unwrap());
 }
 
 void FRigLogic::CalculateJoints(FRigInstance* Instance, uint16 JointGroupIndex) const
 {
+#if STATS
 	DECLARE_SCOPE_HIERARCHICAL_COUNTER_FUNC()
 	SCOPE_CYCLE_COUNTER(STAT_RigLogic_CalculateJointsTime);
+#endif  // STATS
+#if CPUPROFILERTRACE_ENABLED
 	TRACE_CPUPROFILER_EVENT_SCOPE_STR("FRigLogic::CalculateJoints");
+#endif  // CPUPROFILERTRACE_ENABLED
 
 	RigLogic->calculateJoints(Instance->Unwrap(), JointGroupIndex);
 }
 
 void FRigLogic::CalculateBlendShapes(FRigInstance* Instance) const
 {
+#if STATS
 	DECLARE_SCOPE_HIERARCHICAL_COUNTER_FUNC()
 	SCOPE_CYCLE_COUNTER(STAT_RigLogic_CalculateBlendShapesTime);
+#endif  // STATS
+#if CPUPROFILERTRACE_ENABLED
 	TRACE_CPUPROFILER_EVENT_SCOPE_STR("FRigLogic::CalculateBlendShapes");
+#endif  // CPUPROFILERTRACE_ENABLED
 
 	RigLogic->calculateBlendShapes(Instance->Unwrap());
 }
 
 void FRigLogic::CalculateAnimatedMaps(FRigInstance* Instance) const
 {
+#if STATS
 	DECLARE_SCOPE_HIERARCHICAL_COUNTER_FUNC()
 	SCOPE_CYCLE_COUNTER(STAT_RigLogic_CalculateAnimatedMapsTime);
+#endif  // STATS
+#if CPUPROFILERTRACE_ENABLED
 	TRACE_CPUPROFILER_EVENT_SCOPE_STR("FRigLogic::CalculateAnimatedMaps");
+#endif  // CPUPROFILERTRACE_ENABLED
 
 	RigLogic->calculateAnimatedMaps(Instance->Unwrap());
 }
 
 void FRigLogic::Calculate(FRigInstance* Instance) const
 {
+#if STATS
 	DECLARE_SCOPE_HIERARCHICAL_COUNTER_FUNC()
-	QUICK_SCOPE_CYCLE_COUNTER(STAT_FRigLogic_Calculate);
+	QUICK_SCOPE_CYCLE_COUNTER(STAT_RigLogic_Calculate);
+#endif  // STATS
+#if CPUPROFILERTRACE_ENABLED
 	TRACE_CPUPROFILER_EVENT_SCOPE_STR("FRigLogic::Calculate");
+#endif  // CPUPROFILERTRACE_ENABLED
 
 	RigLogic->calculate(Instance->Unwrap());
 }
diff --git a/Engine/Plugins/Animation/RigLogic/Source/RigLogicModule/Public/RigLogicModule.h b/Engine/Plugins/Animation/RigLogic/Source/RigLogicModule/Public/RigLogicModule.h
--- a/Engine/Plugins/Animation/RigLogic/Source/RigLogicModule/Public/RigLogicModule.h
+++ b/Engine/Plugins/Animation/RigLogic/Source/RigLogicModule/Public/RigLogicModule.h
@@ -22,8 +22,8 @@ DECLARE_DWORD_COUNTER_STAT_EXTERN(TEXT("Joint Delta Value Count"), STAT_RigLogic
 DECLARE_DWORD_COUNTER_STAT_EXTERN(TEXT("Blend Shape Channel Count"), STAT_RigLogic_BlendShapeChannelCount, STATGROUP_RigLogic, RIGLOGICMODULE_API);
 DECLARE_DWORD_COUNTER_STAT_EXTERN(TEXT("Animated Map Count"), STAT_RigLogic_AnimatedMapCount, STATGROUP_RigLogic, RIGLOGICMODULE_API);
 
-DECLARE_CYCLE_STAT_EXTERN(TEXT("Map GUI To Raw Controls Time"), STAT_FRigLogic_MapGUIToRawControlsTime, STATGROUP_RigLogic, RIGLOGICMODULE_API);
-DECLARE_CYCLE_STAT_EXTERN(TEXT("Map Raw To GUI Controls Time"), STAT_FRigLogic_MapRawToGUIControlsTime, STATGROUP_RigLogic, RIGLOGICMODULE_API);
+DECLARE_CYCLE_STAT_EXTERN(TEXT("Map GUI To Raw Controls Time"), STAT_RigLogic_MapGUIToRawControlsTime, STATGROUP_RigLogic, RIGLOGICMODULE_API);
+DECLARE_CYCLE_STAT_EXTERN(TEXT("Map Raw To GUI Controls Time"), STAT_RigLogic_MapRawToGUIControlsTime, STATGROUP_RigLogic, RIGLOGICMODULE_API);
 DECLARE_CYCLE_STAT_EXTERN(TEXT("Calculate RBF Controls Time"), STAT_RigLogic_CalculateRBFControlsTime, STATGROUP_RigLogic, RIGLOGICMODULE_API);
 DECLARE_CYCLE_STAT_EXTERN(TEXT("Calculate ML Controls Time"), STAT_RigLogic_CalculateMLControlsTime, STATGROUP_RigLogic, RIGLOGICMODULE_API);
 DECLARE_CYCLE_STAT_EXTERN(TEXT("Calculate PSD Controls Time"), STAT_RigLogic_CalculateControlsTime, STATGROUP_RigLogic, RIGLOGICMODULE_API);
