Enable half floats on Windows PC builds

diff --git a/Engine/Platforms/Windows/Plugins/Animation/RigLogic/RigLogic_Windows.uplugin b/Engine/Platforms/Windows/Plugins/Animation/RigLogic/RigLogic_Windows.uplugin
new file mode 100644
--- /dev/null
+++ b/Engine/Platforms/Windows/Plugins/Animation/RigLogic/RigLogic_Windows.uplugin
@@ -0,0 +1,5 @@
+{
+	"FileVersion": 3,
+	"bIsPluginExtension": true
+}
+
diff --git a/Engine/Platforms/Windows/Plugins/Animation/RigLogic/Source/RigLogicLib/RigLogicLib_Windows.Build.cs b/Engine/Platforms/Windows/Plugins/Animation/RigLogic/Source/RigLogicLib/RigLogicLib_Windows.Build.cs
new file mode 100644
--- /dev/null
+++ b/Engine/Platforms/Windows/Plugins/Animation/RigLogic/Source/RigLogicLib/RigLogicLib_Windows.Build.cs
@@ -0,0 +1,12 @@
+// Copyright Epic Games, Inc. All Rights Reserved.
+namespace UnrealBuildTool.Rules
+{
+	public class RigLogicLib_Windows : RigLogicLib
+	{
+		public RigLogicLib_Windows(ReadOnlyTargetRules Target) : base(Target)
+		{
+			PrivateDefinitions.Add("RL_BUILD_WITH_SSE");
+			PrivateDefinitions.Add("RL_BUILD_WITH_HALF_FLOATS");
+		}
+	}
+}
\ No newline at end of file
