Fix for CreateTwoJointNeutralTransforms
Found with [[nodiscard]]

#rnx

diff --git a/Engine/Plugins/Animation/RigLogic/Source/RigLogicModule/Private/RigUnit_RigLogic_Test.cpp b/Engine/Plugins/Animation/RigLogic/Source/RigLogicModule/Private/RigUnit_RigLogic_Test.cpp
--- a/Engine/Plugins/Animation/RigLogic/Source/RigLogicModule/Private/RigUnit_RigLogic_Test.cpp
+++ b/Engine/Plugins/Animation/RigLogic/Source/RigLogicModule/Private/RigUnit_RigLogic_Test.cpp
@@ -192,15 +192,15 @@ TArrayView<const float> FRigUnit_RigLogic::TestAccessor::CreateTwoJointNeutralTr
 	float* Transform1Ptr = InValueArray;
 
 	FTransform Joint1Transform;
-	Joint1Transform.TransformRotation(FQuat::MakeFromEuler(FVector(1.0f, 0.0f, 0.0f)));
-	Joint1Transform.TransformPosition(FVector(1.0f, 0.0f, 0.0f));
+	Joint1Transform.SetRotation(Joint1Transform.TransformRotation(FQuat::MakeFromEuler(FVector(1.0f, 0.0f, 0.0f))));
+	Joint1Transform.SetTranslation(Joint1Transform.TransformPosition(FVector(1.0f, 0.0f, 0.0f)));
 	Joint1Transform.SetScale3D(FVector(1.0f, 1.0f, 0.0f));
 	AddToTransformArray(Transform1Ptr, Joint1Transform);
 
 	float* Transform2Ptr = InValueArray + MAX_ATTRS_PER_JOINT;
 	FTransform Joint2Transform;
-	Joint2Transform.TransformRotation(FQuat::MakeFromEuler(FVector(1.0f, 0.0f, 0.0f)));
-	Joint2Transform.TransformPosition(FVector(1.0f, 0.0f, 0.0f));
+	Joint2Transform.SetRotation(Joint2Transform.TransformRotation(FQuat::MakeFromEuler(FVector(1.0f, 0.0f, 0.0f))));
+	Joint2Transform.SetTranslation(Joint2Transform.TransformPosition(FVector(1.0f, 0.0f, 0.0f)));
 	Joint2Transform.SetScale3D(FVector(1.0f, 1.0f, 0.0f));
 	AddToTransformArray(Transform2Ptr, Joint2Transform);
 
