Fix FRigUnit_RigLogic expected test results

diff --git a/Engine/Plugins/Animation/RigLogic/Source/RigLogicModule/Private/RigUnit_RigLogic_Test.cpp b/Engine/Plugins/Animation/RigLogic/Source/RigLogicModule/Private/RigUnit_RigLogic_Test.cpp
--- a/Engine/Plugins/Animation/RigLogic/Source/RigLogicModule/Private/RigUnit_RigLogic_Test.cpp
+++ b/Engine/Plugins/Animation/RigLogic/Source/RigLogicModule/Private/RigUnit_RigLogic_Test.cpp
@@ -192,14 +192,14 @@ TArrayView<const float> FRigUnit_RigLogic::TestAccessor::CreateTwoJointNeutralTr
 	float* Transform1Ptr = InValueArray;
 
 	FTransform Joint1Transform;
-	Joint1Transform.SetRotation(Joint1Transform.TransformRotation(FQuat::MakeFromEuler(FVector(1.0f, 0.0f, 0.0f))));
+	Joint1Transform.SetRotation(Joint1Transform.TransformRotation(FQuat::MakeFromEuler(FVector(90.0f, 0.0f, 0.0f))));
 	Joint1Transform.SetTranslation(Joint1Transform.TransformPosition(FVector(1.0f, 0.0f, 0.0f)));
 	Joint1Transform.SetScale3D(FVector(1.0f, 1.0f, 0.0f));
 	AddToTransformArray(Transform1Ptr, Joint1Transform);
 
 	float* Transform2Ptr = InValueArray + MAX_ATTRS_PER_JOINT;
 	FTransform Joint2Transform;
-	Joint2Transform.SetRotation(Joint2Transform.TransformRotation(FQuat::MakeFromEuler(FVector(1.0f, 0.0f, 0.0f))));
+	Joint2Transform.SetRotation(Joint2Transform.TransformRotation(FQuat::MakeFromEuler(FVector(90.0f, 0.0f, 0.0f))));
 	Joint2Transform.SetTranslation(Joint2Transform.TransformPosition(FVector(1.0f, 0.0f, 0.0f)));
 	Joint2Transform.SetScale3D(FVector(1.0f, 1.0f, 0.0f));
 	AddToTransformArray(Transform2Ptr, Joint2Transform);
@@ -785,8 +785,10 @@ IMPLEMENT_RIGUNIT_AUTOMATION_TEST(FRigUnit_RigLogic)
 	Test.Exec_UpdateJoints(TestHierarchyTwoBones.Get(), TArrayView<const float>{NeutralValues}, TArrayView<const float>{DeltaTransformData});
 	//Assert
 	//Note that BoneB.GlobalTransform.Z should be zero since the scale.Z is zero. Also, translation Y becomes -Y 
-	AddErrorIfFalse(TestHierarchyTwoBones->GetGlobalTransform(0).GetTranslation().Equals(FVector(1.f, 0.f, 0.f)), TEXT("UpdateJoints LOD0 Bone 01 - unexpected transform"));
-	AddErrorIfFalse(TestHierarchyTwoBones->GetGlobalTransform(1).GetTranslation().Equals(FVector(2.f, 2.f, 0.f)), TEXT("UpdateJoints LOD0 Bone 02 - unexpected transform"));
+	AddErrorIfFalse(TestHierarchyTwoBones->GetLocalTransform(0).GetTranslation().Equals(FVector(2.f, 0.f, 0.f)), TEXT("UpdateJoints LOD0 Bone 01 Local - unexpected transform"));
+	AddErrorIfFalse(TestHierarchyTwoBones->GetLocalTransform(1).GetTranslation().Equals(FVector(2.f, 2.f, 7.f)), TEXT("UpdateJoints LOD0 Bone 02 Local - unexpected transform"));
+	AddErrorIfFalse(TestHierarchyTwoBones->GetGlobalTransform(0).GetTranslation().Equals(FVector(2.f, 0.f, 0.f)), TEXT("UpdateJoints LOD0 Bone 01 Global - unexpected transform"));
+	AddErrorIfFalse(TestHierarchyTwoBones->GetGlobalTransform(1).GetTranslation().Equals(FVector(4.f, 0.f, -2.f)), TEXT("UpdateJoints LOD0 Bone 02 Global - unexpected transform"));
 
 	return true;
 }
