RigLogic DNA Importer: Exposing import/reimport functionality to blueprints/python

diff --git a/Engine/Plugins/Animation/RigLogic/Source/RigLogicEditor/Private/DNAAssetImportFactory.cpp b/Engine/Plugins/Animation/RigLogic/Source/RigLogicEditor/Private/DNAAssetImportFactory.cpp
--- a/Engine/Plugins/Animation/RigLogic/Source/RigLogicEditor/Private/DNAAssetImportFactory.cpp
+++ b/Engine/Plugins/Animation/RigLogic/Source/RigLogicEditor/Private/DNAAssetImportFactory.cpp
@@ -156,7 +156,7 @@ UObject* UDNAAssetImportFactory::FactoryCreateFile(UClass* InClass, UObject* InP
 	}
 
 	// Show the import dialog only when not in a "yes to all" state or when automating import
-	bool bIsAutomated = IsAutomatedImport();
+	bool bIsAutomated = IsAutomatedImport() || IsAutomatedReimport();
 	bool bShowImportDialog = !bIsAutomated;
 
 	ImportOptions = GetImportOptions(DNAImporter, ImportUI, bShowImportDialog, bIsAutomated, InParent->GetPathName(), bOutOperationCanceled, UFactory::CurrentFilename);
diff --git a/Engine/Plugins/Animation/RigLogic/Source/RigLogicEditor/Private/DNAImporterLibrary.cpp b/Engine/Plugins/Animation/RigLogic/Source/RigLogicEditor/Private/DNAImporterLibrary.cpp
new file mode 100644
--- /dev/null
+++ b/Engine/Plugins/Animation/RigLogic/Source/RigLogicEditor/Private/DNAImporterLibrary.cpp
@@ -0,0 +1,23 @@
+﻿// Copyright Epic Games, Inc. All Rights Reserved.
+
+#include "DNAImporterLibrary.h"
+#include "DNAAssetImportFactory.h"
+
+DEFINE_LOG_CATEGORY_STATIC(LogDNAImporterLibrary, Log, All);
+
+#define LOCTEXT_NAMESPACE "RigLogicEditor"
+
+void UDNAImporterLibrary::ImportSkeletalMeshDNA(const FString FileName, UObject* Mesh)
+{
+	UDNAAssetImportFactory* Factory = Cast<UDNAAssetImportFactory>(UDNAAssetImportFactory::StaticClass()->GetDefaultObject());
+	//Reimport will do the same thing as import we just won't get problems when having the same DNA name as SkeletalMesh, new DNA gets initialized anyway
+	bool bSuccess = FReimportManager::Instance()->Reimport(Mesh, false, false, FileName, Factory, INDEX_NONE, false, true, false);
+	
+	if (!bSuccess)
+	{
+		const FText Message = LOCTEXT("DNA_ReimportFailedMessage", "Reimporting of DNA failed");
+		UE_LOG(LogDNAImporterLibrary, Error, TEXT("%s"), *Message.ToString());
+	}
+}
+
+#undef LOCTEXT_NAMESPACE
diff --git a/Engine/Plugins/Animation/RigLogic/Source/RigLogicEditor/Public/DNAImporterLibrary.h b/Engine/Plugins/Animation/RigLogic/Source/RigLogicEditor/Public/DNAImporterLibrary.h
new file mode 100644
--- /dev/null
+++ b/Engine/Plugins/Animation/RigLogic/Source/RigLogicEditor/Public/DNAImporterLibrary.h
@@ -0,0 +1,20 @@
+﻿// Copyright Epic Games, Inc. All Rights Reserved.
+
+#pragma once
+
+#include "CoreMinimal.h"
+#include "Kismet/BlueprintFunctionLibrary.h"
+#include "DNAImporterLibrary.generated.h"
+
+/**
+ * 
+ */
+UCLASS()
+class RIGLOGICEDITOR_API UDNAImporterLibrary : public UBlueprintFunctionLibrary
+{
+	GENERATED_BODY()
+	
+	/** Import a DNA file onto the specified mesh asset.*/
+	UFUNCTION(CallInEditor, BlueprintCallable, meta = (DisplayName = "Import Skeletal Mesh DNA File", Keywords = "DNA Importer Skeletal Mesh"), Category = "DNA")
+	static void ImportSkeletalMeshDNA(const FString FileName, UObject* Mesh);
+};
