Follow up: ControlRig: Remove use of access to execute context in unit tests

#rnx

diff --git a/Engine/Plugins/Animation/RigLogic/Source/RigLogicModule/Private/RigUnit_RigLogic_Test.cpp b/Engine/Plugins/Animation/RigLogic/Source/RigLogicModule/Private/RigUnit_RigLogic_Test.cpp
--- a/Engine/Plugins/Animation/RigLogic/Source/RigLogicModule/Private/RigUnit_RigLogic_Test.cpp
+++ b/Engine/Plugins/Animation/RigLogic/Source/RigLogicModule/Private/RigUnit_RigLogic_Test.cpp
@@ -275,7 +275,6 @@ IMPLEMENT_RIGUNIT_AUTOMATION_TEST(FRigUnit_RigLogic)
 	}
 
 	Test.GetData()->SkelMeshComponent = SkelMeshComponent;
-	Unit.ExecuteContext.Hierarchy = Hierarchy.Get();
 
 	//Test
 	Execute();
