[Animation] Update Tooltip

diff --git a/Engine/Plugins/Animation/RigLogic/Source/RigLogicModule/Public/AnimNode_RigLogic.h b/Engine/Plugins/Animation/RigLogic/Source/RigLogicModule/Public/AnimNode_RigLogic.h
--- a/Engine/Plugins/Animation/RigLogic/Source/RigLogicModule/Public/AnimNode_RigLogic.h
+++ b/Engine/Plugins/Animation/RigLogic/Source/RigLogicModule/Public/AnimNode_RigLogic.h
@@ -97,9 +97,9 @@ private:
 
 private:
 	/*
-	 * Max LOD level that post-process AnimBPs are evaluated.
-	 * For example if you have the threshold set to 2, it will evaluate until including LOD 2 (based on 0 index). In case the LOD level gets set to 3, it will stop evaluating the post-process AnimBP.
-	 * Setting it to -1 will always evaluate it and disable LODing.
+	 * Max LOD level that RigLogic Node is evaluated.
+	 * For example if you have the threshold set to 2, it will evaluate until including LOD 2 (based on 0 index). In case the LOD level gets set to 3, it will stop evaluating the Rig Logic.
+	 * Setting it to -1 will always evaluate it.
 	 */
 	UPROPERTY(EditAnywhere, Category = RigLogic, meta = (PinHiddenByDefault))
 	int32 LODThreshold = INDEX_NONE;
