[Arm64EC] Fix RigLogic platform detection to only define one  platform architecture, favoring arm64 (it was detecting arm64ec as both X86 and Arm)

#rnx

diff --git a/Engine/Plugins/Animation/RigLogic/Source/RigLogicLib/Public/trimd/Platform.h b/Engine/Plugins/Animation/RigLogic/Source/RigLogicLib/Public/trimd/Platform.h
--- a/Engine/Plugins/Animation/RigLogic/Source/RigLogicLib/Public/trimd/Platform.h
+++ b/Engine/Plugins/Animation/RigLogic/Source/RigLogicLib/Public/trimd/Platform.h
@@ -3,16 +3,15 @@
 #pragma once
 
 // *INDENT-OFF*
-#if defined(i386) || defined(__i386) || defined(__i386__) || defined(__IA32__) || defined(_M_IX86) || defined(__X86__) || \
+#if defined(__arm__) || defined(__aarch64__) || defined(_M_ARM64) || defined(_M_ARM64EC)
+    #define TRIMD_PLATFORM_ARM 1
+
+#elif defined(i386) || defined(__i386) || defined(__i386__) || defined(__IA32__) || defined(_M_IX86) || defined(__X86__) || \
     defined(_X86_) || defined(__THW_INTEL__) || defined(__I86__) || defined(__amd64__) || defined(__amd64) || \
     defined(__x86_64__) || defined(__x86_64) || defined(_M_X64) || defined(_M_AMD64)
     #define TRIMD_PLATFORM_X86 1
 #endif
 
-#if defined(__arm__) || defined(__aarch64__) || defined(_M_ARM64) || defined(_M_ARM64EC)
-    #define TRIMD_PLATFORM_ARM 1
-#endif
-
 namespace trimd {
 
 struct CPUFeatures {
