[AutoRTFM] Allow RigLogicLib to be called from within a Verse transaction.

This can occur when `ACinematicSequenceDeviceBase::GoToEndAndStop` is called.

diff --git a/Engine/Plugins/Animation/RigLogic/Source/RigLogicLib/RigLogicLib.Build.cs b/Engine/Plugins/Animation/RigLogic/Source/RigLogicLib/RigLogicLib.Build.cs
--- a/Engine/Plugins/Animation/RigLogic/Source/RigLogicLib/RigLogicLib.Build.cs
+++ b/Engine/Plugins/Animation/RigLogic/Source/RigLogicLib/RigLogicLib.Build.cs
@@ -55,8 +55,6 @@ namespace UnrealBuildTool.Rules
 
 			PrivateDefinitions.Add("RL_AUTODETECT_SSE=1");
 			PrivateDefinitions.Add("RL_AUTODETECT_HALF_FLOATS=1");
-
-			bDisableAutoRTFMInstrumentation = true;
 		}
 	}
 }
