Fixed driver joint controls of RigLogic

diff --git a/Engine/Plugins/Animation/RigLogic/Source/RigLogicModule/Private/DNAIndexMapping.cpp b/Engine/Plugins/Animation/RigLogic/Source/RigLogicModule/Private/DNAIndexMapping.cpp
--- a/Engine/Plugins/Animation/RigLogic/Source/RigLogicModule/Private/DNAIndexMapping.cpp
+++ b/Engine/Plugins/Animation/RigLogic/Source/RigLogicModule/Private/DNAIndexMapping.cpp
@@ -96,7 +96,9 @@ void FDNAIndexMapping::MapDriverJoints(const IDNAReader* DNAReader, const USkele
 		{
 			continue;
 		}
-		const FString DriverJointName = DriverJointAttrName.Mid(0, DriverJointAttrName.Len() - 2);
+		int32 IndexOfLastDot;
+		DriverJointAttrName.FindLastChar('.', IndexOfLastDot);
+		const FString DriverJointName = DriverJointAttrName.Mid(0, IndexOfLastDot);
 		const FName BoneName = FName(*DriverJointName);
 		const int32 BoneIndex = RefSkeleton.FindBoneIndex(BoneName);
 		if (BoneIndex == INDEX_NONE)
@@ -120,19 +122,19 @@ void FDNAIndexMapping::MapDriverJoints(const IDNAReader* DNAReader, const USkele
 		FMeshPoseBoneControlAttributeMapping& Mapping = DriverJointsToControlAttributesMap[MappingIndex];
 		Mapping.DNAJointIndex = JointsMapDNAIndicesToMeshPoseBoneIndices.Find(Mapping.MeshPoseBoneIndex);
 
-		if (DriverJointAttrName.EndsWith(TEXT(".x")))
+		if (DriverJointAttrName.EndsWith(TEXT(".x")) || DriverJointAttrName.EndsWith(TEXT(".qx")))
 		{
 			Mapping.RotationX = ControlIndex;
 		}
-		else if (DriverJointAttrName.EndsWith(TEXT(".y")))
+		else if (DriverJointAttrName.EndsWith(TEXT(".y")) || DriverJointAttrName.EndsWith(TEXT(".qy")))
 		{
 			Mapping.RotationY = ControlIndex;
 		}
-		else if (DriverJointAttrName.EndsWith(TEXT(".z")))
+		else if (DriverJointAttrName.EndsWith(TEXT(".z")) || DriverJointAttrName.EndsWith(TEXT(".qz")))
 		{
 			Mapping.RotationZ = ControlIndex;
 		}
-		else if (DriverJointAttrName.EndsWith(TEXT(".w")))
+		else if (DriverJointAttrName.EndsWith(TEXT(".w")) || DriverJointAttrName.EndsWith(TEXT(".qw")))
 		{
 			Mapping.RotationW = ControlIndex;
 		}
