Making DNAReader Public and adding RBF output caching

diff --git a/Engine/Plugins/Animation/RigLogic/Source/RigLogicModule/Private/RigLogicDNAReader.cpp b/Engine/Plugins/Animation/RigLogic/Source/RigLogicModule/Private/RigLogicDNAReader.cpp
--- a/Engine/Plugins/Animation/RigLogic/Source/RigLogicModule/Private/RigLogicDNAReader.cpp
+++ b/Engine/Plugins/Animation/RigLogic/Source/RigLogicModule/Private/RigLogicDNAReader.cpp
@@ -568,7 +568,8 @@ dna::ConstArrayView<std::uint16_t> RigLogicDNAReader::getRBFPoseAnimatedMapOutpu
 }
 dna::ConstArrayView<float> RigLogicDNAReader::getRBFPoseJointOutputValues(std::uint16_t poseIndex) const
 {
-	return Reader->getRBFPoseJointOutputValues(poseIndex);
+	CacheRBFPoseJointOutputValues(poseIndex);
+	return {RBFValues.GetData(), static_cast<std::size_t>(RBFValues.Num())};
 }
 float RigLogicDNAReader::getRBFPoseScale(std::uint16_t poseIndex) const
 {
@@ -800,6 +801,39 @@ void RigLogicDNAReader::CacheJointGroup(std::uint16_t JointGroupIndex) const
 	Id = JointGroupIndex;
 }
 
+void RigLogicDNAReader::CacheRBFPoseJointOutputValues(std::uint16_t PoseIndex) const
+{
+	if ((Key == CachedDataKey::RBFPoseJointOutputValues) && (Id == PoseIndex))
+	{
+		// Already cached
+		return;
+	}
+
+	const auto OutputIndices = Reader->getRBFPoseJointOutputIndices(PoseIndex);
+	const auto OriginalValues = Reader->getRBFPoseJointOutputValues(PoseIndex);
+	const auto RowCount = OutputIndices.size();
+	const auto ColCount = (RowCount == 0) ? 0 : OriginalValues.size() / RowCount;
+
+	RBFValues.Reset(OriginalValues.size());
+	RBFValues.Append(OriginalValues.data(), OriginalValues.size());
+	for (std::size_t Row = 0; Row < RowCount; ++Row)
+	{
+		static constexpr auto JointAttrCount = 9;
+		const auto RelAttrIndex = OutputIndices[Row] % JointAttrCount;
+		// Translation Ys, Rotation Xs and Rotation Zs only need to be flipped
+		if ((RelAttrIndex == 1) || (RelAttrIndex == 3) || (RelAttrIndex == 5))
+		{
+			for (std::size_t Col = 0; Col < ColCount; ++Col)
+			{
+				RBFValues[Row * ColCount + Col] = -RBFValues[Row * ColCount + Col];
+			}
+		}
+	}
+
+	Key = CachedDataKey::RBFPoseJointOutputValues;
+	Id = PoseIndex;
+}
+
 void RigLogicDNAReader::CacheRBFSolverRawControlValues(std::uint16_t SolverIndex) const
 {
 	if ((Key == CachedDataKey::RBFSolverRawControlValues) && (Id == SolverIndex))
diff --git a/Engine/Plugins/Animation/RigLogic/Source/RigLogicModule/Private/RigLogicDNAReader.h b/Engine/Plugins/Animation/RigLogic/Source/RigLogicModule/Private/RigLogicDNAReader.h
deleted file mode 100644
diff --git a/Engine/Plugins/Animation/RigLogic/Source/RigLogicModule/Public/RigLogicDNAReader.h b/Engine/Plugins/Animation/RigLogic/Source/RigLogicModule/Public/RigLogicDNAReader.h
new file mode 100644
--- /dev/null
+++ b/Engine/Plugins/Animation/RigLogic/Source/RigLogicModule/Public/RigLogicDNAReader.h
@@ -0,0 +1,224 @@
+// Copyright Epic Games, Inc. All Rights Reserved.
+
+#pragma once
+
+#include "CoreMinimal.h"
+
+#include <dna/Reader.h>
+
+class RIGLOGICMODULE_API RigLogicDNAReader : public dna::Reader {
+public:
+	explicit RigLogicDNAReader(const dna::Reader* DNAReader);
+
+	// Header
+	std::uint16_t getFileFormatGeneration() const override;
+	std::uint16_t getFileFormatVersion() const override;
+
+	// Descriptor
+	dna::StringView getName() const override;
+	dna::Archetype getArchetype() const override;
+	dna::Gender getGender() const override;
+	std::uint16_t getAge() const override;
+	std::uint32_t getMetaDataCount() const override;
+	dna::StringView getMetaDataKey(std::uint32_t index) const override;
+	dna::StringView getMetaDataValue(const char* key) const override;
+	dna::TranslationUnit getTranslationUnit() const override;
+	dna::RotationUnit getRotationUnit() const override;
+	dna::CoordinateSystem getCoordinateSystem() const override;
+	std::uint16_t getLODCount() const override;
+	std::uint16_t getDBMaxLOD() const override;
+	dna::StringView getDBComplexity() const override;
+	dna::StringView getDBName() const override;
+
+	// Definition
+	std::uint16_t getGUIControlCount() const override;
+	dna::StringView getGUIControlName(std::uint16_t index) const override;
+	std::uint16_t getRawControlCount() const override;
+	dna::StringView getRawControlName(std::uint16_t index) const override;
+	std::uint16_t getJointCount() const override;
+	dna::StringView getJointName(std::uint16_t index) const override;
+	std::uint16_t getJointIndexListCount() const override;
+	dna::ConstArrayView<std::uint16_t> getJointIndicesForLOD(std::uint16_t lod) const override;
+	std::uint16_t getJointParentIndex(std::uint16_t index) const override;
+	std::uint16_t getBlendShapeChannelCount() const override;
+	dna::StringView getBlendShapeChannelName(std::uint16_t index) const override;
+	std::uint16_t getBlendShapeChannelIndexListCount() const override;
+	dna::ConstArrayView<std::uint16_t> getBlendShapeChannelIndicesForLOD(std::uint16_t lod) const override;
+	std::uint16_t getAnimatedMapCount() const override;
+	dna::StringView getAnimatedMapName(std::uint16_t index) const override;
+	std::uint16_t getAnimatedMapIndexListCount() const override;
+	dna::ConstArrayView<std::uint16_t> getAnimatedMapIndicesForLOD(std::uint16_t lod) const override;
+	std::uint16_t getMeshCount() const override;
+	dna::StringView getMeshName(std::uint16_t index) const override;
+	std::uint16_t getMeshIndexListCount() const override;
+	dna::ConstArrayView<std::uint16_t> getMeshIndicesForLOD(std::uint16_t lod) const override;
+	std::uint16_t getMeshBlendShapeChannelMappingCount() const override;
+	dna::MeshBlendShapeChannelMapping getMeshBlendShapeChannelMapping(std::uint16_t index) const override;
+	dna::ConstArrayView<std::uint16_t> getMeshBlendShapeChannelMappingIndicesForLOD(std::uint16_t lod) const override;
+	dna::Vector3 getNeutralJointTranslation(std::uint16_t index) const override;
+	dna::ConstArrayView<float> getNeutralJointTranslationXs() const override;
+	dna::ConstArrayView<float> getNeutralJointTranslationYs() const override;
+	dna::ConstArrayView<float> getNeutralJointTranslationZs() const override;
+	dna::Vector3 getNeutralJointRotation(std::uint16_t index) const override;
+	dna::ConstArrayView<float> getNeutralJointRotationXs() const override;
+	dna::ConstArrayView<float> getNeutralJointRotationYs() const override;
+	dna::ConstArrayView<float> getNeutralJointRotationZs() const override;
+
+	// Behavior
+	dna::ConstArrayView<std::uint16_t> getGUIToRawInputIndices() const override;
+	dna::ConstArrayView<std::uint16_t> getGUIToRawOutputIndices() const override;
+	dna::ConstArrayView<float> getGUIToRawFromValues() const override;
+	dna::ConstArrayView<float> getGUIToRawToValues() const override;
+	dna::ConstArrayView<float> getGUIToRawSlopeValues() const override;
+	dna::ConstArrayView<float> getGUIToRawCutValues() const override;
+	std::uint16_t getPSDCount() const override;
+	dna::ConstArrayView<std::uint16_t> getPSDRowIndices() const override;
+	dna::ConstArrayView<std::uint16_t> getPSDColumnIndices() const override;
+	dna::ConstArrayView<float> getPSDValues() const override;
+	std::uint16_t getJointRowCount() const override;
+	std::uint16_t getJointColumnCount() const override;
+	dna::ConstArrayView<std::uint16_t> getJointVariableAttributeIndices(std::uint16_t lod) const override;
+	std::uint16_t getJointGroupCount() const override;
+	dna::ConstArrayView<std::uint16_t> getJointGroupLODs(std::uint16_t jointGroupIndex) const override;
+	dna::ConstArrayView<std::uint16_t> getJointGroupInputIndices(std::uint16_t jointGroupIndex) const override;
+	dna::ConstArrayView<std::uint16_t> getJointGroupOutputIndices(std::uint16_t jointGroupIndex) const override;
+	dna::ConstArrayView<float> getJointGroupValues(std::uint16_t jointGroupIndex) const override;
+	dna::ConstArrayView<std::uint16_t> getJointGroupJointIndices(std::uint16_t jointGroupIndex) const override;
+	dna::ConstArrayView<std::uint16_t> getBlendShapeChannelLODs() const override;
+	dna::ConstArrayView<std::uint16_t> getBlendShapeChannelInputIndices() const override;
+	dna::ConstArrayView<std::uint16_t> getBlendShapeChannelOutputIndices() const override;
+	dna::ConstArrayView<std::uint16_t> getAnimatedMapLODs() const override;
+	dna::ConstArrayView<std::uint16_t> getAnimatedMapInputIndices() const override;
+	dna::ConstArrayView<std::uint16_t> getAnimatedMapOutputIndices() const override;
+	dna::ConstArrayView<float> getAnimatedMapFromValues() const override;
+	dna::ConstArrayView<float> getAnimatedMapToValues() const override;
+	dna::ConstArrayView<float> getAnimatedMapSlopeValues() const override;
+	dna::ConstArrayView<float> getAnimatedMapCutValues() const override;
+
+	// Geometry
+	std::uint32_t getVertexPositionCount(std::uint16_t meshIndex) const override;
+	dna::Position getVertexPosition(std::uint16_t meshIndex, std::uint32_t vertexIndex) const override;
+	dna::ConstArrayView<float> getVertexPositionXs(std::uint16_t meshIndex) const override;
+	dna::ConstArrayView<float> getVertexPositionYs(std::uint16_t meshIndex) const override;
+	dna::ConstArrayView<float> getVertexPositionZs(std::uint16_t meshIndex) const override;
+	std::uint32_t getVertexTextureCoordinateCount(std::uint16_t meshIndex) const override;
+	dna::TextureCoordinate getVertexTextureCoordinate(std::uint16_t meshIndex, std::uint32_t textureCoordinateIndex) const override;
+	dna::ConstArrayView<float> getVertexTextureCoordinateUs(std::uint16_t meshIndex) const override;
+	dna::ConstArrayView<float> getVertexTextureCoordinateVs(std::uint16_t meshIndex) const override;
+	std::uint32_t getVertexNormalCount(std::uint16_t meshIndex) const override;
+	dna::Normal getVertexNormal(std::uint16_t meshIndex, std::uint32_t normalIndex) const override;
+	dna::ConstArrayView<float> getVertexNormalXs(std::uint16_t meshIndex) const override;
+	dna::ConstArrayView<float> getVertexNormalYs(std::uint16_t meshIndex) const override;
+	dna::ConstArrayView<float> getVertexNormalZs(std::uint16_t meshIndex) const override;
+	std::uint32_t getVertexLayoutCount(std::uint16_t meshIndex) const override;
+	dna::VertexLayout getVertexLayout(std::uint16_t meshIndex, std::uint32_t layoutIndex) const override;
+	dna::ConstArrayView<std::uint32_t> getVertexLayoutPositionIndices(std::uint16_t meshIndex) const override;
+	dna::ConstArrayView<std::uint32_t> getVertexLayoutTextureCoordinateIndices(std::uint16_t meshIndex) const override;
+	dna::ConstArrayView<std::uint32_t> getVertexLayoutNormalIndices(std::uint16_t meshIndex) const override;
+	std::uint32_t getFaceCount(std::uint16_t meshIndex) const override;
+	dna::ConstArrayView<std::uint32_t> getFaceVertexLayoutIndices(std::uint16_t meshIndex, std::uint32_t faceIndex) const override;
+	std::uint16_t getMaximumInfluencePerVertex(std::uint16_t meshIndex) const override;
+	std::uint32_t getSkinWeightsCount(std::uint16_t meshIndex) const override;
+	dna::ConstArrayView<float> getSkinWeightsValues(std::uint16_t meshIndex, std::uint32_t vertexIndex) const override;
+	dna::ConstArrayView<std::uint16_t> getSkinWeightsJointIndices(std::uint16_t meshIndex, std::uint32_t vertexIndex) const override;
+	std::uint16_t getBlendShapeTargetCount(std::uint16_t meshIndex) const override;
+	std::uint16_t getBlendShapeChannelIndex(std::uint16_t meshIndex, std::uint16_t blendShapeTargetIndex) const override;
+	std::uint32_t getBlendShapeTargetDeltaCount(std::uint16_t meshIndex, std::uint16_t blendShapeTargetIndex) const override;
+	dna::Delta getBlendShapeTargetDelta(std::uint16_t meshIndex, std::uint16_t blendShapeTargetIndex, std::uint32_t deltaIndex) const override;
+	dna::ConstArrayView<float> getBlendShapeTargetDeltaXs(std::uint16_t meshIndex, std::uint16_t blendShapeTargetIndex) const override;
+	dna::ConstArrayView<float> getBlendShapeTargetDeltaYs(std::uint16_t meshIndex, std::uint16_t blendShapeTargetIndex) const override;
+	dna::ConstArrayView<float> getBlendShapeTargetDeltaZs(std::uint16_t meshIndex, std::uint16_t blendShapeTargetIndex) const override;
+	dna::ConstArrayView<std::uint32_t> getBlendShapeTargetVertexIndices(std::uint16_t meshIndex, std::uint16_t blendShapeTargetIndex) const override;
+
+	// Machine Learned Behavior
+	std::uint16_t getMLControlCount() const override;
+	dna::StringView getMLControlName(std::uint16_t index) const override;
+	std::uint16_t getNeuralNetworkCount() const override;
+	std::uint16_t getNeuralNetworkIndexListCount() const override;
+	dna::ConstArrayView<std::uint16_t> getNeuralNetworkIndicesForLOD(std::uint16_t lod) const override;
+	std::uint16_t getMeshRegionCount(std::uint16_t meshIndex) const override;
+	dna::StringView getMeshRegionName(std::uint16_t meshIndex, std::uint16_t regionIndex) const override;
+	dna::ConstArrayView<std::uint16_t> getNeuralNetworkIndicesForMeshRegion(std::uint16_t meshIndex, std::uint16_t regionIndex) const override;
+	dna::ConstArrayView<std::uint16_t> getNeuralNetworkInputIndices(std::uint16_t netIndex) const override;
+	dna::ConstArrayView<std::uint16_t> getNeuralNetworkOutputIndices(std::uint16_t netIndex) const override;
+	std::uint16_t getNeuralNetworkLayerCount(std::uint16_t netIndex) const override;
+	dna::ActivationFunction getNeuralNetworkLayerActivationFunction(std::uint16_t netIndex, std::uint16_t layerIndex) const override;
+	dna::ConstArrayView<float> getNeuralNetworkLayerActivationFunctionParameters(std::uint16_t netIndex, std::uint16_t layerIndex) const override;
+	dna::ConstArrayView<float> getNeuralNetworkLayerBiases(std::uint16_t netIndex, std::uint16_t layerIndex) const override;
+	dna::ConstArrayView<float> getNeuralNetworkLayerWeights(std::uint16_t netIndex, std::uint16_t layerIndex) const override;
+
+	// RBFBehaviorReader methods
+	std::uint16_t getRBFPoseCount() const override;
+	dna::StringView getRBFPoseName(std::uint16_t poseIndex) const override;
+	dna::ConstArrayView<std::uint16_t> getRBFPoseJointOutputIndices(std::uint16_t poseIndex) const override;
+	dna::ConstArrayView<std::uint16_t> getRBFPoseBlendShapeChannelOutputIndices(std::uint16_t poseIndex) const override;
+	dna::ConstArrayView<std::uint16_t> getRBFPoseAnimatedMapOutputIndices(std::uint16_t poseIndex) const override;
+	dna::ConstArrayView<float> getRBFPoseJointOutputValues(std::uint16_t poseIndex) const override;
+	float getRBFPoseScale(std::uint16_t poseIndex) const override;
+	std::uint16_t getRBFPoseControlCount() const override;
+	dna::StringView getRBFPoseControlName(std::uint16_t poseControlIndex) const override;
+	dna::ConstArrayView<std::uint16_t> getRBFPoseInputControlIndices(std::uint16_t poseIndex) const override;
+	dna::ConstArrayView<std::uint16_t> getRBFPoseOutputControlIndices(std::uint16_t poseIndex) const override;
+	dna::ConstArrayView<float> getRBFPoseOutputControlWeights(std::uint16_t poseIndex) const override;
+	std::uint16_t getRBFSolverCount() const override;
+	std::uint16_t getRBFSolverIndexListCount() const override;
+	dna::ConstArrayView<std::uint16_t> getRBFSolverIndicesForLOD(std::uint16_t lod) const override;
+	dna::StringView getRBFSolverName(std::uint16_t solverIndex) const override;
+	dna::ConstArrayView<std::uint16_t> getRBFSolverRawControlIndices(std::uint16_t solverIndex) const override;
+	dna::ConstArrayView<std::uint16_t> getRBFSolverPoseIndices(std::uint16_t solverIndex) const override;
+	dna::ConstArrayView<float> getRBFSolverRawControlValues(std::uint16_t solverIndex) const override;
+	dna::RBFSolverType getRBFSolverType(std::uint16_t solverIndex) const override;
+	float getRBFSolverRadius(std::uint16_t solverIndex) const override;
+	dna::AutomaticRadius getRBFSolverAutomaticRadius(std::uint16_t solverIndex) const override;
+	float getRBFSolverWeightThreshold(std::uint16_t solverIndex) const override;
+	dna::RBFDistanceMethod getRBFSolverDistanceMethod(std::uint16_t solverIndex) const override;
+	dna::RBFNormalizeMethod getRBFSolverNormalizeMethod(std::uint16_t solverIndex) const override;
+	dna::RBFFunctionType getRBFSolverFunctionType(std::uint16_t solverIndex) const override;
+	dna::TwistAxis getRBFSolverTwistAxis(std::uint16_t solverIndex) const override;
+
+	// JointBehaviorMetadataReader methods
+	dna::TranslationRepresentation getJointTranslationRepresentation(std::uint16_t jointIndex) const override;
+	dna::RotationRepresentation getJointRotationRepresentation(std::uint16_t jointIndex) const override;
+	dna::ScaleRepresentation getJointScaleRepresentation(std::uint16_t jointIndex) const override;
+
+	// TwistSwingBehaviorReader methods
+	std::uint16_t getTwistCount() const override;
+	dna::TwistAxis getTwistSetupTwistAxis(std::uint16_t twistIndex) const override;
+	dna::ConstArrayView<std::uint16_t> getTwistInputControlIndices(std::uint16_t twistIndex) const override;
+	dna::ConstArrayView<std::uint16_t> getTwistOutputJointIndices(std::uint16_t twistIndex) const override;
+	dna::ConstArrayView<float> getTwistBlendWeights(std::uint16_t twistIndex) const override;
+	std::uint16_t getSwingCount() const override;
+	dna::TwistAxis getSwingSetupTwistAxis(std::uint16_t swingIndex) const override;
+	dna::ConstArrayView<std::uint16_t> getSwingInputControlIndices(std::uint16_t swingIndex) const override;
+	dna::ConstArrayView<std::uint16_t> getSwingOutputJointIndices(std::uint16_t swingIndex) const override;
+	dna::ConstArrayView<float> getSwingBlendWeights(std::uint16_t swingIndex) const override;
+
+	// Reader
+	void unload(dna::DataLayer Layer) override;
+	static void destroy(dna::Reader* Pointer);
+
+private:
+	void CacheNeutralJointTranslations() const;
+	void CacheNeutralJointRotations() const;
+	void CacheJointGroup(std::uint16_t JointGroupIndex) const;
+	void CacheRBFSolverRawControlValues(std::uint16_t SolverIndex) const;
+	void CacheRBFPoseJointOutputValues(std::uint16_t PoseIndex) const;
+
+private:
+	enum class CachedDataKey
+	{
+		None,
+		NeutralJointTranslations,
+		NeutralJointRotations,
+		JointGroup,
+		RBFSolverRawControlValues,
+		RBFPoseJointOutputValues
+	};
+
+private:
+	const dna::Reader* Reader;
+	mutable TArray<float> Values;
+	mutable TArray<float> RBFValues;
+	mutable CachedDataKey Key;
+	mutable std::uint16_t Id;
+};
\ No newline at end of file
