Fix DNAAsset mapping of DNA indices to UE bone and curve indices, such that the mapping cache key is composed from both the SkeletalMesh and Skeleton, not just the SkeletalMesh, and also make sure the containers are sorted during initialization to prevent it happening at runtime which may corrupt the data structure

diff --git a/Engine/Plugins/Animation/RigLogic/Source/RigLogicModule/Private/DNAAsset.cpp b/Engine/Plugins/Animation/RigLogic/Source/RigLogicModule/Private/DNAAsset.cpp
--- a/Engine/Plugins/Animation/RigLogic/Source/RigLogicModule/Private/DNAAsset.cpp
+++ b/Engine/Plugins/Animation/RigLogic/Source/RigLogicModule/Private/DNAAsset.cpp
@@ -232,9 +232,9 @@ TSharedPtr<FDNAIndexMapping> UDNAAsset::GetDNAIndexMapping(const USkeleton* Skel
 	TSharedPtr<FDNAIndexMapping> DNAIndexMapping;
 	for (auto Iterator = DNAIndexMappingContainer.CreateIterator(); Iterator; ++Iterator)
 	{
-		if (Iterator->Key.IsValid())
+		if ((Iterator->Key.SkeletalMesh.IsValid()) && (Iterator->Key.Skeleton.IsValid()))
 		{
-			if (Iterator->Key == SkeletalMesh)
+			if ((Iterator->Key.SkeletalMesh == SkeletalMesh) && (Iterator->Key.Skeleton == Skeleton))
 			{
 				DNAIndexMapping = Iterator->Value;
 			}
@@ -251,7 +251,7 @@ TSharedPtr<FDNAIndexMapping> UDNAAsset::GetDNAIndexMapping(const USkeleton* Skel
 	{
 		DNAIndexMapping = MakeShared<FDNAIndexMapping>();
 		DNAIndexMapping->Init(BehaviorReader.Get(), Skeleton, SkeletalMesh);
-		DNAIndexMappingContainer.Add(SkeletalMesh, DNAIndexMapping);
+		DNAIndexMappingContainer.Add({SkeletalMesh, Skeleton}, DNAIndexMapping);
 	}
 
 	return DNAIndexMapping;
@@ -392,6 +392,7 @@ void UDNAAsset::PostEditChangeProperty(FPropertyChangedEvent& PropertyChangedEve
 	const FName& MemberPropertyName = PropertyChangedEvent.MemberProperty->GetFName();
 	if (MemberPropertyName == GET_MEMBER_NAME_CHECKED(UDNAAsset, RigLogicConfiguration))
 	{
+		FWriteScopeLock DNAScopeLock{DNAUpdateLock};
 		InvalidateRigRuntimeContext();
 		InitializeRigRuntimeContext();
 	}
diff --git a/Engine/Plugins/Animation/RigLogic/Source/RigLogicModule/Private/DNAIndexMapping.cpp b/Engine/Plugins/Animation/RigLogic/Source/RigLogicModule/Private/DNAIndexMapping.cpp
--- a/Engine/Plugins/Animation/RigLogic/Source/RigLogicModule/Private/DNAIndexMapping.cpp
+++ b/Engine/Plugins/Animation/RigLogic/Source/RigLogicModule/Private/DNAIndexMapping.cpp
@@ -49,6 +49,10 @@ void FDNAIndexMapping::MapControlCurves(const IDNAReader* DNAReader, const USkel
 		}
 		ControlAttributeCurves.Add(*AnimatedControlName, ControlIndex);
 	}
+
+	// This has the side-effect of sorting the NamedValueArray, which is necessary to avoid it happening at
+	// runtime, which would cause concurrent mutation of the data structure, not guarded by any locking mechanism
+	UE::Anim::FNamedValueArrayUtils::RemoveByPredicate(ControlAttributeCurves, FCachedIndexedCurve{}, [](const UE::Anim::FCurveElementIndexed&, const UE::Anim::FCurveElementIndexed&) { return true; });
 }
 
 void FDNAIndexMapping::MapNeuralNetworkMaskCurves(const IDNAReader* DNAReader, const USkeleton* Skeleton)
@@ -75,6 +79,10 @@ void FDNAIndexMapping::MapNeuralNetworkMaskCurves(const IDNAReader* DNAReader, c
 			}
 		}
 	}
+
+	// This has the side-effect of sorting the NamedValueArray, which is necessary to avoid it happening at
+	// runtime, which would cause concurrent mutation of the data structure, not guarded by any locking mechanism
+	UE::Anim::FNamedValueArrayUtils::RemoveByPredicate(NeuralNetworkMaskCurves, FCachedIndexedCurve{}, [](const UE::Anim::FCurveElementIndexed&, const UE::Anim::FCurveElementIndexed&) { return true; });
 }
 
 void FDNAIndexMapping::MapDriverJoints(const IDNAReader* DNAReader, const USkeletalMesh* SkeletalMesh)
@@ -191,6 +199,10 @@ void FDNAIndexMapping::MapMorphTargets(const IDNAReader* DNAReader, const USkele
 				MorphTargetCurvesPerLOD[LODIndex].Add(MorphTarget->GetFName(), Mapping.BlendShapeChannelIndex);
 			}
 		}
+
+		// This has the side-effect of sorting the NamedValueArray, which is necessary to avoid it happening at
+		// runtime, which would cause concurrent mutation of the data structure, not guarded by any locking mechanism
+		UE::Anim::FNamedValueArrayUtils::RemoveByPredicate(MorphTargetCurvesPerLOD[LODIndex], FCachedIndexedCurve{}, [](const UE::Anim::FCurveElementIndexed&, const UE::Anim::FCurveElementIndexed&) { return true; });
 	}
 }
 
@@ -220,6 +232,10 @@ void FDNAIndexMapping::MapMaskMultipliers(const IDNAReader* DNAReader, const USk
 
 			MaskMultiplierCurvesPerLOD[LODIndex].Add(*MaskMultiplierNameStr, AnimMapIndex);
 		}
+
+		// This has the side-effect of sorting the NamedValueArray, which is necessary to avoid it happening at
+		// runtime, which would cause concurrent mutation of the data structure, not guarded by any locking mechanism
+		UE::Anim::FNamedValueArrayUtils::RemoveByPredicate(MaskMultiplierCurvesPerLOD[LODIndex], FCachedIndexedCurve{}, [](const UE::Anim::FCurveElementIndexed&, const UE::Anim::FCurveElementIndexed&) { return true; });
 	}
 }
 
diff --git a/Engine/Plugins/Animation/RigLogic/Source/RigLogicModule/Public/DNAAsset.h b/Engine/Plugins/Animation/RigLogic/Source/RigLogicModule/Public/DNAAsset.h
--- a/Engine/Plugins/Animation/RigLogic/Source/RigLogicModule/Public/DNAAsset.h
+++ b/Engine/Plugins/Animation/RigLogic/Source/RigLogicModule/Public/DNAAsset.h
@@ -108,6 +108,26 @@ private:
 	void InvalidateRigRuntimeContext();
 	void InitializeRigRuntimeContext();
 
+private:
+	struct FSkeletalMeshSkeletonToDNAIndexMappingKey
+	{
+		TWeakObjectPtr<const USkeletalMesh> SkeletalMesh;
+		TWeakObjectPtr<const USkeleton> Skeleton;
+
+		bool operator==(const FSkeletalMeshSkeletonToDNAIndexMappingKey& Other) const
+		{
+			return (SkeletalMesh == Other.SkeletalMesh) && (Skeleton == Other.Skeleton);
+		}
+		bool operator!=(const FSkeletalMeshSkeletonToDNAIndexMappingKey& Other) const
+		{
+			return !(*this == Other);
+		}
+		friend uint32 GetTypeHash(const FSkeletalMeshSkeletonToDNAIndexMappingKey& Instance)
+		{
+			return HashCombine(GetTypeHash(Instance.SkeletalMesh), GetTypeHash(Instance.Skeleton));
+		}
+	};
+
 private:
 	// Synchronize DNA updates
 	FRWLock DNAUpdateLock;
@@ -135,5 +155,6 @@ private:
 	  * The mapping object owners will be the SkeletalMeshes, and periodic cleanups will
 	  * ensure that dead objects are deleted from the map.
 	 **/
-	TMap<TWeakObjectPtr<const USkeletalMesh>, TSharedPtr<FDNAIndexMapping>> DNAIndexMappingContainer;
+	TMap<FSkeletalMeshSkeletonToDNAIndexMappingKey, TSharedPtr<FDNAIndexMapping>> DNAIndexMappingContainer;
+
 };
