Fix RigLogic joint delta value counter for stats collection

diff --git a/Engine/Plugins/Animation/RigLogic/Source/RigLogicLib/Private/riglogic/riglogic/RigLogicImpl.cpp b/Engine/Plugins/Animation/RigLogic/Source/RigLogicLib/Private/riglogic/riglogic/RigLogicImpl.cpp
--- a/Engine/Plugins/Animation/RigLogic/Source/RigLogicLib/Private/riglogic/riglogic/RigLogicImpl.cpp
+++ b/Engine/Plugins/Animation/RigLogic/Source/RigLogicLib/Private/riglogic/riglogic/RigLogicImpl.cpp
@@ -343,7 +343,7 @@ void RigLogicImpl::collectCalculationStats(const RigInstance* instance, Stats* s
     stats->blendShapeChannelCount = static_cast<std::uint16_t>(blendShapes->getBlendShapeChannelIndicesForLOD(lod).size());
     stats->animatedMapCount = static_cast<std::uint16_t>(animatedMaps->getAnimatedMapIndicesForLOD(lod).size());
     stats->jointCount = static_cast<std::uint16_t>(joints->getJointIndicesForLOD(lod).size());
-    stats->jointDeltaValueCount = static_cast<std::uint16_t>(joints->getJointDeltaValueCountForLOD(lod));
+    stats->jointDeltaValueCount = joints->getJointDeltaValueCountForLOD(lod);
 }
 
 MemoryResource* RigLogicImpl::getMemoryResource() {
