Windows Arm64 and Arm64 fixes for RigLogicLib:
- force RL_BUILD_WITH_NEON
- use intrinsic for 4-float vector constructor

#rnx

diff --git a/Engine/Plugins/Animation/RigLogic/Source/RigLogicLib/Public/trimd/NEON.h b/Engine/Plugins/Animation/RigLogic/Source/RigLogicLib/Public/trimd/NEON.h
--- a/Engine/Plugins/Animation/RigLogic/Source/RigLogicLib/Public/trimd/NEON.h
+++ b/Engine/Plugins/Animation/RigLogic/Source/RigLogicLib/Public/trimd/NEON.h
@@ -28,7 +28,9 @@ struct F128 {
     explicit F128(float value) : F128{vdupq_n_f32(value)} {
     }
 
-    F128(float v1, float v2, float v3, float v4) : data{v1, v2, v3, v4} {
+    F128(float v1, float v2, float v3, float v4) {
+		const float source[4] = {v1, v2, v3, v4};
+		data = vld1q_f32(source);
     }
 
     static F128 fromAlignedSource(const float* source) {
diff --git a/Engine/Plugins/Animation/RigLogic/Source/RigLogicLib/RigLogicLib.Build.cs b/Engine/Plugins/Animation/RigLogic/Source/RigLogicLib/RigLogicLib.Build.cs
--- a/Engine/Plugins/Animation/RigLogic/Source/RigLogicLib/RigLogicLib.Build.cs
+++ b/Engine/Plugins/Animation/RigLogic/Source/RigLogicLib/RigLogicLib.Build.cs
@@ -39,6 +39,10 @@ namespace UnrealBuildTool.Rules
 				PrivateDefinitions.Add("TRIMD_CUSTOM_WINDOWS_H=\"WindowsPlatformUE.h\"");
 				PrivateDefinitions.Add("TRIO_CUSTOM_WINDOWS_H=\"WindowsPlatformUE.h\"");
 				PrivateDefinitions.Add("TRIO_WINDOWS_FILE_MAPPING_AVAILABLE=1");
+				if (Target.Architecture == UnrealArch.Arm64 || Target.Architecture == UnrealArch.Arm64ec)
+				{
+					PrivateDefinitions.Add("RL_BUILD_WITH_NEON");
+				}
 			}
 
 			if (Target.Platform == UnrealTargetPlatform.Linux || Target.Platform == UnrealTargetPlatform.LinuxArm64)
diff --git a/Engine/Plugins/Animation/RigLogic/Source/RigLogicLibTest/RigLogicLibTest.Build.cs b/Engine/Plugins/Animation/RigLogic/Source/RigLogicLibTest/RigLogicLibTest.Build.cs
--- a/Engine/Plugins/Animation/RigLogic/Source/RigLogicLibTest/RigLogicLibTest.Build.cs
+++ b/Engine/Plugins/Animation/RigLogic/Source/RigLogicLibTest/RigLogicLibTest.Build.cs
@@ -14,9 +14,16 @@ namespace UnrealBuildTool.Rules
 
 			PrivateDefinitions.Add("RL_BUILD_WITH_XYZ_ROTATION_ORDER=1");
 
-			if (Target.Platform == UnrealTargetPlatform.Win64 && Target.Architecture == UnrealArch.X64)
+			if (Target.Platform == UnrealTargetPlatform.Win64)
 			{
-				PrivateDefinitions.Add("RL_BUILD_WITH_SSE=1");
+				if (Target.Architecture == UnrealArch.X64)
+				{
+					PrivateDefinitions.Add("RL_BUILD_WITH_SSE=1");
+				}
+				else if (Target.Architecture == UnrealArch.Arm64 || Target.Architecture == UnrealArch.Arm64ec)
+				{
+					PrivateDefinitions.Add("RL_BUILD_WITH_NEON");
+				}
 				PublicDefinitions.Add("GTEST_OS_WINDOWS=1");
 			}
 
