[UEMHC] Fix crash on conform for template meshes with more LODs than archetype

diff --git a/Engine/Plugins/Animation/RigLogic/Source/RigLogicModule/Private/DNAToSkelMeshMap.cpp b/Engine/Plugins/Animation/RigLogic/Source/RigLogicModule/Private/DNAToSkelMeshMap.cpp
--- a/Engine/Plugins/Animation/RigLogic/Source/RigLogicModule/Private/DNAToSkelMeshMap.cpp
+++ b/Engine/Plugins/Animation/RigLogic/Source/RigLogicModule/Private/DNAToSkelMeshMap.cpp
@@ -101,7 +101,7 @@ void FDNAToSkelMeshMap::InitVertexMap(IDNAReader* InDNAReader)
 	}
 	// Then for each LOD and each Vertex in LODModel Vertex Map find corresponding DNA Vertex.
 	int32 ImportedVertexCount = 0;
-	for (int32 LODIndex = 0; LODIndex < ImportedModel->LODModels.Num(); LODIndex++)
+	for (int32 LODIndex = 0; LODIndex < NumEngineLods && LODIndex < NumLODs; LODIndex++)
 	{
 		const FSkeletalMeshLODModel& LODModel = ImportedModel->LODModels[LODIndex];
 		const TArrayView<const uint16> LODMeshIndices = InDNAReader->GetMeshIndicesForLOD(LODIndex);
