Fixed header units compile errors

diff --git a/Engine/Plugins/Animation/RigLogic/Source/RigLogicLib/Private/dna/stream/FilteredBinaryInputArchive.cpp b/Engine/Plugins/Animation/RigLogic/Source/RigLogicLib/Private/dna/stream/FilteredBinaryInputArchive.cpp
--- a/Engine/Plugins/Animation/RigLogic/Source/RigLogicLib/Private/dna/stream/FilteredBinaryInputArchive.cpp
+++ b/Engine/Plugins/Animation/RigLogic/Source/RigLogicLib/Private/dna/stream/FilteredBinaryInputArchive.cpp
@@ -378,28 +378,28 @@ void FilteredBinaryInputArchive::process(RawMachineLearnedBehavior& dest) {
         const auto neuralNetCount = static_cast<std::uint16_t>(processSize());
 
         dest.lodNeuralNetworkMapping.discardLODs(lodConstraint);
-        const auto passingIndices = dest.lodNeuralNetworkMapping.getCombinedDistinctIndices(memRes);
-        UnorderedMap<std::uint16_t, std::uint16_t> remappedIndices{memRes};
-        remap(neuralNetCount, passingIndices, remappedIndices);
-        dest.lodNeuralNetworkMapping.mapIndices([&remappedIndices](std::uint16_t value) {
-                return remappedIndices.at(value);
+        const auto passingIndices2 = dest.lodNeuralNetworkMapping.getCombinedDistinctIndices(memRes);
+        UnorderedMap<std::uint16_t, std::uint16_t> remappedIndices2{memRes};
+        remap(neuralNetCount, passingIndices2, remappedIndices2);
+        dest.lodNeuralNetworkMapping.mapIndices([&remappedIndices2](std::uint16_t value) {
+                return remappedIndices2.at(value);
             });
 
         // Delete and remap neural network indices in neuralNetworkToMeshRegion mapping
         for (auto& mesh : dest.neuralNetworkToMeshRegion.indices) {
             for (auto& region : mesh) {
-                extd::filter(region, extd::byValue(passingIndices));
+                extd::filter(region, extd::byValue(passingIndices2));
                 for (auto& netIndex : region) {
-                    netIndex = remappedIndices.at(netIndex);
+                    netIndex = remappedIndices2.at(netIndex);
                 }
             }
         }
 
-        dest.neuralNetworks.reserve(passingIndices.size());
+        dest.neuralNetworks.reserve(passingIndices2.size());
         for (std::uint16_t i = {}; i < neuralNetCount; ++i) {
             // Check if the neural net indices filtered for the current maxLOD permit loading
             // this neural network
-            if (extd::contains(passingIndices, i)) {
+            if (extd::contains(passingIndices2, i)) {
                 RawNeuralNetwork neuralNetwork{memRes};
                 process(neuralNetwork);
                 dest.neuralNetworks.push_back(std::move(neuralNetwork));
@@ -488,18 +488,18 @@ void FilteredBinaryInputArchive::process(RawRBFBehavior& dest) {
     const auto solverCount = static_cast<std::uint16_t>(processSize());
     dest.lodSolverMapping.discardLODs(lodConstraint);
 
-    const auto passingIndices = dest.lodSolverMapping.getCombinedDistinctIndices(memRes);
-    UnorderedMap<std::uint16_t, std::uint16_t> remappedIndices{memRes};
-    remap(solverCount, passingIndices, remappedIndices);
-    dest.lodSolverMapping.mapIndices([&remappedIndices](std::uint16_t value) {
-            return remappedIndices.at(value);
+    const auto passingIndices2 = dest.lodSolverMapping.getCombinedDistinctIndices(memRes);
+    UnorderedMap<std::uint16_t, std::uint16_t> remappedIndices2{memRes};
+    remap(solverCount, passingIndices2, remappedIndices2);
+    dest.lodSolverMapping.mapIndices([&remappedIndices2](std::uint16_t value) {
+            return remappedIndices2.at(value);
         });
 
-    dest.solvers.reserve(passingIndices.size());
+    dest.solvers.reserve(passingIndices2.size());
     for (std::uint16_t i = {}; i < solverCount; ++i) {
         // Check if the solver indices filtered for the current maxLOD permit loading
         // this solver
-        if (extd::contains(passingIndices, i)) {
+        if (extd::contains(passingIndices2, i)) {
             RawRBFSolver solver{memRes};
             process(solver);
             dest.solvers.push_back(std::move(solver));
diff --git a/Engine/Source/Runtime/CoreUObject/Public/UObject/RemoteObjectTypes.h b/Engine/Source/Runtime/CoreUObject/Public/UObject/RemoteObjectTypes.h
--- a/Engine/Source/Runtime/CoreUObject/Public/UObject/RemoteObjectTypes.h
+++ b/Engine/Source/Runtime/CoreUObject/Public/UObject/RemoteObjectTypes.h
@@ -17,11 +17,11 @@ namespace UE::Net
 	struct FRemoteObjectReferenceNetSerializer;
 }
 
-constexpr int32 REMOTE_OBJECT_SERVER_ID_BIT_SIZE = 10;
-constexpr int32 REMOTE_OBJECT_SERIAL_NUMBER_BIT_SIZE = 53;
+constexpr inline int32 REMOTE_OBJECT_SERVER_ID_BIT_SIZE = 10;
+constexpr inline int32 REMOTE_OBJECT_SERIAL_NUMBER_BIT_SIZE = 53;
 
-constexpr uint64 MAX_REMOTE_OBJECT_SERVER_ID = ((uint64_t)1 << REMOTE_OBJECT_SERVER_ID_BIT_SIZE) - 1;
-constexpr uint64 MAX_REMOTE_OBJECT_SERIAL_NUMBER = ((uint64_t)1 << REMOTE_OBJECT_SERIAL_NUMBER_BIT_SIZE) - 1;
+constexpr inline uint64 MAX_REMOTE_OBJECT_SERVER_ID = ((uint64_t)1 << REMOTE_OBJECT_SERVER_ID_BIT_SIZE) - 1;
+constexpr inline uint64 MAX_REMOTE_OBJECT_SERIAL_NUMBER = ((uint64_t)1 << REMOTE_OBJECT_SERIAL_NUMBER_BIT_SIZE) - 1;
 
 enum class ERemoteServerIdConstants : uint32
 {
diff --git a/Engine/Source/Runtime/Engine/Classes/Engine/StreamableManager.h b/Engine/Source/Runtime/Engine/Classes/Engine/StreamableManager.h
--- a/Engine/Source/Runtime/Engine/Classes/Engine/StreamableManager.h
+++ b/Engine/Source/Runtime/Engine/Classes/Engine/StreamableManager.h
@@ -702,15 +702,15 @@ struct FStreamableAsyncLoadParams
 struct FStreamableManager : public FGCObject
 {
 	// Default priority for all async loads
-	static constexpr TAsyncLoadPriority DefaultAsyncLoadPriority	= UE::StreamableManager::Private::DefaultAsyncLoadPriority;
+	static inline constexpr TAsyncLoadPriority DefaultAsyncLoadPriority		= UE::StreamableManager::Private::DefaultAsyncLoadPriority;
 	// Priority to try and load immediately
-	static constexpr TAsyncLoadPriority AsyncLoadHighPriority		= UE::StreamableManager::Private::AsyncLoadHighPriority;
+	static inline constexpr TAsyncLoadPriority AsyncLoadHighPriority		= UE::StreamableManager::Private::AsyncLoadHighPriority;
 
-	static constexpr FStreamableDownloadPriority DownloadVeryLowPriority	= UE::StreamableManager::Private::DownloadVeryLowPriority;
-	static constexpr FStreamableDownloadPriority DownloadLowPriority		= UE::StreamableManager::Private::DownloadLowPriority;
-	static constexpr FStreamableDownloadPriority DownloadDefaultPriority	= UE::StreamableManager::Private::DownloadDefaultPriority;
-	static constexpr FStreamableDownloadPriority DownloadHighPriority		= UE::StreamableManager::Private::DownloadHighPriority;
-	static constexpr FStreamableDownloadPriority DownloadVeryHighPriority	= UE::StreamableManager::Private::DownloadVeryHighPriority;
+	static inline constexpr FStreamableDownloadPriority DownloadVeryLowPriority		= UE::StreamableManager::Private::DownloadVeryLowPriority;
+	static inline constexpr FStreamableDownloadPriority DownloadLowPriority			= UE::StreamableManager::Private::DownloadLowPriority;
+	static inline constexpr FStreamableDownloadPriority DownloadDefaultPriority		= UE::StreamableManager::Private::DownloadDefaultPriority;
+	static inline constexpr FStreamableDownloadPriority DownloadHighPriority		= UE::StreamableManager::Private::DownloadHighPriority;
+	static inline constexpr FStreamableDownloadPriority DownloadVeryHighPriority	= UE::StreamableManager::Private::DownloadVeryHighPriority;
 
 	/** 
 	 * This is the primary streamable operation. Requests streaming of one or more target objects. When complete, a delegate function is called. Returns a Streamable Handle.
