Fix access to inverse rotations cache in presence of body rig controls but without body rig behavior data

diff --git a/Engine/Plugins/Animation/RigLogic/Source/RigLogicModule/Private/DNAIndexMapping.cpp b/Engine/Plugins/Animation/RigLogic/Source/RigLogicModule/Private/DNAIndexMapping.cpp
--- a/Engine/Plugins/Animation/RigLogic/Source/RigLogicModule/Private/DNAIndexMapping.cpp
+++ b/Engine/Plugins/Animation/RigLogic/Source/RigLogicModule/Private/DNAIndexMapping.cpp
@@ -89,6 +89,11 @@ void FDNAIndexMapping::MapDriverJoints(const IDNAReader* DNAReader, const USkele
 {
 	LLM_SCOPE_BYNAME(TEXT("Animation/RigLogic"));
 
+	if ((DNAReader->GetRBFSolverCount() == 0) && (DNAReader->GetSwingCount() == 0) && (DNAReader->GetTwistCount() == 0))
+	{
+		return;
+	}
+
 	const FReferenceSkeleton& RefSkeleton = SkeletalMesh->GetRefSkeleton();
 	const uint32 ControlCount = DNAReader->GetRawControlCount();
 
