Make `RigRuntimeContextUpdateLock` a transactionally safe lock.

#rnx

diff --git a/Engine/Plugins/Animation/RigLogic/Source/RigLogicModule/Private/DNAAsset.cpp b/Engine/Plugins/Animation/RigLogic/Source/RigLogicModule/Private/DNAAsset.cpp
--- a/Engine/Plugins/Animation/RigLogic/Source/RigLogicModule/Private/DNAAsset.cpp
+++ b/Engine/Plugins/Animation/RigLogic/Source/RigLogicModule/Private/DNAAsset.cpp
@@ -97,21 +97,21 @@ UDNAAsset::~UDNAAsset() = default;
 
 TSharedPtr<IDNAReader> UDNAAsset::GetBehaviorReader()
 {
-	FReadScopeLock DNAScopeLock{DNAUpdateLock};
+	UE::TReadScopeLock DNAScopeLock{DNAUpdateLock};
 	return BehaviorReader;
 }
 
 #if WITH_EDITORONLY_DATA
 TSharedPtr<IDNAReader> UDNAAsset::GetGeometryReader()
 {
-	FReadScopeLock DNAScopeLock{DNAUpdateLock};
+	UE::TReadScopeLock DNAScopeLock{DNAUpdateLock};
 	return GeometryReader;
 }
 #endif
 
 void UDNAAsset::SetBehaviorReader(TSharedPtr<IDNAReader> SourceDNAReader)
 {
-	FWriteScopeLock DNAScopeLock{DNAUpdateLock};
+	UE::TWriteScopeLock DNAScopeLock{DNAUpdateLock};
 	const size_t PredictedSize = (SourceDNAReader->GetNeuralNetworkCount() != 0) ? AVG_BEHAVIOR_SIZE + AVG_MACHINE_LEARNED_BEHAVIOR_SIZE : AVG_BEHAVIOR_SIZE;
 	const EDNADataLayer BehaviorLayers = (
 		EDNADataLayer::Behavior |
@@ -126,16 +126,16 @@ void UDNAAsset::SetBehaviorReader(TSharedPtr<IDNAReader> SourceDNAReader)
 void UDNAAsset::SetGeometryReader(TSharedPtr<IDNAReader> SourceDNAReader)
 {
 #if WITH_EDITORONLY_DATA
-	FWriteScopeLock DNAScopeLock{DNAUpdateLock};
+	UE::TWriteScopeLock DNAScopeLock{DNAUpdateLock};
 	GeometryReader = CopyDNALayer(SourceDNAReader.Get(), EDNADataLayer::Geometry, AVG_GEOMETRY_SIZE);
 #endif // #if WITH_EDITORONLY_DATA
 }
 
 void UDNAAsset::InitializeForRuntimeFrom(UDNAAsset* Other)
 {
-	FWriteScopeLock DNAScopeLock{DNAUpdateLock};
-	FWriteScopeLock ContextScopeLock{RigRuntimeContextUpdateLock};
-	FWriteScopeLock MappingScopeLock{DNAIndexMappingUpdateLock};
+	UE::TWriteScopeLock DNAScopeLock{DNAUpdateLock};
+	UE::TWriteScopeLock ContextScopeLock{RigRuntimeContextUpdateLock};
+	UE::TWriteScopeLock MappingScopeLock{DNAIndexMappingUpdateLock};
 
 	// Store a reference to the other asset's runtime context, using the accessor function to 
 	// ensure that the necessary lock is taken.
@@ -171,8 +171,8 @@ void UDNAAsset::InitializeForRuntimeFrom(UDNAAsset* Other)
 
 void UDNAAsset::InvalidateRigRuntimeContext()
 {
-	FWriteScopeLock ContextScopeLock{RigRuntimeContextUpdateLock};
-	FWriteScopeLock MappingScopeLock(DNAIndexMappingUpdateLock);
+	UE::TWriteScopeLock ContextScopeLock{RigRuntimeContextUpdateLock};
+	UE::TWriteScopeLock MappingScopeLock(DNAIndexMappingUpdateLock);
 	RigRuntimeContext = nullptr;
 	DNAIndexMappingContainer.Empty(1);
 }
@@ -198,7 +198,7 @@ void UDNAAsset::InitializeRigRuntimeContext()
 		}
 
 		{
-			FWriteScopeLock ContextScopeLock{RigRuntimeContextUpdateLock};
+			UE::TWriteScopeLock ContextScopeLock{RigRuntimeContextUpdateLock};
 			RigRuntimeContext = NewContext;
 		}
 #if !WITH_EDITOR
@@ -217,7 +217,7 @@ TSharedPtr<FSharedRigRuntimeContext> UDNAAsset::GetRigRuntimeContext()
 {
 	LLM_SCOPE_BYNAME(TEXT("Animation/RigLogic"));
 
-	FReadScopeLock ContextScopeLock{RigRuntimeContextUpdateLock};
+	UE::TReadScopeLock ContextScopeLock{RigRuntimeContextUpdateLock};
 	return RigRuntimeContext;
 }
 
@@ -225,8 +225,8 @@ TSharedPtr<FDNAIndexMapping> UDNAAsset::GetDNAIndexMapping(const USkeleton* Skel
 {
 	LLM_SCOPE_BYNAME(TEXT("Animation/RigLogic"));
 
-	FReadScopeLock DNAScopeLock{DNAUpdateLock};
-	FWriteScopeLock MappingScopeLock(DNAIndexMappingUpdateLock);
+	UE::TReadScopeLock DNAScopeLock{DNAUpdateLock};
+	UE::TWriteScopeLock MappingScopeLock(DNAIndexMappingUpdateLock);
 
 	// Find currently needed mapping and also clean stale objects along the way (requires only one iteration over the map)
 	TSharedPtr<FDNAIndexMapping> DNAIndexMapping;
@@ -289,7 +289,7 @@ bool UDNAAsset::Init(const FString& DNAFilename)
 		return false;
 	}
 	
-	FWriteScopeLock DNAScopeLock{DNAUpdateLock};
+	UE::TWriteScopeLock DNAScopeLock{DNAUpdateLock};
 
 	// Load run-time data (behavior) from whole-DNA buffer into BehaviorReader
 	const EDNADataLayer BehaviorLayers = (
@@ -331,7 +331,7 @@ void UDNAAsset::Serialize(FArchive& Ar)
 
 	Ar.UsingCustomVersion(FDNAAssetCustomVersion::GUID);
 
-	FWriteScopeLock DNAScopeLock{DNAUpdateLock};
+	UE::TWriteScopeLock DNAScopeLock{DNAUpdateLock};
 
 	if (Ar.CustomVer(FDNAAssetCustomVersion::GUID) >= FDNAAssetCustomVersion::BeforeCustomVersionWasAdded)
 	{
@@ -392,7 +392,7 @@ void UDNAAsset::PostEditChangeProperty(FPropertyChangedEvent& PropertyChangedEve
 	const FName& MemberPropertyName = PropertyChangedEvent.MemberProperty->GetFName();
 	if (MemberPropertyName == GET_MEMBER_NAME_CHECKED(UDNAAsset, RigLogicConfiguration))
 	{
-		FWriteScopeLock DNAScopeLock{DNAUpdateLock};
+		UE::TWriteScopeLock DNAScopeLock{DNAUpdateLock};
 		InvalidateRigRuntimeContext();
 		InitializeRigRuntimeContext();
 	}
diff --git a/Engine/Plugins/Animation/RigLogic/Source/RigLogicModule/Public/DNAAsset.h b/Engine/Plugins/Animation/RigLogic/Source/RigLogicModule/Public/DNAAsset.h
--- a/Engine/Plugins/Animation/RigLogic/Source/RigLogicModule/Public/DNAAsset.h
+++ b/Engine/Plugins/Animation/RigLogic/Source/RigLogicModule/Public/DNAAsset.h
@@ -6,6 +6,8 @@
 #include "UObject/Object.h"
 #include "UObject/ObjectMacros.h"
 
+#include "Misc/TransactionallySafeRWLock.h"
+
 #include "Engine/AssetUserData.h"
 #include "Interfaces/Interface_AssetUserData.h"
 
@@ -135,7 +137,7 @@ private:
 	FRWLock DNAUpdateLock;
 
 	// Synchronize Rig Runtime Context updates
-	FRWLock RigRuntimeContextUpdateLock;
+	FTransactionallySafeRWLock RigRuntimeContextUpdateLock;
 
 	// Synchronize DNA Index Mapping updates
 	FRWLock DNAIndexMappingUpdateLock;
