[RigLogicEditor] Fixed bug related to SKM selection for DNA import and some small improvements to DNA Importing process

diff --git a/Engine/Plugins/Animation/RigLogic/Source/RigLogicEditor/Private/DNAAssetImportFactory.cpp b/Engine/Plugins/Animation/RigLogic/Source/RigLogicEditor/Private/DNAAssetImportFactory.cpp
--- a/Engine/Plugins/Animation/RigLogic/Source/RigLogicEditor/Private/DNAAssetImportFactory.cpp
+++ b/Engine/Plugins/Animation/RigLogic/Source/RigLogicEditor/Private/DNAAssetImportFactory.cpp
@@ -138,10 +138,13 @@ UObject* UDNAAssetImportFactory::FactoryCreateFile(UClass* InClass, UObject* InP
 	ImportUI->bIsReimport = false;
 	ImportUI->ReimportMesh = nullptr;
 
-	// If the user provided a SkeletalMesh through the AssetImportTask Options 
-	// refrain from searching for the Skeletal Mesh in the asset package
+	// Show the import dialog only when not in a "yes to all" state or when automating import
+	bool bIsAutomated = IsAutomatedImport() || IsAutomatedReimport();
+	bool bShowImportDialog = !bIsAutomated;
 
-	if (!IsValid(ImportUI->SkeletalMesh)) {
+	// If the user provided a SkeletalMesh through the AssetImportTask Options, which means we have automated import,
+	// refrain from searching for the Skeletal Mesh in the asset package
+	if (!IsValid(ImportUI->SkeletalMesh) || !bIsAutomated) {
 		USkeletalMesh* SkeletalMeshForDNAImport = nullptr;
 		ForEachObjectWithOuter(InParent, [&SkeletalMeshForDNAImport](UObject* Object)
 			{
@@ -155,10 +158,6 @@ UObject* UDNAAssetImportFactory::FactoryCreateFile(UClass* InClass, UObject* InP
 		ImportUI->SkeletalMesh = SkeletalMeshForDNAImport;
 	}
 
-	// Show the import dialog only when not in a "yes to all" state or when automating import
-	bool bIsAutomated = IsAutomatedImport() || IsAutomatedReimport();
-	bool bShowImportDialog = !bIsAutomated;
-
 	ImportOptions = GetImportOptions(DNAImporter, ImportUI, bShowImportDialog, bIsAutomated, InParent->GetPathName(), bOutOperationCanceled, UFactory::CurrentFilename);
 
 	if (bOutOperationCanceled)
@@ -187,8 +186,7 @@ UObject* UDNAAssetImportFactory::FactoryCreateFile(UClass* InClass, UObject* InP
 			UDNAAsset* DNAAsset = NewObject< UDNAAsset >(SkelMesh, FName(*DNAAssetFileName)); //SkelMesh has to be its outer, otherwise DNAAsset won't be saved			
 			if (DNAAsset->Init(*Filename))
 			{
-				UAssetUserData* DNAAssetUserData = Cast<UAssetUserData>(DNAAsset);
-				SkelMesh->AddAssetUserData(DNAAssetUserData);
+				SkelMesh->AddAssetUserData(DNAAsset);
 
 				GEditor->GetEditorSubsystem<UImportSubsystem>()->BroadcastAssetPostImport(this, DNAAsset);
 				DNAImporter->PartialCleanUp();
diff --git a/Engine/Plugins/Animation/RigLogic/Source/RigLogicModule/Public/DNAAsset.h b/Engine/Plugins/Animation/RigLogic/Source/RigLogicModule/Public/DNAAsset.h
--- a/Engine/Plugins/Animation/RigLogic/Source/RigLogicModule/Public/DNAAsset.h
+++ b/Engine/Plugins/Animation/RigLogic/Source/RigLogicModule/Public/DNAAsset.h
@@ -36,7 +36,7 @@ enum class EDNADataLayer: uint8;
   * Currently, the design-time part still loads the geometry, as it is needed for the skeletal mesh update; once SkeletalMeshDNAReader is
   * fully implemented, it will be able to read the geometry directly from the SkeletalMesh and won't load it into this asset 
   **/
-UCLASS(MinimalAPI, BlueprintType, meta = (DisplayName = "MetaHuman DNA data"))
+UCLASS(MinimalAPI, BlueprintType, meta = (DisplayName = "MetaHuman DNA Data"))
 class UDNAAsset : public UAssetUserData
 {
 	GENERATED_BODY()
