Make `UDNAAsset::GetDNAIndexMapping` work under transactions by using transactionally-safe locks.

diff --git a/Engine/Plugins/Animation/RigLogic/Source/RigLogicModule/Public/DNAAsset.h b/Engine/Plugins/Animation/RigLogic/Source/RigLogicModule/Public/DNAAsset.h
--- a/Engine/Plugins/Animation/RigLogic/Source/RigLogicModule/Public/DNAAsset.h
+++ b/Engine/Plugins/Animation/RigLogic/Source/RigLogicModule/Public/DNAAsset.h
@@ -140,13 +140,13 @@ private:
 
 private:
 	// Synchronize DNA updates
-	FRWLock DNAUpdateLock;
+	FTransactionallySafeRWLock DNAUpdateLock;
 
 	// Synchronize Rig Runtime Context updates
 	FTransactionallySafeRWLock RigRuntimeContextUpdateLock;
 
 	// Synchronize DNA Index Mapping updates
-	FRWLock DNAIndexMappingUpdateLock;
+	FTransactionallySafeRWLock DNAIndexMappingUpdateLock;
 
 	/** Part of the .dna file needed for run-time execution of RigLogic;
 	 **/
