Correct file paths of RigLogic platform specific configs

diff --git a/Engine/Platforms/Android/Plugins/Animation/RigLogic/RigLogic_Android.uplugin b/Engine/Platforms/Android/Plugins/Animation/RigLogic/RigLogic_Android.uplugin
new file mode 100644
--- /dev/null
+++ b/Engine/Platforms/Android/Plugins/Animation/RigLogic/RigLogic_Android.uplugin
@@ -0,0 +1,5 @@
+{
+	"FileVersion": 3,
+	"bIsPluginExtension": true
+}
+
diff --git a/Engine/Platforms/Android/Plugins/Animation/RigLogic/Source/RigLogicLib/RigLogicLib_Android.Build.cs b/Engine/Platforms/Android/Plugins/Animation/RigLogic/Source/RigLogicLib/RigLogicLib_Android.Build.cs
new file mode 100644
--- /dev/null
+++ b/Engine/Platforms/Android/Plugins/Animation/RigLogic/Source/RigLogicLib/RigLogicLib_Android.Build.cs
@@ -0,0 +1,13 @@
+// Copyright Epic Games, Inc. All Rights Reserved.
+namespace UnrealBuildTool.Rules
+{
+	public class RigLogicLib_Android : RigLogicLib
+	{
+		public RigLogicLib_Android(ReadOnlyTargetRules Target) : base(Target)
+		{
+			PrivateDefinitions.Add("RL_AUTODETECT_SSE");
+			PrivateDefinitions.Add("RL_AUTODETECT_NEON");
+			PrivateDefinitions.Add("RL_AUTODETECT_HALF_FLOATS");
+		}
+	}
+}
\ No newline at end of file
diff --git a/Engine/Platforms/Android/Plugins/RigLogic/RigLogic_Android.uplugin b/Engine/Platforms/Android/Plugins/RigLogic/RigLogic_Android.uplugin
deleted file mode 100644
diff --git a/Engine/Platforms/Android/Plugins/RigLogic/Source/RigLogicLib/RigLogicLib_Android.Build.cs b/Engine/Platforms/Android/Plugins/RigLogic/Source/RigLogicLib/RigLogicLib_Android.Build.cs
deleted file mode 100644
diff --git a/Engine/Platforms/IOS/Plugins/Animation/RigLogic/RigLogic_IOS.uplugin b/Engine/Platforms/IOS/Plugins/Animation/RigLogic/RigLogic_IOS.uplugin
new file mode 100644
--- /dev/null
+++ b/Engine/Platforms/IOS/Plugins/Animation/RigLogic/RigLogic_IOS.uplugin
@@ -0,0 +1,5 @@
+{
+	"FileVersion": 3,
+	"bIsPluginExtension": true
+}
+
diff --git a/Engine/Platforms/IOS/Plugins/Animation/RigLogic/Source/RigLogicLib/RigLogicLib_IOS.Build.cs b/Engine/Platforms/IOS/Plugins/Animation/RigLogic/Source/RigLogicLib/RigLogicLib_IOS.Build.cs
new file mode 100644
--- /dev/null
+++ b/Engine/Platforms/IOS/Plugins/Animation/RigLogic/Source/RigLogicLib/RigLogicLib_IOS.Build.cs
@@ -0,0 +1,12 @@
+// Copyright Epic Games, Inc. All Rights Reserved.
+namespace UnrealBuildTool.Rules
+{
+	public class RigLogicLib_IOS : RigLogicLib
+	{
+		public RigLogicLib_IOS(ReadOnlyTargetRules Target) : base(Target)
+		{
+			PrivateDefinitions.Add("RL_AUTODETECT_NEON");
+			PrivateDefinitions.Add("RL_AUTODETECT_HALF_FLOATS");
+		}
+	}
+}
\ No newline at end of file
diff --git a/Engine/Platforms/IOS/Plugins/RigLogic/RigLogic_IOS.uplugin b/Engine/Platforms/IOS/Plugins/RigLogic/RigLogic_IOS.uplugin
deleted file mode 100644
diff --git a/Engine/Platforms/IOS/Plugins/RigLogic/Source/RigLogicLib/RigLogicLib_IOS.Build.cs b/Engine/Platforms/IOS/Plugins/RigLogic/Source/RigLogicLib/RigLogicLib_IOS.Build.cs
deleted file mode 100644
