Trigger caching of all otherwise lazily computed data when a DNA is loaded from a binary stream, which can avoid potential race conditions if the same DNA reader is utilized in a multithreaded environment

diff --git a/Engine/Plugins/Animation/RigLogic/Source/RigLogicLib/Private/dna/DenormalizedData.h b/Engine/Plugins/Animation/RigLogic/Source/RigLogicLib/Private/dna/DenormalizedData.h
--- a/Engine/Plugins/Animation/RigLogic/Source/RigLogicLib/Private/dna/DenormalizedData.h
+++ b/Engine/Plugins/Animation/RigLogic/Source/RigLogicLib/Private/dna/DenormalizedData.h
@@ -39,6 +39,15 @@ struct DenormalizedData {
     }
 
     public:
+        void populate(const Reader* source) {
+            populateJointVariableAttributeIndices(source);
+            populateMeshBlendShapeMappingIndices(source);
+            populateRBFPoseJointOutputIndices(source);
+            populateRBFPoseJointOutputValues(source);
+            populateRBFBlendShapeOutputIndices(source);
+            populateRBFAnimatedMapOutputIndices(source);
+        }
+
         void populateJointVariableAttributeIndices(const Reader* source) {
             // Prepare storage for all available LODs
             const auto lodCount = source->getLODCount();
diff --git a/Engine/Plugins/Animation/RigLogic/Source/RigLogicLib/Private/dna/stream/BinaryStreamReaderImpl.cpp b/Engine/Plugins/Animation/RigLogic/Source/RigLogicLib/Private/dna/stream/BinaryStreamReaderImpl.cpp
--- a/Engine/Plugins/Animation/RigLogic/Source/RigLogicLib/Private/dna/stream/BinaryStreamReaderImpl.cpp
+++ b/Engine/Plugins/Animation/RigLogic/Source/RigLogicLib/Private/dna/stream/BinaryStreamReaderImpl.cpp
@@ -144,6 +144,8 @@ void BinaryStreamReaderImpl::read() {
         status.set(VersionMismatchError, dna.version.generation, dna.version.version);
         return;
     }
+
+    cache.populate(this);
 }
 
 }  // namespace dna
