Fix RigLogic VersionName attribute in plugin metadata

diff --git a/Engine/Plugins/Animation/RigLogic/RigLogic.uplugin b/Engine/Plugins/Animation/RigLogic/RigLogic.uplugin
--- a/Engine/Plugins/Animation/RigLogic/RigLogic.uplugin
+++ b/Engine/Plugins/Animation/RigLogic/RigLogic.uplugin
@@ -1,7 +1,7 @@
 {
 	"FileVersion": 3,
 	"Version": 10.3,
-	"VersionName": "10.2.7",
+	"VersionName": "10.3.0",
 	"FriendlyName": "RigLogic Plugin v10.3.0",
 	"Description": "RigLogic Plugin for Facial Animation v10.3.0",
 	"Category": "Animation",
