Fix RigLogic RBF solver initialization to always convert DNA data into UE coordinate system
#Synced-CL 48501008

diff --git a/Engine/Plugins/Animation/RigLogic/Source/RigLogicModule/Private/RigLogicDNAReader.cpp b/Engine/Plugins/Animation/RigLogic/Source/RigLogicModule/Private/RigLogicDNAReader.cpp
--- a/Engine/Plugins/Animation/RigLogic/Source/RigLogicModule/Private/RigLogicDNAReader.cpp
+++ b/Engine/Plugins/Animation/RigLogic/Source/RigLogicModule/Private/RigLogicDNAReader.cpp
@@ -847,9 +847,8 @@ void RigLogicDNAReader::CacheRBFSolverRawControlValues(std::uint16_t SolverIndex
 
 	Values.Reset(RawControlValues.size());
 	Values.Append(RawControlValues.data(), RawControlValues.size());
-	if (SolverDistanceMethod != dna::RBFDistanceMethod::Euclidean)
+	if (Values.Num() % 4 == 0)
 	{
-		ensureMsgf(Values.Num() % 4 == 0, TEXT("DNA RBF Solver Raw Control Value count invalid"));
 		for (std::size_t Offset = 0; Offset < Values.Num(); Offset += 4)
 		{
 			const tdm::fquat PoseRotationInDNASpace{Values[Offset + 0], Values[Offset + 1], Values[Offset + 2], Values[Offset + 3]};
