Fix RigLogic quaterion raw control initialization leaving controls named with the .qw suffix incorrectly just zero-initialized, instead of setting them to the identity quaternion

diff --git a/Engine/Plugins/Animation/RigLogic/Source/RigLogicLib/Private/riglogic/controls/ControlsFactory.cpp b/Engine/Plugins/Animation/RigLogic/Source/RigLogicLib/Private/riglogic/controls/ControlsFactory.cpp
--- a/Engine/Plugins/Animation/RigLogic/Source/RigLogicLib/Private/riglogic/controls/ControlsFactory.cpp
+++ b/Engine/Plugins/Animation/RigLogic/Source/RigLogicLib/Private/riglogic/controls/ControlsFactory.cpp
@@ -121,7 +121,8 @@ static Vector<ControlInitializer> createInitialControlValues(const dna::Reader*
     auto makeSuffix = [](const char* str) {
             return StringView{str, std::strlen(str)};
         };
-    const auto qwSuffix = makeSuffix(".w");
+    const auto wSuffix = makeSuffix(".w");
+    const auto qwSuffix = makeSuffix(".qw");
     const auto rawControlCount = static_cast<std::size_t>(reader->getRawControlCount());
     const auto psdControlCount = static_cast<std::size_t>(reader->getPSDCount());
     const auto mlControlCount = static_cast<std::size_t>(reader->getMLControlCount());
@@ -133,7 +134,7 @@ static Vector<ControlInitializer> createInitialControlValues(const dna::Reader*
         const auto solverRawControlIndices = reader->getRBFSolverRawControlIndices(solverIndex);
         for (const auto rawControlIndex : solverRawControlIndices) {
             const auto rawControlName = reader->getRawControlName(rawControlIndex);
-            if (endsWith(rawControlName, qwSuffix)) {
+            if (endsWith(rawControlName, wSuffix) || endsWith(rawControlName, qwSuffix)) {
                 initialValues.push_back({rawControlIndex, 1.0f});
             }
         }
