Move ArchiveMemoryStream into RigLogicLib but make it publicly available, such that GeneSplicerModule may utilize it as well

diff --git a/Engine/Plugins/Animation/GeneSplicer/Source/GeneSplicerModule/Private/ArchiveMemoryStream.cpp b/Engine/Plugins/Animation/GeneSplicer/Source/GeneSplicerModule/Private/ArchiveMemoryStream.cpp
deleted file mode 100644
diff --git a/Engine/Plugins/Animation/GeneSplicer/Source/GeneSplicerModule/Private/ArchiveMemoryStream.h b/Engine/Plugins/Animation/GeneSplicer/Source/GeneSplicerModule/Private/ArchiveMemoryStream.h
deleted file mode 100644
diff --git a/Engine/Plugins/Animation/RigLogic/Source/RigLogicLib/Private/ArchiveMemoryStream.cpp b/Engine/Plugins/Animation/RigLogic/Source/RigLogicLib/Private/ArchiveMemoryStream.cpp
new file mode 100644
--- /dev/null
+++ b/Engine/Plugins/Animation/RigLogic/Source/RigLogicLib/Private/ArchiveMemoryStream.cpp
@@ -0,0 +1,86 @@
+// Copyright Epic Games, Inc. All Rights Reserved.
+
+#include "ArchiveMemoryStream.h"
+
+#include "Math/NumericLimits.h"
+#include "Serialization/Archive.h"
+
+FArchiveMemoryStream::FArchiveMemoryStream(FArchive* Archive) :
+	Archive{Archive},
+	Origin{}
+{
+	check(Archive != nullptr);
+	Origin = Archive->Tell();
+}
+
+void FArchiveMemoryStream::seek(std::uint64_t Position)
+{
+	ensure(Position <= static_cast<std::uint64_t>(TNumericLimits<int64>::Max()));
+	Archive->Seek(Origin + static_cast<int64>(Position));
+}
+
+std::uint64_t FArchiveMemoryStream::tell()
+{
+	const int64 RelPosition = (Archive->Tell() - Origin);
+	ensure(RelPosition >= 0);
+	return static_cast<std::uint64_t>(RelPosition);
+}
+
+void FArchiveMemoryStream::open()
+{
+}
+
+void FArchiveMemoryStream::close()
+{
+}
+
+size_t FArchiveMemoryStream::read(char* ReadToBuffer, size_t Size)
+{
+	Archive->Serialize(ReadToBuffer, Size);
+	return Size;
+}
+
+size_t FArchiveMemoryStream::read(Writable* Destination, size_t Size)
+{
+	constexpr size_t BufferSize = 4096ul;
+	char Buffer[BufferSize];
+	size_t SizeToRead = Size;
+	while (SizeToRead > BufferSize) {
+		Archive->Serialize(Buffer, BufferSize);
+		Destination->write(Buffer, BufferSize);
+		SizeToRead -= BufferSize;
+	}
+	Archive->Serialize(Buffer, SizeToRead);
+	Destination->write(Buffer, SizeToRead);
+	return Size;
+}
+
+size_t FArchiveMemoryStream::write(const char* WriteFromBuffer, size_t Size)
+{
+	Archive->Serialize(const_cast<char*>(WriteFromBuffer), Size);
+	return Size;
+}
+
+size_t FArchiveMemoryStream::write(Readable* Source, size_t Size)
+{
+	constexpr size_t BufferSize = 4096ul;
+	char Buffer[BufferSize];
+	size_t SizeToWrite = Size;
+	while (SizeToWrite > BufferSize) {
+		Source->read(Buffer, BufferSize);
+		Archive->Serialize(const_cast<char*>(Buffer), BufferSize);
+		SizeToWrite -= BufferSize;
+	}
+	Source->read(Buffer, SizeToWrite);
+	Archive->Serialize(const_cast<char*>(Buffer), SizeToWrite);
+	return Size;
+}
+
+std::uint64_t FArchiveMemoryStream::size()
+{
+	const int64 TotalSize = Archive->GetArchiveState().TotalSize();
+	ensure(TotalSize >= 0);
+	const int64 StreamSize = TotalSize - Origin;
+	ensure(StreamSize >= 0);
+	return static_cast<std::uint64_t>(StreamSize);
+}
diff --git a/Engine/Plugins/Animation/RigLogic/Source/RigLogicLib/Public/ArchiveMemoryStream.h b/Engine/Plugins/Animation/RigLogic/Source/RigLogicLib/Public/ArchiveMemoryStream.h
new file mode 100644
--- /dev/null
+++ b/Engine/Plugins/Animation/RigLogic/Source/RigLogicLib/Public/ArchiveMemoryStream.h
@@ -0,0 +1,35 @@
+// Copyright Epic Games, Inc. All Rights Reserved.
+
+#pragma once
+
+#ifdef RIGLOGIC_MODULE_DISCARD
+	#define RIGLOGICLIB_API
+#endif  // RIGLOGIC_MODULE_DISCARD
+
+#include "CoreMinimal.h"
+
+#include "riglogic/RigLogic.h"
+
+class FArchive;
+
+/** Adapter that allows using an FArchive instance as a rl4::BoundedIOStream */
+class FArchiveMemoryStream: public rl4::BoundedIOStream
+{
+public:
+	RIGLOGICLIB_API explicit FArchiveMemoryStream(FArchive* Archive);
+
+	RIGLOGICLIB_API void seek(std::uint64_t Position) override;
+	RIGLOGICLIB_API std::uint64_t tell() override;
+	RIGLOGICLIB_API void open() override;
+	RIGLOGICLIB_API void close() override;
+	RIGLOGICLIB_API size_t read(char* ReadToBuffer, size_t Size) override;
+	RIGLOGICLIB_API size_t read(Writable* Destination, size_t Size) override;
+	RIGLOGICLIB_API size_t write(const char* WriteFromBuffer, size_t Size) override;
+	RIGLOGICLIB_API size_t write(Readable* Source, size_t Size) override;
+	RIGLOGICLIB_API std::uint64_t size() override;
+
+private:
+	FArchive* Archive;
+	int64 Origin;
+
+};
diff --git a/Engine/Plugins/Animation/RigLogic/Source/RigLogicModule/Private/ArchiveMemoryStream.cpp b/Engine/Plugins/Animation/RigLogic/Source/RigLogicModule/Private/ArchiveMemoryStream.cpp
deleted file mode 100644
diff --git a/Engine/Plugins/Animation/RigLogic/Source/RigLogicModule/Private/ArchiveMemoryStream.h b/Engine/Plugins/Animation/RigLogic/Source/RigLogicModule/Private/ArchiveMemoryStream.h
deleted file mode 100644
