[RigLogic] New DNA Asset work.
- Created a new AssetUserData class and made sure it's used properly as a wrapper for DNAAsset which is now a proper asset type
- Changed rest of the code in RigLogic plugin that uses DNA to work with new asset
- Refactored importing and reimporting of DNA asset
- Added fallaback to new asset user data so nothing breaks even when old DNA asset user data is used 6940014e4c744387d74fa4bb

#Synced-CL 49271574

diff --git a/Engine/Plugins/Animation/RigLogic/Config/DefaultRigLogic.ini b/Engine/Plugins/Animation/RigLogic/Config/DefaultRigLogic.ini
new file mode 100644
--- /dev/null
+++ b/Engine/Plugins/Animation/RigLogic/Config/DefaultRigLogic.ini
@@ -0,0 +1,4 @@
+[CoreRedirects]
+
+
++ClassRedirects=(OldName="/Script/RigLogicEditor.UDNAAssetImportUI", NewName="/Script/RigLogicEditor.UDEPRECATED_DNAAssetImportUI")
\ No newline at end of file
diff --git a/Engine/Plugins/Animation/RigLogic/Source/RigLogicEditor/Private/AssetDefinitions/AssetDefinition_DNAAsset.cpp b/Engine/Plugins/Animation/RigLogic/Source/RigLogicEditor/Private/AssetDefinitions/AssetDefinition_DNAAsset.cpp
new file mode 100644
--- /dev/null
+++ b/Engine/Plugins/Animation/RigLogic/Source/RigLogicEditor/Private/AssetDefinitions/AssetDefinition_DNAAsset.cpp
@@ -0,0 +1,206 @@
+// Copyright Epic Games, Inc. All Rights Reserved.
+
+#include "AssetDefinition_DNAAsset.h"
+#include "DNAAsset.h"
+#include "DNA.h"
+#include "DNAAssetUserData.h"
+#include "DNAImporter.h"
+#include "ToolMenu.h"
+#include "ToolMenus.h"
+#include "ContentBrowserMenuContexts.h"
+#include "ThumbnailRendering/SceneThumbnailInfo.h"
+#include "ThumbnailHelpers.h"
+#include "Engine/SkeletalMesh.h"
+
+DEFINE_LOG_CATEGORY_STATIC(LogDNAAssetDefinition, Log, All);
+
+#define LOCTEXT_NAMESPACE "AssetDefinition_DNAAsset"
+
+namespace MenuExtention_SkelMeshDNAAsset
+{
+	static void ExecuteDNAImport(const UContentBrowserAssetContextMenuContext* InCBContext)
+	{
+		UDNAImporter* DNAImporter = NewObject<UDNAImporter>();
+		TArray<USkeletalMesh*> SkeletalMeshes;
+		for(const FAssetData& SelectedSkeletalMesh : InCBContext->SelectedAssets)
+		{
+			SkeletalMeshes.Add(Cast<USkeletalMesh>(SelectedSkeletalMesh.GetAsset()));
+		}
+
+		if(!DNAImporter->ImportDNAWithPrompt(SkeletalMeshes))
+		{
+			const FText Message = LOCTEXT("DNAImportFromSkeletalMeshMenuFail", "Importing of DNA failed");
+			UE_LOGFMT(LogDNAAssetDefinition, Error, "{}", *Message.ToString());
+		}
+	}
+
+	static void ExecuteDNAReimport(const UContentBrowserAssetContextMenuContext* InCBContext)
+	{
+		UDNA* DNAAsset = Cast<UDNA>(InCBContext->SelectedAssets[0].GetAsset());
+
+		UDNAImporter* DNAImporter = NewObject<UDNAImporter>();
+
+		DNAImporter->ReimportDNAs(DNAAsset);
+	}
+
+	static void ExecuteDNAExport(const UContentBrowserAssetContextMenuContext* InCBContext)
+	{
+		UDNA* DNAAsset = Cast<UDNA>(InCBContext->SelectedAssets[0].GetAsset());
+
+		UDNAImporter* DNAImporter = NewObject<UDNAImporter>();
+
+		DNAImporter->ExportDNA(DNAAsset);
+	}
+
+	static void ExecuteDNAConversion(const UContentBrowserAssetContextMenuContext* InCBContext)
+	{
+		UDNAImporter* DNAImporter = NewObject<UDNAImporter>();
+		for (const FAssetData& SelectedSkeletalMesh : InCBContext->SelectedAssets)
+		{
+			USkeletalMesh* SkelMesh = Cast<USkeletalMesh>(SelectedSkeletalMesh.GetAsset());
+			if(SkelMesh->GetAssetUserData<UDNAAsset>())
+			{
+				DNAImporter->ConvertFromLegacyAssetUserData(SkelMesh);
+			}
+		}
+	}
+
+	static void ExtendDNASkelMeshAssetActions()
+	{
+		FToolMenuOwnerScoped OwnderScoped(UE_MODULE_NAME);
+		{
+			UE::ContentBrowser::ExtendToolMenu_AssetContextMenu(USkeletalMesh::StaticClass())
+				->AddDynamicSection(
+					NAME_None,
+					FNewToolMenuDelegate::CreateLambda(
+						[](UToolMenu* InMenu)
+						{
+							const UContentBrowserAssetContextMenuContext* Context = InMenu->FindContext< UContentBrowserAssetContextMenuContext>();
+							if (Context && Context->SelectedAssets.Num() > 0)
+							{
+								FToolMenuSection& Section = InMenu->FindOrAddSection(TEXT("GetAssetActions"));
+								Section.AddSubMenu(
+									"SkeletalMeshDNASubmenu",
+									LOCTEXT("DNASkeletalMeshSubmenu", "MetaHuman DNA"),
+									LOCTEXT("DNAImportSubmenu_ToolTip", "DNA related actions"),
+									FNewToolMenuDelegate::CreateLambda([Context](UToolMenu* SubMenu)
+									{
+											FToolMenuSection& SubSection = SubMenu->AddSection(
+												"DNA Asset Actions", LOCTEXT("DNAAssetActions", "DNA Actions"));
+											{
+												const FText Label = LOCTEXT("DNAAsset_Import", "Import DNA Asset");
+												const FText Tooltip = LOCTEXT("DNAAsset_ImportTooltip", "Import a new DNA Asset and attach it to this Skeletal Mesh");
+												const FSlateIcon Icon{ FAppStyle::GetAppStyleSetName(), "Icons.MetaHuman" };
+												const FUIAction UIAction = FUIAction(
+													FExecuteAction::CreateStatic(&ExecuteDNAImport, Context)
+												);
+												SubSection.AddMenuEntry(TEXT("DNAAsset_Import"), Label, Tooltip, Icon, UIAction);
+											}
+											{
+												for (FAssetData SelectedSkeletalMesh : Context->SelectedAssets)
+												{
+													USkeletalMesh* SkelMesh = Cast<USkeletalMesh>(SelectedSkeletalMesh.GetAsset());
+													if(SkelMesh->GetAssetUserData<UDNAAsset>())
+													{
+														const FText Label = LOCTEXT("DNAAsset_Convert", "Create DNA Asset");
+														const FText Tooltip = LOCTEXT("DNAAsset_ConvertTooltip", "Convert DNA asset user data to full DNA asset which is attached to this Skeletal Mesh");
+														const FSlateIcon Icon{ FAppStyle::GetAppStyleSetName(), "Icons.MetaHuman" };
+														const FUIAction UIAction = FUIAction(
+															FExecuteAction::CreateStatic(&ExecuteDNAConversion, Context)
+														);
+														SubSection.AddMenuEntry(TEXT("DNAAsset_Convert"), Label, Tooltip, Icon, UIAction);
+														break;
+													}
+												}
+											}
+									}),
+									false,
+									FSlateIcon(FAppStyle::GetAppStyleSetName(), "Icons.Import")
+								);
+							}
+						}
+					)
+				);
+
+			UE::ContentBrowser::ExtendToolMenu_AssetContextMenu(UDNA::StaticClass())
+				->AddDynamicSection(
+					NAME_None,
+					FNewToolMenuDelegate::CreateLambda(
+						[](UToolMenu* InMenu)
+						{
+							const UContentBrowserAssetContextMenuContext* Context = InMenu->FindContext< UContentBrowserAssetContextMenuContext>();
+							if (Context && Context->SelectedAssets.Num() > 0)
+							{
+								FToolMenuSection& Section = InMenu->FindOrAddSection(TEXT("GetAssetActions"));
+								{
+									{
+										const FText Label = LOCTEXT("DNAAssetReimport", "Reimport DNA");
+										const FText Tooltip = LOCTEXT("DNAAssetReimportTooltip", "Reimport this DNA asset");
+										const FSlateIcon Icon{ FAppStyle::GetAppStyleSetName(), "ClassIcon.MetaHumanCharacter" };
+										const FUIAction UIAction = FUIAction(
+											FExecuteAction::CreateStatic(&ExecuteDNAReimport, Context)
+										);
+										Section.AddMenuEntry(TEXT("DNAAssetReimport"), Label, Tooltip, Icon, UIAction);
+									}
+									{
+										const FText Label = LOCTEXT("DNAAssetExport", "Export DNA");
+										const FText Tooltip = LOCTEXT("DNAAssetExportTooltip", "Export DNA file from this asset");
+										const FSlateIcon Icon{ FAppStyle::GetAppStyleSetName(), "ClassIcon.MetaHumanCharacter" };
+										const FUIAction UIAction = FUIAction(
+											FExecuteAction::CreateStatic(&ExecuteDNAExport, Context)
+										);
+										Section.AddMenuEntry(TEXT("DNAAsssetExport"), Label, Tooltip, Icon, UIAction);
+									}
+								}
+							}
+						}
+					)
+				);
+		}
+	}
+
+	static FDelayedAutoRegisterHelper DelayedAutoRegister(
+		EDelayedRegisterRunPhase::EndOfEngineInit,
+		[]
+		{
+			UToolMenus::RegisterStartupCallback(FSimpleMulticastDelegate::FDelegate::CreateStatic(&ExtendDNASkelMeshAssetActions));
+		}
+	);
+}
+
+
+FText UAssetDefinition_DNAAsset::GetAssetDisplayName() const
+{
+	return LOCTEXT("DNAAssetDisplayName", "DNA");
+}
+
+FLinearColor UAssetDefinition_DNAAsset::GetAssetColor() const
+{
+	return FColor::Green;
+}
+
+TSoftClassPtr<UObject> UAssetDefinition_DNAAsset::GetAssetClass() const
+{
+	return UDNA::StaticClass();
+}
+
+TConstArrayView<FAssetCategoryPath> UAssetDefinition_DNAAsset::GetAssetCategories() const
+{
+	static const FAssetCategoryPath Categories[] = { FAssetCategoryPath{ LOCTEXT("MetaHumanAssetCategoryPath", "MetaHuman") } };
+	return Categories;
+}
+
+// TODO: Have MH logo as thumbnail?
+UThumbnailInfo* UAssetDefinition_DNAAsset::LoadThumbnailInfo(const FAssetData& InAssetData) const
+{
+	return UE::Editor::FindOrCreateThumbnailInfo(InAssetData.GetAsset(), USceneThumbnailInfo::StaticClass());
+}
+
+EAssetCommandResult UAssetDefinition_DNAAsset::OpenAssets(const FAssetOpenArgs& InOpenArgs) const
+{
+	// This should open up FSimpleAssetEditor, that should do it for now, but for future work we should probably have our own AssetEditor
+	return Super::OpenAssets(InOpenArgs);
+}
+
+
+#undef LOCTEXT_NAMESPACE
\ No newline at end of file
diff --git a/Engine/Plugins/Animation/RigLogic/Source/RigLogicEditor/Private/AssetDefinitions/AssetDefinition_DNAAsset.h b/Engine/Plugins/Animation/RigLogic/Source/RigLogicEditor/Private/AssetDefinitions/AssetDefinition_DNAAsset.h
new file mode 100644
--- /dev/null
+++ b/Engine/Plugins/Animation/RigLogic/Source/RigLogicEditor/Private/AssetDefinitions/AssetDefinition_DNAAsset.h
@@ -0,0 +1,23 @@
+// Copyright Epic Games, Inc. All Rights Reserved.
+
+#pragma once
+
+#include "AssetDefinitionDefault.h"
+
+#include "AssetDefinition_DNAAsset.generated.h"
+
+UCLASS()
+class UAssetDefinition_DNAAsset : public UAssetDefinitionDefault
+{
+	GENERATED_BODY()
+
+public:
+	//~Begin UAssetDefinitionDefault interface
+	virtual FText GetAssetDisplayName() const override;
+	virtual FLinearColor GetAssetColor() const override;
+	virtual TSoftClassPtr<UObject> GetAssetClass() const override;
+	virtual TConstArrayView<FAssetCategoryPath> GetAssetCategories() const override;
+	virtual UThumbnailInfo* LoadThumbnailInfo(const FAssetData& InAssetData) const override;
+	virtual EAssetCommandResult OpenAssets(const FAssetOpenArgs& InOpenArgs) const override;
+	//~End UAssetDefinitionDefault interface
+};
\ No newline at end of file
diff --git a/Engine/Plugins/Animation/RigLogic/Source/RigLogicEditor/Private/DNAAssetImportFactory.cpp b/Engine/Plugins/Animation/RigLogic/Source/RigLogicEditor/Private/DNAAssetImportFactory.cpp
--- a/Engine/Plugins/Animation/RigLogic/Source/RigLogicEditor/Private/DNAAssetImportFactory.cpp
+++ b/Engine/Plugins/Animation/RigLogic/Source/RigLogicEditor/Private/DNAAssetImportFactory.cpp
@@ -1,32 +1,26 @@
 // Copyright Epic Games, Inc. All Rights Reserved.
 
 #include "DNAAssetImportFactory.h"
-#include "DNAAsset.h"
+#include "DNAAssetUserData.h"
+#include "DNA.h"
 
 #include "Engine/SkeletalMesh.h"
-#include "ReferenceSkeleton.h"
-#include "Rendering/SkeletalMeshModel.h"
-#include "Rendering/SkeletalMeshLODModel.h"
-#include "Rendering/SkeletalMeshRenderData.h"
-#include "Rendering/SkeletalMeshLODRenderData.h"
-#include "ComponentReregisterContext.h"
 #include "UObject/ConstructorHelpers.h"
 #include "Editor.h"
 
-#include "ActorFactories/ActorFactoryStaticMesh.h"
 #include "ProfilingDebugging/ScopedTimers.h"
 
-#include "Engine/StaticMesh.h"
 #include "Misc/MessageDialog.h"
 #include "Misc/Paths.h"
 
 #include "Framework/Notifications/NotificationManager.h"
 #include "Widgets/Notifications/SNotificationList.h"
 
-#include "AssetImportTask.h"
+#include "EditorFramework/AssetImportData.h"
 #include "Misc/FileHelper.h"
+#include "AssetRegistry/AssetRegistryModule.h"
+#include "UObject/Package.h"
 
-#include "DNAAssetImportUI.h"
 #include "DNAImporter.h"
 
 #include UE_INLINE_GENERATED_CPP_BY_NAME(DNAAssetImportFactory)
@@ -39,10 +33,7 @@ UDNAAssetImportFactory::UDNAAssetImportFactory(const FObjectInitializer& ObjectI
 {
 	bCreateNew = false;
 	bEditAfterNew = true;
-	SupportedClass = UDNAAsset::StaticClass();
-
-	bEditorImport = true;
-	bText = false;
+	SupportedClass = UDNA::StaticClass();
 
 	Formats.Add(TEXT("dna;Character DNA file"));
 }
@@ -52,8 +43,6 @@ void UDNAAssetImportFactory::PostInitProperties()
 	Super::PostInitProperties();
 	bEditorImport = true;
 	bText = false;
-
-	ImportUI = NewObject<UDNAAssetImportUI>(this, NAME_None, RF_NoFlags);
 }
 
 bool UDNAAssetImportFactory::CanReimport(UObject* Obj, TArray<FString>& OutFilenames)
@@ -64,146 +53,69 @@ bool UDNAAssetImportFactory::CanReimport(UObject* Obj, TArray<FString>& OutFilen
 	return FactoryExtensions.Contains(FPaths::GetExtension(PreferredReimportPath));
 }
 
-EReimportResult::Type UDNAAssetImportFactory::Reimport(UObject* Obj, int32 SourceFileIndex)
+void UDNAAssetImportFactory::SetReimportPaths(UObject* Obj, const TArray<FString>& NewReimportPaths)
 {
-	// This is done here in oreder to enable importing of the dna file with the same name as 
-	// skeletal mesh. This is treated as an normal import but engine is recognizing it as an skeletal mesh reimport.
-	// Also this is called for regular Import since DNA import and reimport don't differ that much
-
-	UPackage* Pkg = Obj->GetPackage();
-	FString Name = FPaths::GetBaseFilename(PreferredReimportPath);
-	bool bImportWasCancelled = false;
-	UObject* Result = ImportObject(SupportedClass, Pkg, FName(*Name), RF_Public | RF_Standalone | RF_Transactional, PreferredReimportPath, nullptr, bImportWasCancelled);
+	UDNA* DNAAsset = Cast<UDNA>(Obj);
+	if (DNAAsset && ensure(NewReimportPaths.Num() == 1) && DNAAsset->AssetImportData)
+	{
+		DNAAsset->AssetImportData->UpdateFilenameOnly(NewReimportPaths[0]);
+	}
+}
 
-	// Clean up and remove the factories we created from the root set
-	CleanUp();
+EReimportResult::Type UDNAAssetImportFactory::Reimport(UObject* Obj)
+{
+	UDNA* DNAAsset = Cast<UDNA>(Obj);
+	bool Result = false;
 
-	if (Result != nullptr)
+	if(DNAAsset && DNAAsset->AssetImportData->GetFirstFilename().Len() > 0)
 	{
-		// TODO: Find a better solution to showing the right message, withouth this we get reimport message success for Skeletal Mesh
-		FNotificationInfo Info(LOCTEXT("DNA_ReimportSuccessMessage", "DNA file successfully imported"));
-		Info.ExpireDuration = 3.0f;
-		Info.bUseLargeFont = false;
-		TSharedPtr<SNotificationItem> Notification = FSlateNotificationManager::Get().AddNotification(Info);
-		if (Notification.IsValid())
-		{
-			Notification->SetCompletionState(SNotificationItem::CS_Success);
-		}
+		Result = DNAAsset->Init(DNAAsset->AssetImportData->GetFirstFilename());
+	} 
 
-		//Cancel is sent here because the notification created for successful import ( reimport is only registered because DNA has the same name as SkelMesh ) is already showed with above code
-		return EReimportResult::Cancelled;
+	if (Result)
+	{
+		return EReimportResult::Succeeded;
 	}
 
 	return EReimportResult::Failed;
 }
 
-int32 UDNAAssetImportFactory::GetPriority() const
-{
-	return ImportPriority;
-}
-
-UObject* UDNAAssetImportFactory::FactoryCreateFile(UClass* InClass, UObject* InParent, FName InName, EObjectFlags Flags, const FString& Filename, const TCHAR* Parms, FFeedbackContext* Warn, bool& bOutOperationCanceled)
+UObject* UDNAAssetImportFactory::FactoryCreateFile(UClass* InClass, UObject* InParent, FName InName, EObjectFlags InFlags, const FString& Filename, const TCHAR* Parms, FFeedbackContext* Warn, bool& bOutOperationCanceled)
 {
-	// this one prints all messages that are stored in FFbxImporter
-	FDNAImporter* DNAImporter = FDNAImporter::GetInstance();
-	FDNAAssetImportOptions* ImportOptions = DNAImporter->GetImportOptions();
-	//Clean up the options
-	FDNAAssetImportOptions::ResetOptions(ImportOptions);
+	UDNA* DNAAsset = NewObject<UDNA>(InParent, InClass, InName, InFlags | RF_Public | RF_Standalone | RF_Transactional);
 
-	struct FRestoreImportUI
+	if (DNAAsset->Init(*Filename))
 	{
-		FRestoreImportUI(UDNAAssetImportFactory* InDNAAssetImportFactory)
-			: DNAAssetImportFactory(InDNAAssetImportFactory)
+		if (!DNAAsset->AssetImportData)
 		{
-			ensure(DNAAssetImportFactory->OriginalImportUI == nullptr);
-			DNAAssetImportFactory->OriginalImportUI = DNAAssetImportFactory->ImportUI;
+			DNAAsset->AssetImportData = NewObject<UAssetImportData>(DNAAsset);
 		}
+				
+		DNAAsset->AssetImportData->UpdateFilenameOnly(Filename);
 
-		~FRestoreImportUI()
+		
+		FNotificationInfo Info(LOCTEXT("DNA_ImportSuccessMessage", "DNA file successfully imported!"));
+		Info.ExpireDuration = 3.0f;
+		Info.bFireAndForget = true;
+		Info.bUseLargeFont = false;
+		TSharedPtr<SNotificationItem> Notification = FSlateNotificationManager::Get().AddNotification(Info);
+		if (Notification.IsValid())
 		{
-			DNAAssetImportFactory->ImportUI = DNAAssetImportFactory->OriginalImportUI;
-			DNAAssetImportFactory->OriginalImportUI = nullptr;
+			Notification->SetCompletionState(SNotificationItem::CS_Success);
 		}
 
-	private:
-		UDNAAssetImportFactory* DNAAssetImportFactory;
-	};
-	FRestoreImportUI RestoreImportUI(this);
-	UDNAAssetImportUI* OverrideImportUI = AssetImportTask ? Cast<UDNAAssetImportUI>(AssetImportTask->Options) : nullptr;
-	if (OverrideImportUI)
-	{
-		ImportUI = OverrideImportUI;
-	}
-	//We are not re-importing
-	ImportUI->bIsReimport = false;
-	ImportUI->ReimportMesh = nullptr;
-
-	// Show the import dialog only when not in a "yes to all" state or when automating import
-	bool bIsAutomated = IsAutomatedImport() || IsAutomatedReimport();
-	bool bShowImportDialog = !bIsAutomated;
-
-	// If the user provided a SkeletalMesh through the AssetImportTask Options, which means we have automated import,
-	// refrain from searching for the Skeletal Mesh in the asset package
-	if (!IsValid(ImportUI->SkeletalMesh) || !bIsAutomated) {
-		USkeletalMesh* SkeletalMeshForDNAImport = nullptr;
-		ForEachObjectWithOuter(InParent, [&SkeletalMeshForDNAImport](UObject* Object)
-			{
-				if (USkeletalMesh* SkeletalMesh = Cast<USkeletalMesh>(Object))
-				{
-					SkeletalMeshForDNAImport = SkeletalMesh;
-					return false;
-				}
-				return true;
-			});
-		ImportUI->SkeletalMesh = SkeletalMeshForDNAImport;
+		GEditor->GetEditorSubsystem<UImportSubsystem>()->BroadcastAssetPostImport(this, DNAAsset);
+		return DNAAsset;
 	}
-
-	ImportOptions = GetImportOptions(DNAImporter, ImportUI, bShowImportDialog, bIsAutomated, InParent->GetPathName(), bOutOperationCanceled, UFactory::CurrentFilename);
-
-	if (bOutOperationCanceled)
+	else
 	{
-		// User cancelled, clean up and return
-		DNAImporter->PartialCleanUp();
-		return nullptr;
-	}
-
-	if (ImportOptions)
-	{
-		USkeletalMesh* SkelMesh = ImportOptions->SkeletalMesh;
-		if (SkelMesh)
-		{
-			FString SkelMeshFullPath(SkelMesh->GetPathName());
-			FString SkelMeshJustPath(FPaths::GetPath(*SkelMeshFullPath));
-			FString SkelMeshFileName(FPaths::GetBaseFilename(*SkelMeshFullPath));
-			FString DNAAssetFileName(TEXT("DNA"));
-			FString DNAAssetFullName(SkelMeshJustPath);
-			DNAAssetFullName.PathAppend(*SkelMeshFileName, SkelMeshFileName.Len());
-			DNAAssetFullName.Append(TEXT("."));
-			DNAAssetFullName.Append(DNAAssetFileName);
-
-			GEditor->GetEditorSubsystem<UImportSubsystem>()->BroadcastAssetPreImport(this, InClass, InParent, *DNAAssetFullName, TEXT("dna"));
-
-			UDNAAsset* DNAAsset = NewObject< UDNAAsset >(SkelMesh, FName(*DNAAssetFileName)); //SkelMesh has to be its outer, otherwise DNAAsset won't be saved			
-			if (DNAAsset->Init(*Filename))
-			{
-				SkelMesh->AddAssetUserData(DNAAsset);
-
-				GEditor->GetEditorSubsystem<UImportSubsystem>()->BroadcastAssetPostImport(this, DNAAsset);
-				DNAImporter->PartialCleanUp();
-				return DNAAsset;
-			}
-		}
-		else
-		{
-			// Skeletal mesh not found.
-			const FText Message = LOCTEXT("DNAImportFailed", "DNA file can not be uploaded without Skeletal Mesh.");
-			FMessageDialog::Open(EAppMsgType::Ok, Message);
-			UE_LOG(LogDNAImportFactory, Warning, TEXT("%s"), *Message.ToString());
-		}
+		// TODO: Use FMessageLog instead
+		const FText Message = LOCTEXT("DNAImportFailed", "DNA asset failed to initialize from given file:");
+		FMessageDialog::Open(EAppMsgType::Ok, Message);
+		UE_LOGFMT(LogDNAImportFactory, Warning, "{} {}", *Message.ToString(), Filename);
 	}
 
 	// Failed to load file or create DNA stream. Clean and return nullptr.
-	DNAImporter->PartialCleanUp();
 	GEditor->GetEditorSubsystem<UImportSubsystem>()->BroadcastAssetPostImport(this, nullptr);
 	return nullptr;
 }
@@ -220,32 +132,5 @@ bool UDNAAssetImportFactory::FactoryCanImport(const FString& Filename)
 	return false;
 }
 
-// DNA Asset Import UI Implementation.
-UDNAAssetImportUI::UDNAAssetImportUI(const FObjectInitializer& ObjectInitializer)
-	: Super(ObjectInitializer)
-{
-	bIsReimport = false;
-	ReimportMesh = nullptr;
-	//Make sure we are transactional to allow undo redo
-	this->SetFlags(RF_Transactional);
-}
-
-
-bool UDNAAssetImportUI::CanEditChange(const FProperty* InProperty) const
-{
-	bool bIsMutable = Super::CanEditChange(InProperty);
-	return bIsMutable;
-}
-
-void UDNAAssetImportUI::ParseFromJson(TSharedRef<class FJsonObject> ImportSettingsJson)
-{
-}
-
-void UDNAAssetImportUI::ResetToDefault()
-{
-	ReloadConfig();
-}
-// END of DNA Asset Import UI Implementation.
-
 #undef LOCTEXT_NAMESPACE
 
diff --git a/Engine/Plugins/Animation/RigLogic/Source/RigLogicEditor/Private/DNAAssetImportWindow.cpp b/Engine/Plugins/Animation/RigLogic/Source/RigLogicEditor/Private/DNAAssetImportWindow.cpp
--- a/Engine/Plugins/Animation/RigLogic/Source/RigLogicEditor/Private/DNAAssetImportWindow.cpp
+++ b/Engine/Plugins/Animation/RigLogic/Source/RigLogicEditor/Private/DNAAssetImportWindow.cpp
@@ -16,11 +16,8 @@
 
 void SDNAAssetImportWindow::Construct(const FArguments& InArgs)
 {
-	ImportUI = InArgs._ImportUI;
 	WidgetWindow = InArgs._WidgetWindow;
 
-	check(ImportUI);
-
 	TSharedPtr<SBox> ImportTypeDisplay;
 	TSharedPtr<SHorizontalBox> DNAHeaderButtons;
 	TSharedPtr<SBox> InspectorBox;
@@ -142,21 +139,16 @@ void SDNAAssetImportWindow::Construct(const FArguments& InArgs)
 		]
 		]
 	);
-
-	DetailsView->SetObject(ImportUI);
 }
 
 FReply SDNAAssetImportWindow::OnResetToDefaultClick() const
 {
-	ImportUI->ResetToDefault();
-	//Refresh the view to make sure the custom UI are updating correctly
-	DetailsView->SetObject(ImportUI, true);
 	return FReply::Handled();
 }
 
 FText SDNAAssetImportWindow::GetImportTypeDisplayText() const
 {
-	return ImportUI->bIsReimport ? LOCTEXT("DNAAssetImportWindow_ReImportTypeAnim", "Reimport DNA") : LOCTEXT("DNAAssetImportWindow_ImportTypeAnim", "Import DNA");
+	return LOCTEXT("DNAAssetImportWindow_ImportTypeAnim", "Import DNA");
 }
 
 bool SDNAAssetImportWindow::CanImport()  const
diff --git a/Engine/Plugins/Animation/RigLogic/Source/RigLogicEditor/Private/DNAImporter.cpp b/Engine/Plugins/Animation/RigLogic/Source/RigLogicEditor/Private/DNAImporter.cpp
--- a/Engine/Plugins/Animation/RigLogic/Source/RigLogicEditor/Private/DNAImporter.cpp
+++ b/Engine/Plugins/Animation/RigLogic/Source/RigLogicEditor/Private/DNAImporter.cpp
@@ -5,202 +5,270 @@
 =============================================================================*/
 
 #include "DNAImporter.h"
-#include "AssetRegistry/AssetData.h"
-#include "CoreMinimal.h"
-#include "Modules/ModuleManager.h"
+#include "AssetToolsModule.h"
+#include "AssetImportTask.h"
 #include "AssetRegistry/AssetRegistryModule.h"
-#include "Interfaces/IMainFrameModule.h"
-#include "AssetRegistry/ARFilter.h"
-#include "Widgets/SWindow.h"
+#include "DNAAssetImportFactory.h"
+#include "EditorReimportHandler.h" 
+#include "EditorFramework/AssetImportData.h"
+#include "DNA.h"
+#include "DNAUtils.h"
+#include "UObject/UObjectGlobals.h"
+#include "DNAAssetUserData.h"
 #include "Framework/Application/SlateApplication.h"
+#include "Framework/Notifications/NotificationManager.h"
+#include "Widgets/Notifications/SNotificationList.h"
 #include "HAL/PlatformApplicationMisc.h"
-#include "DNAAssetImportWindow.h"
 #include "DesktopPlatformModule.h"
 #include "EditorDirectories.h"
 
-TSharedPtr<FDNAImporter> FDNAImporter::StaticInstance;
-TSharedPtr<FDNAImporter> FDNAImporter::StaticPreviewInstance;
+DEFINE_LOG_CATEGORY_STATIC(LogDNAImporter, Log, All);
 
-#define LOCTEXT_NAMESPACE "DNAAssetImport"
+#define LOCTEXT_NAMESPACE "DNAImport"
 
-FDNAAssetImportOptions* GetImportOptions(FDNAImporter * DNAImporter, UDNAAssetImportUI * ImportUI, bool bShowOptionDialog, bool bIsAutomated, const FString & FullPath, bool & OutOperationCanceled, const FString & InFilename)
+
+UDNA* UDNAImporter::ImportDNAWithPrompt(TArray<USkeletalMesh*> SkeletalMeshes)
 {
-	OutOperationCanceled = false;
+	const FString Filename =  PromptForDNAImportFile();
 
-	if (bShowOptionDialog)
+	if(SkeletalMeshes.IsEmpty())
 	{
-		FDNAAssetImportOptions* ImportOptions = DNAImporter->GetImportOptions();
-		ImportOptions->SkeletalMesh = ImportUI->SkeletalMesh;
+		const FText Message = LOCTEXT("DNA_ImportSkeletalMeshMissing", "Skeletal Mesh was not provided for Import function.");
+		UE_LOGFMT(LogDNAImporter, Error, "{}", *Message.ToString());
+		return nullptr;
+	}
+	UDNA* ImportedDNA = ImportDNA(Filename, SkeletalMeshes[0]);
 
-		// if SkeletalMesh was set by outside, please make sure copy back to UI
-		if (ImportOptions->SkeletalMesh)
-		{
-			ImportUI->SkeletalMesh = ImportOptions->SkeletalMesh;
-		}
-		else
+	if(ImportedDNA)
+	{
+		for (USkeletalMesh* SkelMesh : SkeletalMeshes)
 		{
-			// Look in the current target directory to see if we have a skeleton
-			FARFilter Filter;
-			Filter.PackagePaths.Add(*FPaths::GetPath(FullPath));
-			Filter.ClassPaths.Add(USkeletalMesh::StaticClass()->GetClassPathName());
-
-			IAssetRegistry& AssetRegistry = FModuleManager::LoadModuleChecked<FAssetRegistryModule>("AssetRegistry").Get();
-			TArray<FAssetData> SkeletalMeshAssets;
-			AssetRegistry.GetAssets(Filter, SkeletalMeshAssets);
-			if (SkeletalMeshAssets.Num() > 0)
+			// If there is a legacy DNA Asset User data, we should remove it
+			if (UDNAAsset* LegacyDNA = SkelMesh->GetAssetUserData<UDNAAsset>())
 			{
-				ImportUI->SkeletalMesh = CastChecked<USkeletalMesh>(SkeletalMeshAssets[0].GetAsset());
+				SkelMesh->RemoveUserDataOfClass(UDNAAsset::StaticClass());
 			}
-			else
-			{
-				ImportUI->SkeletalMesh = NULL;
-			}
-		}
 
-		DNAImporter->SetDNAFileName(*FPaths::GetCleanFilename(InFilename));
+			UDNAAssetUserData* DNAAssetUserData = NewObject<UDNAAssetUserData>(SkelMesh);
+			DNAAssetUserData->DNAAsset = ImportedDNA;
+			SkelMesh->AddAssetUserData(DNAAssetUserData);
+			SkelMesh->MarkPackageDirty();
+		}
+		return ImportedDNA;
+	}
+	else
+	{
+		return nullptr;
+	}
+}
 
-		//This option must always be the same value has the skeletalmesh one.
+bool UDNAImporter::ImportDNAAutomated(const FString FileName, USkeletalMesh* SkeletalMesh)
+{
+	UDNA* ImportedDNA = ImportDNA(FileName, SkeletalMesh);
+	if (ImportedDNA)
+	{
+		SkeletalMesh->RemoveUserDataOfClass(UDNAAsset::StaticClass());
+		
+		UDNAAssetUserData* DNAUserData = NewObject<UDNAAssetUserData>(SkeletalMesh);
+		DNAUserData->DNAAsset = ImportedDNA;
+		SkeletalMesh->AddAssetUserData(DNAUserData);
+		SkeletalMesh->MarkPackageDirty();
 		
+		return true;
+	}
 
-		//////////////////////////////////////////////////////////////////////////
-		// Set the information section data
+	return false;
+}
 
-		//Make sure the file is open to be able to read the header before showing the options
-		//If the file is already open it will simply return false.
-		// do analytics on getting DNA data
-		ImportUI->FileName = DNAImporter->GetDNAFileName();
+UDNA* UDNAImporter::ImportDNA(const FString Filename, USkeletalMesh* SkeletalMesh)
+{
+	
+	UDNAAssetImportFactory* DNAFactory = NewObject<UDNAAssetImportFactory>(UDNAAssetImportFactory::StaticClass());
+	// We get the package of the first of selected SkeletalMeshes
+	FString ImportPackagePath = FPackageName::GetLongPackagePath(SkeletalMesh->GetOutermost()->GetName());
+
+	FString AssetName = FPaths::GetBaseFilename(Filename);
+
+	// Check for an existing object
+	FString PathToCheck = FString::Printf(TEXT("%s/%s.%s"), *ImportPackagePath, *AssetName, *AssetName);
+	FAssetData ExistingData = FAssetRegistryModule::GetRegistry().GetAssetByObjectPath(PathToCheck);
+	
+
+	// If we have an existing asset with same name (Skeletal Mesh for example) we should add DNA suffix or create unique name
+	if (ExistingData.IsValid())
+	{ 
+		FString BasePackageName = ImportPackagePath/AssetName;
+		FString PackageName;
+		FAssetToolsModule& AssetToolsModule = FModuleManager::Get().LoadModuleChecked<FAssetToolsModule>("AssetTools");
+		AssetToolsModule.Get().CreateUniqueAssetName(BasePackageName, TEXT("_DNA"), PackageName, AssetName);
+	}
 
-		TSharedPtr<SWindow> ParentWindow;
+	TStrongObjectPtr<UAssetImportTask> ImportTask(NewObject<UAssetImportTask>(this));
+	ImportTask->Filename = Filename;
+	ImportTask->DestinationPath = ImportPackagePath;
+	ImportTask->DestinationName = AssetName;
+	ImportTask->bAutomated = true;
+	ImportTask->bReplaceExisting = true;
+	ImportTask->Factory = DNAFactory;
 
-		if (FModuleManager::Get().IsModuleLoaded("MainFrame"))
-		{
-			IMainFrameModule& MainFrame = FModuleManager::LoadModuleChecked<IMainFrameModule>("MainFrame");
-			ParentWindow = MainFrame.GetParentWindow();
-		}
+	FAssetToolsModule::GetModule().Get().ImportAssetTasks({ ImportTask.Get()});
 
-		// Compute centered window position based on max window size, which include when all categories are expanded
-		const float DNAImportWindowWidth = 410.0f;
-		const float DNAImportWindowHeight = 750.0f;
-		FVector2D DNAImportWindowSize = FVector2D(DNAImportWindowWidth, DNAImportWindowHeight); // Max window size it can get based on current slate
+	if (ImportTask->ImportedObjectPaths.Num() == 1)
+	{
+		UDNA* ImportedDNA = Cast<UDNA>(ImportTask->GetObjects()[0]);
+		return ImportedDNA;
+	}
 
+	return nullptr;
+}
 
-		FSlateRect WorkAreaRect = FSlateApplicationBase::Get().GetPreferredWorkArea();
-		FVector2D DisplayTopLeft(WorkAreaRect.Left, WorkAreaRect.Top);
-		FVector2D DisplaySize(WorkAreaRect.Right - WorkAreaRect.Left, WorkAreaRect.Bottom - WorkAreaRect.Top);
+bool UDNAImporter::ConvertFromLegacyAssetUserData(USkeletalMesh* SkeletalMesh)
+{
+	if(UDNAAsset* LegacyDNA = SkeletalMesh->GetAssetUserData<UDNAAsset>())
+	{
+		UPackage* SkeletalMeshPackage = SkeletalMesh->GetOutermost();
+		FString SkeletalMeshPackagePath = SkeletalMeshPackage->GetName();
 
-		float ScaleFactor = FPlatformApplicationMisc::GetDPIScaleFactorAtPoint(DisplayTopLeft.X, DisplayTopLeft.Y);
-		DNAImportWindowSize *= ScaleFactor;
+		const FString DNAName = FString::Printf(TEXT("%s_DNA"), *SkeletalMesh->GetName());
+		FString DNAPackagePath = FPackageName::GetLongPackagePath(SkeletalMeshPackagePath);
+		DNAPackagePath /= DNAName;
 
-		FVector2D WindowPosition = (DisplayTopLeft + (DisplaySize - DNAImportWindowSize) / 2.0f) / ScaleFactor;
+		UPackage* DNAPackage = CreatePackage(*DNAPackagePath);
 
+		UDNA* NewDNAAsset = NewObject<UDNA>(DNAPackage, UDNA::StaticClass());
+		NewDNAAsset->SetDNAReader(LegacyDNA->GetDNAReader());
+		NewDNAAsset->AssetImportData = DuplicateObject(LegacyDNA->AssetImportData.Get(), NewDNAAsset);
 
-		TSharedRef<SWindow> Window = SNew(SWindow)
-			.Title(NSLOCTEXT("UnrealEd", "DNAImportOpionsTitle", "DNA Import Options"))
-			.SizingRule(ESizingRule::Autosized)
-			.AutoCenter(EAutoCenter::None)
-			.ClientSize(DNAImportWindowSize)
-			.ScreenPosition(WindowPosition);
+		NewDNAAsset->Rename(*DNAName, DNAPackage);
+		NewDNAAsset->SetFlags(RF_Public | RF_Standalone | RF_Transactional);
+		FAssetRegistryModule::AssetCreated(NewDNAAsset);
+		NewDNAAsset->MarkPackageDirty();
+		
 
-		TSharedPtr<SDNAAssetImportWindow> DNAImportWindow;
-		Window->SetContent
-		(
-			SAssignNew(DNAImportWindow, SDNAAssetImportWindow)
-			.ImportUI(ImportUI)
-			.WidgetWindow(Window)
-			.FullPath(FText::FromString(FullPath))
-			.MaxWindowHeight(DNAImportWindowHeight)
-			.MaxWindowWidth(DNAImportWindowWidth)
-		);
+		SkeletalMesh->RemoveUserDataOfClass(UDNAAsset::StaticClass());
+		UDNAAssetUserData* DNAAssetUserData = NewObject<UDNAAssetUserData>(SkeletalMesh);
+		DNAAssetUserData->DNAAsset = NewDNAAsset;
+		SkeletalMesh->AddAssetUserData(DNAAssetUserData);
+		SkeletalMesh->MarkPackageDirty();
 
-		// @todo: we can make this slow as showing progress bar later
-		FSlateApplication::Get().AddModalWindow(Window, ParentWindow, false);
+		return true;
+	}
+	else
+	{
+		FText SkelMeshName = FText::FromString(SkeletalMesh->GetName());
+		const FText Message = FText::Format(LOCTEXT("DNA_CreateErrorMessage", "{0} SkeletalMesh doesn't have Legacy DNA Asset User Data attached!"), SkelMeshName);
+		FNotificationInfo Info(Message);
+		Info.ExpireDuration = 3.0f;
+		Info.bUseLargeFont = false;
+		TSharedPtr<SNotificationItem> Notification = FSlateNotificationManager::Get().AddNotification(Info);
+		if (Notification.IsValid())
+		{
+			Notification->SetCompletionState(SNotificationItem::CS_Fail);
+		}
 
-		ImportUI->SaveConfig();
+		return false;
+	}
+}
 
-		if (DNAImportWindow->ShouldImport())
-		{
-			// open dialog
-			// see if it's canceled
-			ApplyImportUIToImportOptions(ImportUI, *ImportOptions);
+void UDNAImporter::ReimportDNAs(UDNA* ReimportDNA)
+{
+	if(!ReimportDNA)
+	{
+		const FText Message = LOCTEXT("DNA_ReimportFailedInvalidDNA", "Invalid DNA for reimport!");
+		UE_LOGFMT(LogDNAImporter, Error, "{}", *Message.ToString());
+		return;
+	}
 
-			return ImportOptions;
-		}
-		else
+	if (ReimportDNA->AssetImportData && !ReimportDNA->AssetImportData->GetFirstFilename().IsEmpty())
+	{
+		FString Filename = ReimportDNA->AssetImportData->GetFirstFilename();
+		UDNAAssetImportFactory* DNAFactory = NewObject<UDNAAssetImportFactory>(UDNAAssetImportFactory::StaticClass());
+		
+		bool bSuccess = FReimportManager::Instance()->Reimport(ReimportDNA, false, true, Filename, DNAFactory);
+			
+		if (!bSuccess)
 		{
-			OutOperationCanceled = true;
+			const FText Message = LOCTEXT("DNA_ReimportFailedMessage", "Reimporting of DNA failed");
+			UE_LOGFMT(LogDNAImporter, Error, "{}", *Message.ToString());
 		}
 	}
-	else if (bIsAutomated)
-	{
-		//Automation tests set ImportUI settings directly.  Just copy them over
-		FDNAAssetImportOptions* ImportOptions = DNAImporter->GetImportOptions();
-		//Clean up the options
-		FDNAAssetImportOptions::ResetOptions(ImportOptions);
-		ApplyImportUIToImportOptions(ImportUI, *ImportOptions);
-		return ImportOptions;
-	}
 	else
 	{
-		return DNAImporter->GetImportOptions();
-	}
-
-	return NULL;
-}
-
-void ApplyImportUIToImportOptions(UDNAAssetImportUI* ImportUI, FDNAAssetImportOptions& InOutImportOptions)
-{
-	check(ImportUI);
-	InOutImportOptions.SkeletalMesh = ImportUI->SkeletalMesh;
-}
+		FText DNAName = FText::FromString(ReimportDNA->GetName());
+		const FText Message = FText::Format(LOCTEXT("DNA_ReimportErrorMessage", "DNA Asset or AssetImportData is missing for Reimport {0} "), DNAName);
+		FNotificationInfo Info(Message);
+		Info.ExpireDuration = 3.0f;
+		Info.bUseLargeFont = false;
+		TSharedPtr<SNotificationItem> Notification = FSlateNotificationManager::Get().AddNotification(Info);
+		if (Notification.IsValid())
+		{
+			Notification->SetCompletionState(SNotificationItem::CS_Fail);
+		}
 
-FDNAImporter::~FDNAImporter()
-{
-	CleanUp();
+		UE_LOGFMT(LogDNAImporter, Error, "{}", *Message.ToString());
+	}
 }
 
-FDNAImporter* FDNAImporter::GetInstance()
+void UDNAImporter::ExportDNA(UDNA* ExportDNA)
 {
-	if (!StaticInstance.IsValid())
+	if(!ExportDNA)
 	{
-		StaticInstance = MakeShareable(new FDNAImporter());
+		const FText Message = LOCTEXT("DNA_ExportFailedInvalidDNA", "Invalid DNA for export!");
+		UE_LOGFMT(LogDNAImporter, Error, "{}", *Message.ToString());
+		return;
 	}
-	return StaticInstance.Get();
-}
 
-void FDNAImporter::DeleteInstance()
-{
-	StaticInstance.Reset();
-}
+	IDesktopPlatform* DesktopPlatform =
+		FDesktopPlatformModule::Get();
 
-FDNAAssetImportOptions* FDNAImporter::GetImportOptions() const
-{
-	return ImportOptions;
-}
+	if (DesktopPlatform)
+	{
+		FString ChosenFolder;
+		const bool bPicked = DesktopPlatform->OpenDirectoryDialog(
+			nullptr,
+			TEXT("Choose Folder"),
+			TEXT(""),
+			ChosenFolder
+		);
 
-FDNAImporter::FDNAImporter()
-	: ImportOptions(NULL)
-{
-	ImportOptions = new FDNAAssetImportOptions();
-	FMemory::Memzero(*ImportOptions);
+		if (bPicked)
+		{
+			FString FileName = FPaths::MakeValidFileName(ExportDNA->GetName());
+			const FString Extension = TEXT(".dna");
+			const FString FullExportPath = FPaths::Combine(ChosenFolder, FileName + Extension);
 
-	CurPhase = NOTSTARTED;
-}
 
-void FDNAImporter::CleanUp()
-{
-	PartialCleanUp();
+			if (ExportDNA && !ChosenFolder.IsEmpty())
+			{
+				WriteDNAToFile(ExportDNA->GetDNAReader().Get(), EDNADataLayer::All, FullExportPath);
+			}
 
-	delete ImportOptions;
-	ImportOptions = NULL;
-}
+			FText Message;
+			FText DNAName = FText::FromString(ExportDNA->GetName());
+			bool IsExported = IFileManager::Get().FileExists(*FullExportPath);
 
-void FDNAImporter::PartialCleanUp()
-{
-	// reset
-	CurPhase = NOTSTARTED;
+			if(IsExported)
+			{
+				Message = FText::Format(LOCTEXT("DNA_ExportMessageSuccess", "DNA Asset {0} successfuly exported!"), DNAName);
+			}
+			else
+			{
+				Message = FText::Format(LOCTEXT("DNA_ExportMessageFail", "DNA Asset {0} failed to be exported. "), DNAName);
+			}
+
+			FNotificationInfo Info(Message);
+			Info.ExpireDuration = 3.0f;
+			Info.bUseLargeFont = false;
+			TSharedPtr<SNotificationItem> Notification = FSlateNotificationManager::Get().AddNotification(Info);
+			if (Notification.IsValid())
+			{
+				Notification->SetCompletionState(IsExported? SNotificationItem::CS_Success : SNotificationItem::CS_Fail);
+			}
+		}
+	}
 }
 
-FString FDNAImporter::PromptForDNAImportFile()
+FString UDNAImporter::PromptForDNAImportFile()
 {
 	const FText PromptTitle = LOCTEXT("DNAPromptTitle", "Choose a file to import for DNA");
 
diff --git a/Engine/Plugins/Animation/RigLogic/Source/RigLogicEditor/Private/DNAImporterLibrary.cpp b/Engine/Plugins/Animation/RigLogic/Source/RigLogicEditor/Private/DNAImporterLibrary.cpp
--- a/Engine/Plugins/Animation/RigLogic/Source/RigLogicEditor/Private/DNAImporterLibrary.cpp
+++ b/Engine/Plugins/Animation/RigLogic/Source/RigLogicEditor/Private/DNAImporterLibrary.cpp
@@ -2,6 +2,9 @@
 
 #include "DNAImporterLibrary.h"
 #include "DNAAssetImportFactory.h"
+#include "DNAImporter.h"
+#include "Misc/Paths.h"
+#include "HAL/FileManager.h"
 
 #include UE_INLINE_GENERATED_CPP_BY_NAME(DNAImporterLibrary)
 
@@ -11,15 +14,26 @@ DEFINE_LOG_CATEGORY_STATIC(LogDNAImporterLibrary, Log, All);
 
 void UDNAImporterLibrary::ImportSkeletalMeshDNA(const FString FileName, UObject* Mesh)
 {
-	UDNAAssetImportFactory* Factory = Cast<UDNAAssetImportFactory>(UDNAAssetImportFactory::StaticClass()->GetDefaultObject());
-	//Reimport will do the same thing as import we just won't get problems when having the same DNA name as SkeletalMesh, new DNA gets initialized anyway
-	bool bSuccess = FReimportManager::Instance()->Reimport(Mesh, false, false, FileName, Factory, INDEX_NONE, false, true, false);
-	
-	if (!bSuccess)
+	// Clean the path if needed
+	FString CleanedFileName = FileName.TrimStartAndEnd();
+
+	// Convert to absolute if needed (relative to project dir)
+	if (FPaths::IsRelative(CleanedFileName))
 	{
-		const FText Message = LOCTEXT("DNA_ReimportFailedMessage", "Reimporting of DNA failed");
+		CleanedFileName = FPaths::ConvertRelativePathToFull(CleanedFileName);
+	}
+
+	FPaths::MakeStandardFilename(CleanedFileName);           
+	FPaths::CollapseRelativeDirectories(CleanedFileName);    
+
+	TStrongObjectPtr<UDNAImporter> DNAImporter(NewObject<UDNAImporter>());
+
+	if (!DNAImporter->ImportDNAAutomated(CleanedFileName, Cast<USkeletalMesh>(Mesh)))
+	{
+		const FText Message = LOCTEXT("DNAImportFromSkeletalMeshMenuFail", "Importing of DNA failed");
 		UE_LOG(LogDNAImporterLibrary, Error, TEXT("%s"), *Message.ToString());
 	}
+
 }
 
 #undef LOCTEXT_NAMESPACE
diff --git a/Engine/Plugins/Animation/RigLogic/Source/RigLogicEditor/Private/RigLogicEditor.cpp b/Engine/Plugins/Animation/RigLogic/Source/RigLogicEditor/Private/RigLogicEditor.cpp
--- a/Engine/Plugins/Animation/RigLogic/Source/RigLogicEditor/Private/RigLogicEditor.cpp
+++ b/Engine/Plugins/Animation/RigLogic/Source/RigLogicEditor/Private/RigLogicEditor.cpp
@@ -1,21 +1,15 @@
 // Copyright Epic Games, Inc. All Rights Reserved.
 
 #include "RigLogicEditor.h"
-#include "Modules/ModuleManager.h"
 #include "Logging/LogMacros.h"
-#include "DNAImporter.h"
 #include "RigUnit_RigLogic.h"
 #include "EditorFramework/AssetImportData.h"
 #include "IAssetTypeActions.h"
-#include "AssetToolsModule.h"
 #include "Toolkits/AssetEditorToolkit.h"
-#include "Framework/MultiBox/MultiBoxBuilder.h"
 #include "ISettingsModule.h"
+#include "DNAAssetUserData.h"
+#include "DNA.h"
 #include "PropertyEditorModule.h"
-#include <ContentBrowserModule.h>
-#include <DNAAssetImportFactory.h>
-#include "Framework/Notifications/NotificationManager.h"
-#include "Widgets/Notifications/SNotificationList.h"
 #include "Misc/Paths.h"
 #include "AssetRegistry/IAssetRegistry.h"
 #include "AssetRegistry/AssetData.h"
@@ -30,151 +24,52 @@ DEFINE_LOG_CATEGORY_STATIC(LogRigLogicEditor, Log, All);
 
 void FRigLogicEditor::StartupModule()
 {
-	FContentBrowserModule& ContentBrowserModule = FModuleManager::LoadModuleChecked<FContentBrowserModule>(TEXT("ContentBrowser"));
-	TArray<FContentBrowserMenuExtender_SelectedAssets>& MenuExtenderDelegates = ContentBrowserModule.GetAllAssetViewContextMenuExtenders();
-
-	MenuExtenderDelegates.Add(FContentBrowserMenuExtender_SelectedAssets::CreateStatic(&FRigLogicEditor::OnExtendSkelMeshWithDNASelectionMenu));
-
 	UObject::FAssetRegistryTag::OnGetExtraObjectTagsWithContext.AddStatic(&GetAssetRegistryTagsForDNA);
 }
 
-
-TSharedRef<FExtender> FRigLogicEditor::OnExtendSkelMeshWithDNASelectionMenu(const TArray<FAssetData>& SelectedAssets)
-{
-	TSharedRef<FExtender> Extender = MakeShared<FExtender>();
-
-	Extender->AddMenuExtension(
-		"GetAssetActions",
-		EExtensionHook::After,
-		nullptr,
-		FMenuExtensionDelegate::CreateStatic(&FRigLogicEditor::CreateDnaActionsSubMenu, SelectedAssets)
-	);
-	return Extender;
-}
-
-void FRigLogicEditor::CreateDnaActionsSubMenu(FMenuBuilder& MenuBuilder, const TArray<FAssetData> SelectedAssets)
-{
-	check(!SelectedAssets.IsEmpty());
-	const FAssetData& Asset = SelectedAssets[0];
-	UClass* AssetClass = Asset.GetClass();
-
-	if ((AssetClass != nullptr) && AssetClass->IsChildOf(USkeletalMesh::StaticClass()))
-	{
-		MenuBuilder.AddSubMenu(
-			LOCTEXT("DNASkeletalMeshSubmenu", "MetaHuman DNA"),
-			LOCTEXT("DNAImportSubmenu_ToolTip", "DNA related actions"),
-			FNewMenuDelegate::CreateStatic(&FRigLogicEditor::GetDNAMenu, SelectedAssets),
-			false,
-			FSlateIcon(FAppStyle::GetAppStyleSetName(), "Icons.Import")
-		);
-	}
-
-}
-
-void FRigLogicEditor::GetDNAMenu(FMenuBuilder& MenuBuilder, const TArray<FAssetData> SelectedAssets)
-{
-	check(!SelectedAssets.IsEmpty());
-	auto Mesh = SelectedAssets[0].GetAsset();
-
-	MenuBuilder.AddMenuEntry(
-		LOCTEXT("Import DNA", "Import new DNA File"),
-		LOCTEXT("ImportDNA_Tooltip", "Import DNA"),
-		FSlateIcon(FAppStyle::GetAppStyleSetName(), "Icons.Import"),
-		FUIAction(FExecuteAction::CreateStatic(&FRigLogicEditor::ExecuteDNAImport, static_cast<UObject*>(Mesh)))
-	);
-
-	MenuBuilder.AddMenuEntry(
-		LOCTEXT("Reimport DNA", "Reimport existing DNA File"),
-		LOCTEXT("ReimportDNA_Tooltip ", "Reimport DNA"),
-		FSlateIcon(FAppStyle::GetAppStyleSetName(), "Icons.Reimport"),
-		FUIAction(FExecuteAction::CreateStatic(&FRigLogicEditor::ExecuteDNAReimport, static_cast<UObject*>(Mesh)))
-	);
-}
-
-void FRigLogicEditor::ExecuteDNAImport(UObject* Mesh)
-{
-	FAssetToolsModule& AssetToolsModule = FModuleManager::Get().LoadModuleChecked<FAssetToolsModule>("AssetTools");
-
-	FDNAImporter* DNAImporter = FDNAImporter::GetInstance();
-	const TArray<FString> Filenames = { DNAImporter->PromptForDNAImportFile() };
-	UDNAAssetImportFactory* Factory = Cast<UDNAAssetImportFactory>(UDNAAssetImportFactory::StaticClass()->GetDefaultObject());
-
-	//Reimport will do the same thing as import we just won't get problems when having the same DNA name as SkeletalMesh, new DNA gets initialized anyway
-	bool bSuccess = FReimportManager::Instance()->Reimport(Mesh, false, true, Filenames[0], Factory);
-
-	if (!bSuccess)
-	{
-		const FText Message = LOCTEXT("DNA_ReimportFailedMessage", "Reimporting of DNA failed");
-		UE_LOG(LogRigLogicEditor, Error, TEXT("%s"), *Message.ToString());
-	}
-}
-
-void FRigLogicEditor::ExecuteDNAReimport(class UObject* Mesh)
+void FRigLogicEditor::GetAssetRegistryTagsForDNA(FAssetRegistryTagsContext Context)
 {
-	FAssetToolsModule& AssetToolsModule = FModuleManager::Get().LoadModuleChecked<FAssetToolsModule>("AssetTools");
-	USkeletalMesh* SkelMesh = Cast<USkeletalMesh>(Mesh);
-	const TArray<UAssetUserData*>* AssetData = SkelMesh->GetAssetUserDataArray();
-	UDNAAsset* DNAAsset = nullptr;
-
-	for (UAssetUserData* AssetDataElement : *AssetData)
-	{
-		if (Cast<UDNAAsset>(AssetDataElement))
-		{
-			DNAAsset = Cast<UDNAAsset>(AssetDataElement);
-		}
-	}
-
-	if (!AssetData->IsEmpty() && DNAAsset && DNAAsset->AssetImportData && !DNAAsset->AssetImportData->GetFirstFilename().IsEmpty())
-	{
-		const TArray<FString> Filenames = { DNAAsset->AssetImportData->GetFirstFilename() };
-		UFactory* Factory = Cast<UFactory>(UDNAAssetImportFactory::StaticClass()->GetDefaultObject());
-
-		bool bSuccess = FReimportManager::Instance()->Reimport(Mesh, false, true, Filenames[0]);
-
-		if(!bSuccess)
-		{
-			const FText Message = LOCTEXT("DNA_ReimportFailedMessage", "Reimporting of DNA failed");
-			UE_LOG(LogRigLogicEditor, Error, TEXT("%s"), *Message.ToString());
-		}
-	}
-	else
+	const USkeletalMesh* SkelMesh = Cast<USkeletalMesh>(Context.GetObject());
+	if (SkelMesh)
 	{
-		const FText Message = LOCTEXT("DNA_ReimportErrorMessage", "There is no DNA file attached to do Reimport");
-		FNotificationInfo Info(Message);
-		Info.ExpireDuration = 3.0f;
-		Info.bUseLargeFont = false;
-		TSharedPtr<SNotificationItem> Notification = FSlateNotificationManager::Get().AddNotification(Info);
-		if (Notification.IsValid())
+		FString DNAName;
+		
+		const TArray<UAssetUserData*>* AssetUserDataArray = SkelMesh->GetAssetUserDataArray();
+		if (AssetUserDataArray->IsEmpty())
 		{
-			Notification->SetCompletionState(SNotificationItem::CS_Fail);
+			return;
 		}
 
-		UE_LOG(LogRigLogicEditor, Error, TEXT("%s"), *Message.ToString());
-	}
-}
-
-void FRigLogicEditor::GetAssetRegistryTagsForDNA(FAssetRegistryTagsContext Context)
-{
-	const UObject* Object = Context.GetObject();
-	if ((Object != nullptr) && (Object->GetClass() != nullptr) && Object->GetClass()->IsChildOf(USkeletalMesh::StaticClass()))
-	{
-		FString DNAname = (LOCTEXT("DnaNotOnSkeletalMesh", "No DNA Attached")).ToString();
-		USkeletalMesh* SkelMesh = const_cast<USkeletalMesh*>(Cast<USkeletalMesh>(Object));
-		const TArray<UAssetUserData*>* AssetDataArr = SkelMesh->GetAssetUserDataArray();
-		if (!AssetDataArr->IsEmpty())
+		for (const UAssetUserData* AssetUserData : *AssetUserDataArray)
 		{
-			UAssetUserData* UserData = SkelMesh->GetAssetUserDataOfClass(UDNAAsset::StaticClass());
-			if (UserData)
+			if (const UDNAAssetUserData* DNAAssetUserData = Cast<UDNAAssetUserData>(AssetUserData))
 			{
-				const UDNAAsset* DNAAsset = Cast<UDNAAsset>(UserData);
-				if (DNAAsset && DNAAsset->AssetImportData)
+				if (UDNA* DNA = DNAAssetUserData->DNAAsset)
 				{
-					DNAname = DNAAsset->AssetImportData->GetFirstFilename();
-					FPaths::NormalizeFilename(DNAname);
+					DNAName = DNA->GetPackage()->GetName();
+					break;
 				}
 			}
+			else
+			{
+				if (const UDNAAsset* DNAAsset = Cast<UDNAAsset>(AssetUserData))
+				{
+					if (DNAAsset->AssetImportData)
+					{
+						DNAName = DNAAsset->AssetImportData->GetFirstFilename();
+						FPaths::NormalizeFilename(DNAName);
+						break;
+					}
+				}
+			}
+		}
+
+		if (DNAName.IsEmpty())
+		{
+			DNAName = (LOCTEXT("DnaNotOnSkeletalMesh", "No DNA Attached")).ToString();
 		}
-		Context.AddTag(UObject::FAssetRegistryTag("DNA", DNAname, UObject::FAssetRegistryTag::TT_Alphabetical));
+			
+		Context.AddTag(UObject::FAssetRegistryTag("DNA", DNAName, UObject::FAssetRegistryTag::TT_Alphabetical));
 	}
 }
 
diff --git a/Engine/Plugins/Animation/RigLogic/Source/RigLogicEditor/Public/DNAAssetImportFactory.h b/Engine/Plugins/Animation/RigLogic/Source/RigLogicEditor/Public/DNAAssetImportFactory.h
--- a/Engine/Plugins/Animation/RigLogic/Source/RigLogicEditor/Public/DNAAssetImportFactory.h
+++ b/Engine/Plugins/Animation/RigLogic/Source/RigLogicEditor/Public/DNAAssetImportFactory.h
@@ -20,35 +20,28 @@ class UDNAAssetImportFactory: public UFactory, public FReimportHandler
 	GENERATED_UCLASS_BODY()
 
 public:
-	UPROPERTY()
-		TObjectPtr<class UDNAAssetImportUI> ImportUI;
 
-	/** Prevent garbage collection of original when overriding ImportUI property */
-	UPROPERTY()
-		TObjectPtr<class UDNAAssetImportUI> OriginalImportUI;
+PRAGMA_DISABLE_DEPRECATION_WARNINGS
+	UE_DEPRECATED(5.8, "Import UI is not needed since it adds unnecessary complexity and therefor it's deprecated")
+	UPROPERTY(meta = (DeprecatedProperty, DeprecationMessage = "Import UI is not needed since it adds unnecessary complexity and therefor it's deprecated"))
+	TObjectPtr<class UDEPRECATED_DNAAssetImportUI> ImportUI_DEPRECATED;
 
+	UE_DEPRECATED(5.8, "Import UI is not needed since it adds unnecessary complexity and therefor it's deprecated")
+	UPROPERTY(meta = (DeprecatedProperty, DeprecationMessage = "Import UI is not needed since it adds unnecessary complexity and therefor it's deprecated"))
+	TObjectPtr<class UDEPRECATED_DNAAssetImportUI> OriginalImportUI_DEPRECATED;
+PRAGMA_ENABLE_DEPRECATION_WARNINGS
 	/** UObject properties */
 	virtual void PostInitProperties() override;
 
 	//~ Begin FReimportHandler Interface
 	virtual bool CanReimport(UObject* Obj, TArray<FString>& OutFilenames) override;
-	virtual void SetReimportPaths(UObject* Obj, const TArray<FString>& NewReimportPaths) override
-	{
-		return SetReimportPaths(Obj, NewReimportPaths[0], 0);
-	}
-	virtual void SetReimportPaths(UObject* Obj, const FString& NewReimportPath, const int32 SourceIndex) override {};
-
-	virtual EReimportResult::Type Reimport(UObject* Obj) override
-	{
-		return Reimport(Obj, INDEX_NONE);
-	}
-	virtual EReimportResult::Type Reimport(UObject* Obj, int32 SourceFileIndex);
-	virtual int32 GetPriority() const override;
+	virtual void SetReimportPaths(UObject* Obj, const TArray<FString>& NewReimportPaths) override;
+	virtual EReimportResult::Type Reimport(UObject* Obj) override;
 	virtual void PostImportCleanUp() { CleanUp(); }
 	//~ End FReimportHandler Interface
 
 	/** UFactory interface */
-	virtual UObject* FactoryCreateFile(UClass* InClass, UObject* InParent, FName InName, EObjectFlags Flags, const FString& Filename, const TCHAR* Parms, FFeedbackContext* Warn, bool& bOutOperationCanceled) override;
+	virtual UObject* FactoryCreateFile(UClass* InClass, UObject* InParent, FName InName, EObjectFlags InFlags, const FString& Filename, const TCHAR* Parms, FFeedbackContext* Warn, bool& bOutOperationCanceled) override;
 	virtual bool FactoryCanImport(const FString& Filename) override;
 	virtual void CleanUp() override {};
 };
diff --git a/Engine/Plugins/Animation/RigLogic/Source/RigLogicEditor/Public/DNAAssetImportUI.h b/Engine/Plugins/Animation/RigLogic/Source/RigLogicEditor/Public/DNAAssetImportUI.h
--- a/Engine/Plugins/Animation/RigLogic/Source/RigLogicEditor/Public/DNAAssetImportUI.h
+++ b/Engine/Plugins/Animation/RigLogic/Source/RigLogicEditor/Public/DNAAssetImportUI.h
@@ -13,10 +13,10 @@
 #include "Factories/ImportSettings.h"
 #include "DNAAssetImportUI.generated.h"
 
-UCLASS(config = EditorPerProjectUserSettings, AutoExpandCategories = (Mesh), MinimalAPI)
-class UDNAAssetImportUI : public UObject, public IImportSettingsParser
+UCLASS(Deprecated, meta = (DeprecationMessage = "This class is not needed and will be removed."))
+class UDEPRECATED_DNAAssetImportUI : public UObject, public IImportSettingsParser
 {
-	GENERATED_UCLASS_BODY()
+	GENERATED_BODY()
 
 public:
 
@@ -25,13 +25,10 @@ public:
 	TObjectPtr<class USkeletalMesh> SkeletalMesh;
 
 	UFUNCTION(BlueprintCallable, Category = Miscellaneous)
-	void ResetToDefault();
-
-	/** UObject Interface */
-	virtual bool CanEditChange(const FProperty* InProperty) const override;
+	void ResetToDefault() {};
 
 	/** IImportSettings Interface */
-	virtual void ParseFromJson(TSharedRef<class FJsonObject> ImportSettingsJson) override;
+	virtual void ParseFromJson(TSharedRef<class FJsonObject> ImportSettingsJson) override {};
 
 	/* Whether this UI is construct for a reimport */
 	bool bIsReimport;
diff --git a/Engine/Plugins/Animation/RigLogic/Source/RigLogicEditor/Public/DNAAssetImportWindow.h b/Engine/Plugins/Animation/RigLogic/Source/RigLogicEditor/Public/DNAAssetImportWindow.h
--- a/Engine/Plugins/Animation/RigLogic/Source/RigLogicEditor/Public/DNAAssetImportWindow.h
+++ b/Engine/Plugins/Animation/RigLogic/Source/RigLogicEditor/Public/DNAAssetImportWindow.h
@@ -11,19 +11,17 @@
 #include "DNAAssetImportUI.h"
 
 class SButton;
-
-class SDNAAssetImportWindow : public SCompoundWidget
+//UE_DEPRECATED(5.8, "SDNAAssetImportWindow is deprecated and not needed anymore") [[deprecated("deprecated, no replacement")]]
+class UE_DEPRECATED(5.8, "SDNAAssetImportWindow is deprecated and not needed anymore") SDNAAssetImportWindow : public SCompoundWidget
 {
 public:
 	SLATE_BEGIN_ARGS(SDNAAssetImportWindow)
-		: _ImportUI(NULL)
-		, _WidgetWindow()
+		: _WidgetWindow()
 		, _FullPath()
 		, _MaxWindowHeight(0.0f)
 		, _MaxWindowWidth(0.0f)
 		{}
 
-		SLATE_ARGUMENT(UDNAAssetImportUI*, ImportUI)
 		SLATE_ARGUMENT(TSharedPtr<SWindow>, WidgetWindow)
 		SLATE_ARGUMENT(FText, FullPath)
 		SLATE_ARGUMENT(float, MaxWindowHeight)
@@ -75,8 +73,7 @@ public:
 	}
 
 	SDNAAssetImportWindow()
-		: ImportUI(NULL)
-		, bShouldImport(false)
+		: bShouldImport(false)
 	{}
 
 private:
@@ -86,7 +83,6 @@ private:
 	FText GetImportTypeDisplayText() const;
 
 private:
-	UDNAAssetImportUI*	ImportUI;
 	TSharedPtr<class IDetailsView> DetailsView;
 	TWeakPtr< SWindow > WidgetWindow;
 	TSharedPtr< SButton > ImportButton;
diff --git a/Engine/Plugins/Animation/RigLogic/Source/RigLogicEditor/Public/DNAImporter.h b/Engine/Plugins/Animation/RigLogic/Source/RigLogicEditor/Public/DNAImporter.h
--- a/Engine/Plugins/Animation/RigLogic/Source/RigLogicEditor/Public/DNAImporter.h
+++ b/Engine/Plugins/Animation/RigLogic/Source/RigLogicEditor/Public/DNAImporter.h
@@ -3,76 +3,52 @@
 #pragma once
 
 #include "CoreMinimal.h"
-#include "UObject/WeakObjectPtr.h"
-#include "Misc/SecureHash.h"
-#include "DNAAssetImportUI.h"
-#include "Logging/TokenizedMessage.h"
-#include "MeshBuild.h"
+#include "UObject/Object.h"
+#include "DNAImporter.generated.h"
 
 class USkeletalMesh;
 
-struct FDNAAssetImportOptions
+UCLASS(BlueprintType)
+class UDNAImporter : public UObject
 {
-	// General options
-	bool bCanShowDialog;	
-	// Skeletal Mesh options
-	USkeletalMesh* SkeletalMesh;
+	GENERATED_BODY()
 
-	static void ResetOptions(FDNAAssetImportOptions *OptionsToReset)
-	{
-		check(OptionsToReset != nullptr);
-		*OptionsToReset = FDNAAssetImportOptions();
-	}
-};
+public:
+	
+	/** 
+	* Importing DNA asset with prompt for choosing a DNA file from disk 
+	*/
+	UFUNCTION(CallInEditor, BlueprintCallable, Category = "DNAImporter")
+	UDNA* ImportDNAWithPrompt(TArray<USkeletalMesh*> SkeletalMeshes);
 
-FDNAAssetImportOptions* GetImportOptions(class FDNAImporter* DNAImporter, UDNAAssetImportUI* ImportUI, bool bShowOptionDialog, bool bIsAutomated, const FString& FullPath, bool& OutOperationCanceled, const FString& InFilename);
-void ApplyImportUIToImportOptions(UDNAAssetImportUI* ImportUI, FDNAAssetImportOptions& InOutImportOptions);
+	/**
+	* Importing DNA asset automatically from sent file and SkeletalMesh provided
+	*/
+	UFUNCTION(CallInEditor, BlueprintCallable, Category = "DNAImporter")
+	bool ImportDNAAutomated(const FString FileName, USkeletalMesh* SkeletalMesh);
 
-class FDNAImporter
-{
-public:
-	~FDNAImporter();
 	/**
-	 * Returns the importer singleton. It will be created on the first request.
-	 */
-	static FDNAImporter* GetInstance();
-	static void DeleteInstance();
+	* Creating a DNA Asset from legacy DNA Asset User Data that should exist on Skeletal Mesh
+	*/
+	UFUNCTION(CallInEditor, BlueprintCallable, Category = "DNAImporter")
+	bool ConvertFromLegacyAssetUserData(USkeletalMesh* SkeletalMesh);
 
-	FDNAAssetImportOptions* ImportOptions;
 	/**
-	 * Get the object of import options
-	 *
-	 * @return FDNAAssetImportOptions
-	 */
-	FDNAAssetImportOptions* GetImportOptions() const;
-	void PartialCleanUp();
-	FString PromptForDNAImportFile();
-	FString GetDNAFileName() { return DNAFileName; }
-	void SetDNAFileName(FString DNAName) { this->DNAFileName = DNAName; };
-protected:
-	enum IMPORTPHASE
-	{
-		NOTSTARTED,
-		FILEOPENED,
-		IMPORTED,
-		FIXEDANDCONVERTED,
-	};
+	* Reimporting DNA assets associated with SkeletalMeshes provided
+	*/
+	UFUNCTION(CallInEditor, BlueprintCallable, Category = "DNAImporter")
+	void ReimportDNAs(UDNA* ReimportDNA);
 
-	static TSharedPtr<FDNAImporter> StaticInstance;
-	static TSharedPtr<FDNAImporter> StaticPreviewInstance;
+	/**
+	* Export DNA asset to choosen folder
+	*/
+	UFUNCTION(CallInEditor, BlueprintCallable, Category = "DNAImporter")
+	void ExportDNA(UDNA* ExportDNA);
 
-	// scene management
-	IMPORTPHASE CurPhase;
-	FString ErrorMessage;
-	// base path of DNA file
-	FString FileBasePath;
-	TWeakObjectPtr<UObject> Parent;
-	FString DNAFileName;
 
-	FDNAImporter();
+private:
+
+	FString PromptForDNAImportFile();
+	UDNA* ImportDNA(const FString Filename, USkeletalMesh* SkeletalMesh);
 
-	/**
-	 * Clean up to reset the import.
-	 */
-	void CleanUp();
 };
\ No newline at end of file
diff --git a/Engine/Plugins/Animation/RigLogic/Source/RigLogicEditor/RigLogicEditor.Build.cs b/Engine/Plugins/Animation/RigLogic/Source/RigLogicEditor/RigLogicEditor.Build.cs
--- a/Engine/Plugins/Animation/RigLogic/Source/RigLogicEditor/RigLogicEditor.Build.cs
+++ b/Engine/Plugins/Animation/RigLogic/Source/RigLogicEditor/RigLogicEditor.Build.cs
@@ -21,18 +21,23 @@ namespace UnrealBuildTool.Rules
 					"PropertyEditor",
 					"SlateCore",
 					"ApplicationCore",
+					"ToolMenus",
+					"ContentBrowser",
+					"ContentBrowserData",
+					"AssetDefinition",
 					"Slate",
 					"InputCore",
 					"EditorWidgets",
-					"DesktopPlatform"
+					"DesktopPlatform",
+					"AssetTools"
 				}
 			);
 
 			PrivateIncludePathModuleNames.AddRange(
 				new string[]
 				{
-					"AssetTools",
-					"Settings"
+					"Settings",
+					"UnrealEd"
 				}
 			);
 		}
diff --git a/Engine/Plugins/Animation/RigLogic/Source/RigLogicModule/Private/AnimNode_RigLogic.cpp b/Engine/Plugins/Animation/RigLogic/Source/RigLogicModule/Private/AnimNode_RigLogic.cpp
--- a/Engine/Plugins/Animation/RigLogic/Source/RigLogicModule/Private/AnimNode_RigLogic.cpp
+++ b/Engine/Plugins/Animation/RigLogic/Source/RigLogicModule/Private/AnimNode_RigLogic.cpp
@@ -4,6 +4,8 @@
 
 #include "Components/SkeletalMeshComponent.h"
 #include "DNAAsset.h"
+#include "DNA.h"
+#include "DNAAssetUserData.h"
 #include "DNAIndexMapping.h"
 #include "DNAReader.h"
 #include "Engine/SkeletalMesh.h"
@@ -172,13 +174,27 @@ void FAnimNode_RigLogic::CacheBones_AnyThread(const FAnimationCacheBonesContext&
 		return;
 	}
 
-	UDNAAsset* DNAAsset = Cast<UDNAAsset>(SkeletalMesh->GetAssetUserDataOfClass(UDNAAsset::StaticClass()));
-	if (DNAAsset == nullptr)
+	TSharedPtr<FSharedRigRuntimeContext> SharedRigRuntimeContext;
+	TSharedPtr<FDNAIndexMapping> SharedDNAIndexMapping;
+	
+	// We either have DNAAssetUserData on SKM and then use it to get DNA asset or we use legacy asset user data as a fallback option in else
+	if (UDNAAssetUserData* DNAAssetUserData = SkeletalMesh->GetAssetUserData<UDNAAssetUserData>())
 	{
-		return;
+		if (UDNA* DNA = DNAAssetUserData->DNAAsset)
+		{
+			SharedRigRuntimeContext = DNA->GetRigRuntimeContext();
+			SharedDNAIndexMapping = DNA->GetDNAIndexMapping(Skeleton, SkeletalMesh);
+		}
 	}
-
-	TSharedPtr<FSharedRigRuntimeContext> SharedRigRuntimeContext = DNAAsset->GetRigRuntimeContext();
+	else
+	{
+		if(UDNAAsset* DNAAsset = SkeletalMesh->GetAssetUserData<UDNAAsset>())
+		{
+			SharedRigRuntimeContext = DNAAsset->GetRigRuntimeContext();
+			SharedDNAIndexMapping = DNAAsset->GetDNAIndexMapping(Skeleton, SkeletalMesh);
+		}
+	}
+	 
 	if (!SharedRigRuntimeContext.IsValid())
 	{
 		return;
@@ -196,7 +212,6 @@ void FAnimNode_RigLogic::CacheBones_AnyThread(const FAnimationCacheBonesContext&
 
 	RigInstance->SetLOD(Context.AnimInstanceProxy->GetLODLevel());
 
-	TSharedPtr<FDNAIndexMapping> SharedDNAIndexMapping = DNAAsset->GetDNAIndexMapping(Skeleton, SkeletalMesh);
 	if (LocalDNAIndexMapping != SharedDNAIndexMapping)
 	{
 		const uint16 LODCount = LocalRigRuntimeContext->RigLogic->GetLODCount();
diff --git a/Engine/Plugins/Animation/RigLogic/Source/RigLogicModule/Private/DNA.cpp b/Engine/Plugins/Animation/RigLogic/Source/RigLogicModule/Private/DNA.cpp
new file mode 100644
--- /dev/null
+++ b/Engine/Plugins/Animation/RigLogic/Source/RigLogicModule/Private/DNA.cpp
@@ -0,0 +1,403 @@
+// Copyright Epic Games, Inc. All Rights Reserved.
+
+#include "DNA.h"
+#include "DNAUtils.h"
+
+#include "ArchiveMemoryStream.h"
+#include "DNAAssetCustomVersion.h"
+#include "DNAReaderAdapter.h"
+#include "DNAIndexMapping.h"
+#include "FMemoryResource.h"
+#include "RigLogicDNAReader.h"
+#include "RigLogicMemoryStream.h"
+#include "SharedRigRuntimeContext.h"
+
+#if WITH_EDITORONLY_DATA
+#include "EditorFramework/AssetImportData.h"
+#endif
+#include "Engine/AssetUserData.h"
+#include "Serialization/BufferArchive.h"
+#include "Serialization/MemoryReader.h"
+#include "Misc/Paths.h"
+#include "Misc/FileHelper.h"
+#include "UObject/UObjectGlobals.h"
+#include "HAL/LowLevelMemTracker.h"
+#include "Animation/Skeleton.h"
+#include "Components/SkeletalMeshComponent.h"
+#include "Engine/SkeletalMesh.h"
+
+#include "riglogic/RigLogic.h"
+
+#include UE_INLINE_GENERATED_CPP_BY_NAME(DNA)
+
+DEFINE_LOG_CATEGORY(LogDNA);
+namespace DNA_Helpers
+{
+	static constexpr uint32 AVG_EMPTY_SIZE = 4 * 1024;
+	static constexpr uint32 AVG_DNA_SIZE = 55 * 1024 * 1024;
+	static constexpr uint32 AVG_MACHINE_LEARNED_BEHAVIOR_SIZE = 5 * 1024 * 1024;
+
+	static TSharedPtr<IDNAReader> ReadDNAFromStream(rl4::BoundedIOStream* Stream, EDNADataLayer Layer, uint16 MaxLOD)
+	{
+		auto DNAStreamReader = rl4::makeScoped<dna::BinaryStreamReader>(Stream, CalculateDNADataLayerBitmask(Layer), dna::UnknownLayerPolicy::Preserve, MaxLOD, FMemoryResource::Instance());
+		DNAStreamReader->read();
+		if (!rl4::Status::isOk())
+		{
+			UE_LOG(LogDNA, Error, TEXT("%s"), ANSI_TO_TCHAR(rl4::Status::get().message));
+			return nullptr;
+		}
+		return MakeShared<FDNAReader<dna::BinaryStreamReader>>(DNAStreamReader.release());
+	}
+
+	static void WriteDNAToStream(const IDNAReader* Source, EDNADataLayer Layer, rl4::BoundedIOStream* Destination)
+	{
+		auto DNAWriter = rl4::makeScoped<dna::BinaryStreamWriter>(Destination, FMemoryResource::Instance());
+		if (Source != nullptr)
+		{
+			DNAWriter->setFrom(Source->Unwrap(), CalculateDNADataLayerBitmask(Layer), dna::UnknownLayerPolicy::Preserve, FMemoryResource::Instance());
+		}
+		DNAWriter->write();
+	}
+
+	static TSharedPtr<IDNAReader> CopyDNALayer(const IDNAReader* Source, EDNADataLayer DNADataLayer, uint32 PredictedSize)
+	{
+		// To avoid lots of reallocations in `FRigLogicMemoryStream`, reserve an approximate size
+		// that we know would cause at most one reallocation in the worst case (but none for the average DNA)
+		TArray<uint8> MemoryBuffer;
+		MemoryBuffer.Reserve(PredictedSize);
+
+		FRigLogicMemoryStream MemoryStream(&MemoryBuffer);
+		WriteDNAToStream(Source, DNADataLayer, &MemoryStream);
+
+		MemoryStream.seek(0ul);
+
+		return ReadDNAFromBuffer(&MemoryBuffer, DNADataLayer);
+	}
+
+	static TSharedPtr<IDNAReader> CreateEmptyDNA(uint32 PredictedSize)
+	{
+		// To avoid lots of reallocations in `FRigLogicMemoryStream`, reserve an approximate size
+		// that we know would cause at most one reallocation in the worst case (but none for the average DNA)
+		TArray<uint8> MemoryBuffer;
+		MemoryBuffer.Reserve(PredictedSize);
+
+		FRigLogicMemoryStream MemoryStream(&MemoryBuffer);
+		WriteDNAToStream(nullptr, EDNADataLayer::All, &MemoryStream);
+
+		MemoryStream.seek(0ul);
+
+		return ReadDNAFromBuffer(&MemoryBuffer, EDNADataLayer::All);
+	}
+}
+
+
+UDNA::UDNA() : RigRuntimeContext{ nullptr }
+{
+}
+
+UDNA::~UDNA() = default;
+
+TSharedPtr<IDNAReader> UDNA::GetDNAReader()
+{
+	UE::TReadScopeLock DNAScopeLock{ DNAUpdateLock };
+	return DNAReader;
+}
+
+void UDNA::SetDNAReader(TSharedPtr<IDNAReader> SourceDNAReader)
+{
+	UE::TWriteScopeLock DNAScopeLock{ DNAUpdateLock };
+	size_t PredictedSize = DNA_Helpers::AVG_EMPTY_SIZE;
+	if (SourceDNAReader.IsValid())
+	{
+		PredictedSize = (SourceDNAReader->GetNeuralNetworkCount() != 0) ? DNA_Helpers::AVG_DNA_SIZE + DNA_Helpers::AVG_MACHINE_LEARNED_BEHAVIOR_SIZE : DNA_Helpers::AVG_DNA_SIZE;
+	}
+
+	EDNADataLayer DNALayers = (
+		EDNADataLayer::Behavior |
+		EDNADataLayer::MachineLearnedBehavior |
+		EDNADataLayer::RBFBehavior
+		);
+
+#if WITH_EDITORONLY_DATA
+	DNALayers |= EDNADataLayer::Geometry;
+#endif
+
+	DNAReader = DNA_Helpers::CopyDNALayer(SourceDNAReader.Get(), DNALayers, PredictedSize);
+	InvalidateRigRuntimeContext();
+	InitializeRigRuntimeContext();
+}
+
+void UDNA::InitializeForRuntimeFrom(UDNA* Other)
+{
+	UE::TWriteScopeLock DNAScopeLock{ DNAUpdateLock };
+	UE::TWriteScopeLock ContextScopeLock{ RigRuntimeContextUpdateLock };
+	UE::TWriteScopeLock MappingScopeLock{ DNAIndexMappingUpdateLock };
+
+	// Store a reference to the other asset's runtime context, using the accessor function to 
+	// ensure that the necessary lock is taken.
+	//
+	// The runtime context is immutable, i.e. not modified after initialization, so it's safe to
+	// share across UDNAs.
+	RigRuntimeContext = Other->GetRigRuntimeContext();
+
+	// DNAReader is used at runtime by GetDNAIndexMapping, so it needs to be populated here.
+	//
+	// The reference is taken from RigRuntimeContext to ensure it's consistent with the runtime
+	// context, as the UDNA's DNAReader can be changed to point to a new one before the
+	// runtime context is updated.
+	//
+	// As with the runtime context itself, the DNAReader is immutable and safe to share.
+	DNAReader = RigRuntimeContext->DNAReader;
+
+	// Ensure bKeepDNAAfterInitialization reflects the current state of the DNAReader, i.e.
+	// if bKeepDNAAfterInitialization is true, DNAReader will contain DNA and vice versa.
+	bKeepDNAAfterInitialization = Other->bKeepDNAAfterInitialization;
+
+	// This map is populated on demand, so doesn't need to be copied.
+	DNAIndexMappingContainer.Empty(1);
+
+	// Clear any fields not needed at runtime to avoid any old data causing confusion
+#if WITH_EDITORONLY_DATA
+	AssetImportData = nullptr;
+#endif
+}
+
+
+void UDNA::InvalidateRigRuntimeContext()
+{
+	UE::TWriteScopeLock ContextScopeLock{ RigRuntimeContextUpdateLock };
+	UE::TWriteScopeLock MappingScopeLock(DNAIndexMappingUpdateLock);
+	RigRuntimeContext = nullptr;
+	DNAIndexMappingContainer.Empty(1);
+}
+
+void UDNA::InitializeRigRuntimeContext()
+{
+	LLM_SCOPE_BYNAME(TEXT("Animation/RigLogic"));
+
+	// Assumes DNAUpdateLock is locked by caller
+
+	TSharedPtr<FSharedRigRuntimeContext> NewContext = MakeShared<FSharedRigRuntimeContext>();
+	if (DNAReader.IsValid() && (DNAReader->GetJointCount() != 0))
+	{
+		NewContext->DNAReader = DNAReader;
+		// Convert behavior data from DNA to UE space on the fly as RigLogic accesses it
+		RigLogicDNAReader BehaviorReaderInUESpace{ DNAReader->Unwrap() };
+		FDNAReader<RigLogicDNAReader> BehaviorReaderInUESpaceWrapper{ &BehaviorReaderInUESpace };
+		NewContext->RigLogic = MakeShared<FRigLogic>(&BehaviorReaderInUESpaceWrapper, RigLogicConfiguration);
+		NewContext->CacheVariableJointIndices();
+		if ((DNAReader->GetRBFSolverCount() != 0) || (DNAReader->GetSwingCount() != 0) || (DNAReader->GetTwistCount() != 0))
+		{
+			NewContext->CacheInverseNeutralJointRotations();
+		}
+
+		{
+			UE::TWriteScopeLock ContextScopeLock{ RigRuntimeContextUpdateLock };
+			RigRuntimeContext = NewContext;
+		}
+#if !WITH_EDITOR
+		if (!bKeepDNAAfterInitialization)
+		{
+			DNAReader->Unload(EDNADataLayer::Behavior);
+			DNAReader->Unload(EDNADataLayer::Geometry);
+			DNAReader->Unload(EDNADataLayer::MachineLearnedBehavior);
+			DNAReader->Unload(EDNADataLayer::RBFBehavior);
+		}
+#endif  // !WITH_EDITOR
+	}
+}
+
+TSharedPtr<FSharedRigRuntimeContext> UDNA::GetRigRuntimeContext()
+{
+	LLM_SCOPE_BYNAME(TEXT("Animation/RigLogic"));
+
+	UE::TReadScopeLock ContextScopeLock{ RigRuntimeContextUpdateLock };
+	return RigRuntimeContext;
+}
+
+TSharedPtr<FDNAIndexMapping> UDNA::GetDNAIndexMapping(const USkeleton* Skeleton, const USkeletalMesh* SkeletalMesh)
+{
+	LLM_SCOPE_BYNAME(TEXT("Animation/RigLogic"));
+
+	UE::TReadScopeLock DNAScopeLock{ DNAUpdateLock };
+	UE::TWriteScopeLock MappingScopeLock(DNAIndexMappingUpdateLock);
+
+	// Find currently needed mapping and also clean stale objects along the way (requires only one iteration over the map)
+	TSharedPtr<FDNAIndexMapping> DNAIndexMapping;
+	for (auto Iterator = DNAIndexMappingContainer.CreateIterator(); Iterator; ++Iterator)
+	{
+		if ((Iterator->Key.SkeletalMesh.IsValid()) && (Iterator->Key.Skeleton.IsValid()))
+		{
+			if ((Iterator->Key.SkeletalMesh == SkeletalMesh) && (Iterator->Key.Skeleton == Skeleton))
+			{
+				DNAIndexMapping = Iterator->Value;
+			}
+		}
+		else
+		{
+			Iterator.RemoveCurrent();
+		}
+	}
+
+	// Check if currently needed mapping exists, and if not, create it now
+	const FGuid SkeletonGuid = Skeleton->GetGuid();
+	if (!DNAIndexMapping.IsValid() || (SkeletonGuid != DNAIndexMapping->SkeletonGuid))
+	{
+		DNAIndexMapping = MakeShared<FDNAIndexMapping>();
+		DNAIndexMapping->Init(DNAReader.Get(), Skeleton, SkeletalMesh);
+		DNAIndexMappingContainer.Add({ SkeletalMesh, Skeleton }, DNAIndexMapping);
+	}
+
+	return DNAIndexMapping;
+}
+
+bool UDNA::Init(const FString& DNAFilename)
+{
+	LLM_SCOPE_BYNAME(TEXT("Animation/RigLogic"));
+
+	if (!rl4::Status::isOk())
+	{
+		UE_LOG(LogDNA, Warning, TEXT("%s"), ANSI_TO_TCHAR(rl4::Status::get().message));
+	}
+
+#if WITH_EDITORONLY_DATA
+	AssetImportData = NewObject<UAssetImportData>(this, TEXT("AssetImportData"));
+	TArray<FAssetImportInfo::FSourceFile> SourceFiles = { FAssetImportInfo::FSourceFile(DNAFilename) };
+	AssetImportData->SetSourceFiles(MoveTemp(SourceFiles));
+#endif
+
+	if (!FPaths::FileExists(DNAFilename))
+	{
+		UE_LOG(LogDNA, Error, TEXT("DNA file %s doesn't exist!"), *DNAFilename);
+		return false;
+	}
+
+	// Temporary buffer for the DNA file
+	TArray<uint8> TempFileBuffer;
+
+	if (!FFileHelper::LoadFileToArray(TempFileBuffer, *DNAFilename)) //load entire DNA file into the array
+	{
+		UE_LOG(LogDNA, Error, TEXT("Couldn't read DNA file %s!"), *DNAFilename);
+		return false;
+	}
+
+	UE::TWriteScopeLock DNAScopeLock{ DNAUpdateLock };
+
+	DNAReader = ReadDNAFromBuffer(&TempFileBuffer, EDNADataLayer::All, MaxLOD);
+	if (!DNAReader.IsValid())
+	{
+		return false;
+	}
+
+	InvalidateRigRuntimeContext();
+	InitializeRigRuntimeContext();
+
+	return true;
+}
+
+void UDNA::Serialize(FArchive& Ar)
+{
+	LLM_SCOPE_BYNAME(TEXT("Animation/RigLogic"));
+
+	Super::Serialize(Ar);
+
+	Ar.UsingCustomVersion(FDNAAssetCustomVersion::GUID);
+
+	UE::TWriteScopeLock DNAScopeLock{ DNAUpdateLock };
+
+	if (Ar.CustomVer(FDNAAssetCustomVersion::GUID) >= FDNAAssetCustomVersion::BeforeCustomVersionWasAdded)
+	{
+		if (Ar.IsLoading())
+		{
+			if (Ar.CustomVer(FDNAAssetCustomVersion::GUID) == FDNAAssetCustomVersion::BeforeCustomVersionWasAdded)
+			{
+				// This code is referring to version where DNA Asset had different readers for Behavior and Geometry data of the DNA file
+				FArchiveMemoryStream BehaviorStream{ &Ar };
+				const EDNADataLayer BehaviorLayers = (
+					EDNADataLayer::Behavior |
+					EDNADataLayer::MachineLearnedBehavior |
+					EDNADataLayer::RBFBehavior
+					);
+				TSharedPtr<IDNAReader> BehaviorReader = DNA_Helpers::ReadDNAFromStream(&BehaviorStream, BehaviorLayers, MaxLOD);
+				// Geometry data is always present (even if only as an empty placeholder), just so the uasset
+				// format remains consistent between editor and non-editor builds
+				FArchiveMemoryStream GeometryStream{ &Ar };
+				TSharedPtr<IDNAReader> GeometryReader = DNA_Helpers::ReadDNAFromStream(&GeometryStream, EDNADataLayer::Geometry, MaxLOD);
+
+				pma::ScopedPtr<dna::MemoryStream> MemoryStream = pma::makeScoped<dna::MemoryStream>();
+				pma::ScopedPtr<dna::BinaryStreamWriter> DnaWriter = pma::makeScoped<dna::BinaryStreamWriter>(MemoryStream.get());
+
+
+#if WITH_EDITORONLY_DATA
+				DnaWriter->setFrom(GeometryReader->Unwrap(), dna::DataLayer::Geometry);
+#endif
+				DnaWriter->setFrom(BehaviorReader->Unwrap(), CalculateDNADataLayerBitmask(BehaviorLayers));
+
+				DnaWriter->write();
+
+				MemoryStream->seek(0);
+
+				// Now with combined DNA streams we can read the provided memory stream into unified DNAReader
+				DNAReader = DNA_Helpers::ReadDNAFromStream(MemoryStream.get(), EDNADataLayer::All, MaxLOD);
+			}
+
+			else
+			{
+				FArchiveMemoryStream DNAStream{ &Ar };
+				EDNADataLayer DNALayers = (
+					EDNADataLayer::Behavior |
+					EDNADataLayer::MachineLearnedBehavior |
+					EDNADataLayer::RBFBehavior
+					);
+
+				if (!Ar.IsCooking())
+				{
+					DNALayers |= EDNADataLayer::Geometry;
+				}
+
+				DNAReader = DNA_Helpers::ReadDNAFromStream(&DNAStream, DNALayers, MaxLOD);
+			}
+			InvalidateRigRuntimeContext();
+			InitializeRigRuntimeContext();
+		}
+
+		if (Ar.IsSaving())
+		{
+			TSharedPtr<IDNAReader> EmptyDNA = DNA_Helpers::CreateEmptyDNA(DNA_Helpers::AVG_EMPTY_SIZE);
+			IDNAReader* ReaderPtr = (DNAReader.IsValid() ? static_cast<IDNAReader*>(DNAReader.Get()) : EmptyDNA.Get());
+			FArchiveMemoryStream DNAStream{ &Ar };
+
+			EDNADataLayer DNALayers = (
+				EDNADataLayer::Behavior |
+				EDNADataLayer::MachineLearnedBehavior |
+				EDNADataLayer::RBFBehavior
+				);
+#if WITH_EDITORONLY_DATA
+			// Geometry data is discarded unless in Editor
+			DNALayers |= EDNADataLayer::Geometry;
+#endif // #if WITH_EDITORONLY_DATA
+
+			DNA_Helpers::WriteDNAToStream(ReaderPtr, DNALayers, &DNAStream);
+		}
+	}
+}
+
+#if WITH_EDITOR
+void UDNA::PostEditChangeProperty(FPropertyChangedEvent& PropertyChangedEvent)
+{
+	Super::PostEditChangeProperty(PropertyChangedEvent);
+
+	if (PropertyChangedEvent.MemberProperty == nullptr)
+	{
+		return;
+	}
+
+	const FName& MemberPropertyName = PropertyChangedEvent.MemberProperty->GetFName();
+	if (MemberPropertyName == GET_MEMBER_NAME_CHECKED(UDNA, RigLogicConfiguration))
+	{
+		UE::TWriteScopeLock DNAScopeLock{ DNAUpdateLock };
+		InvalidateRigRuntimeContext();
+		InitializeRigRuntimeContext();
+	}
+}
+#endif
diff --git a/Engine/Plugins/Animation/RigLogic/Source/RigLogicModule/Private/DNAAsset.cpp b/Engine/Plugins/Animation/RigLogic/Source/RigLogicModule/Private/DNAAsset.cpp
--- a/Engine/Plugins/Animation/RigLogic/Source/RigLogicModule/Private/DNAAsset.cpp
+++ b/Engine/Plugins/Animation/RigLogic/Source/RigLogicModule/Private/DNAAsset.cpp
@@ -32,62 +32,66 @@
 
 DEFINE_LOG_CATEGORY(LogDNAAsset);
 
-static constexpr uint32 AVG_EMPTY_SIZE = 4 * 1024;
-static constexpr uint32 AVG_DNA_SIZE = 55 * 1024 * 1024;
-static constexpr uint32 AVG_MACHINE_LEARNED_BEHAVIOR_SIZE = 5 * 1024 * 1024;
-
-static TSharedPtr<IDNAReader> ReadDNAFromStream(rl4::BoundedIOStream* Stream, EDNADataLayer Layer, uint16 MaxLOD)
+namespace DNAAsset_Helpers
 {
-	auto DNAStreamReader = rl4::makeScoped<dna::BinaryStreamReader>(Stream, CalculateDNADataLayerBitmask(Layer), dna::UnknownLayerPolicy::Preserve, MaxLOD, FMemoryResource::Instance());
-	DNAStreamReader->read();
-	if (!rl4::Status::isOk())
+	static constexpr uint32 AVG_EMPTY_SIZE = 4 * 1024;
+	static constexpr uint32 AVG_DNA_SIZE = 55 * 1024 * 1024;
+	static constexpr uint32 AVG_MACHINE_LEARNED_BEHAVIOR_SIZE = 5 * 1024 * 1024;
+
+	static TSharedPtr<IDNAReader> ReadDNAFromStream(rl4::BoundedIOStream* Stream, EDNADataLayer Layer, uint16 MaxLOD)
 	{
-		UE_LOG(LogDNAAsset, Error, TEXT("%s"), ANSI_TO_TCHAR(rl4::Status::get().message));
-		return nullptr;
+		auto DNAStreamReader = rl4::makeScoped<dna::BinaryStreamReader>(Stream, CalculateDNADataLayerBitmask(Layer), dna::UnknownLayerPolicy::Preserve, MaxLOD, FMemoryResource::Instance());
+		DNAStreamReader->read();
+		if (!rl4::Status::isOk())
+		{
+			UE_LOG(LogDNAAsset, Error, TEXT("%s"), ANSI_TO_TCHAR(rl4::Status::get().message));
+			return nullptr;
+		}
+		return MakeShared<FDNAReader<dna::BinaryStreamReader>>(DNAStreamReader.release());
 	}
-	return MakeShared<FDNAReader<dna::BinaryStreamReader>>(DNAStreamReader.release());
-}
 
-static void WriteDNAToStream(const IDNAReader* Source, EDNADataLayer Layer, rl4::BoundedIOStream* Destination)
-{
-	auto DNAWriter = rl4::makeScoped<dna::BinaryStreamWriter>(Destination, FMemoryResource::Instance());
-	if (Source != nullptr)
+	static void WriteDNAToStream(const IDNAReader* Source, EDNADataLayer Layer, rl4::BoundedIOStream* Destination)
 	{
-		DNAWriter->setFrom(Source->Unwrap(), CalculateDNADataLayerBitmask(Layer), dna::UnknownLayerPolicy::Preserve, FMemoryResource::Instance());
+		auto DNAWriter = rl4::makeScoped<dna::BinaryStreamWriter>(Destination, FMemoryResource::Instance());
+		if (Source != nullptr)
+		{
+			DNAWriter->setFrom(Source->Unwrap(), CalculateDNADataLayerBitmask(Layer), dna::UnknownLayerPolicy::Preserve, FMemoryResource::Instance());
+		}
+		DNAWriter->write();
 	}
-	DNAWriter->write();
-}
 
-static TSharedPtr<IDNAReader> CopyDNALayer(const IDNAReader* Source, EDNADataLayer DNADataLayer, uint32 PredictedSize)
-{
-	// To avoid lots of reallocations in `FRigLogicMemoryStream`, reserve an approximate size
-	// that we know would cause at most one reallocation in the worst case (but none for the average DNA)
-	TArray<uint8> MemoryBuffer;
-	MemoryBuffer.Reserve(PredictedSize);
+	static TSharedPtr<IDNAReader> CopyDNALayer(const IDNAReader* Source, EDNADataLayer DNADataLayer, uint32 PredictedSize)
+	{
+		// To avoid lots of reallocations in `FRigLogicMemoryStream`, reserve an approximate size
+		// that we know would cause at most one reallocation in the worst case (but none for the average DNA)
+		TArray<uint8> MemoryBuffer;
+		MemoryBuffer.Reserve(PredictedSize);
 
-	FRigLogicMemoryStream MemoryStream(&MemoryBuffer);
-	WriteDNAToStream(Source, DNADataLayer, &MemoryStream);
+		FRigLogicMemoryStream MemoryStream(&MemoryBuffer);
+		WriteDNAToStream(Source, DNADataLayer, &MemoryStream);
 
-	MemoryStream.seek(0ul);
+		MemoryStream.seek(0ul);
 
-	return ReadDNAFromBuffer(&MemoryBuffer, DNADataLayer);
-}
+		return ReadDNAFromBuffer(&MemoryBuffer, DNADataLayer);
+	}
 
-static TSharedPtr<IDNAReader> CreateEmptyDNA(uint32 PredictedSize)
-{
-	// To avoid lots of reallocations in `FRigLogicMemoryStream`, reserve an approximate size
-	// that we know would cause at most one reallocation in the worst case (but none for the average DNA)
-	TArray<uint8> MemoryBuffer;
-	MemoryBuffer.Reserve(PredictedSize);
+	static TSharedPtr<IDNAReader> CreateEmptyDNA(uint32 PredictedSize)
+	{
+		// To avoid lots of reallocations in `FRigLogicMemoryStream`, reserve an approximate size
+		// that we know would cause at most one reallocation in the worst case (but none for the average DNA)
+		TArray<uint8> MemoryBuffer;
+		MemoryBuffer.Reserve(PredictedSize);
 
-	FRigLogicMemoryStream MemoryStream(&MemoryBuffer);
-	WriteDNAToStream(nullptr, EDNADataLayer::All, &MemoryStream);
+		FRigLogicMemoryStream MemoryStream(&MemoryBuffer);
+		WriteDNAToStream(nullptr, EDNADataLayer::All, &MemoryStream);
 
-	MemoryStream.seek(0ul);
+		MemoryStream.seek(0ul);
 
-	return ReadDNAFromBuffer(&MemoryBuffer, EDNADataLayer::All);
+		return ReadDNAFromBuffer(&MemoryBuffer, EDNADataLayer::All);
+	}
 }
 
+
 UDNAAsset::UDNAAsset() : RigRuntimeContext{nullptr}
 {
 }
@@ -122,10 +126,10 @@ TSharedPtr<IDNAReader> UDNAAsset::GetDNAReader()
 void UDNAAsset::SetDNAReader(TSharedPtr<IDNAReader> SourceDNAReader)
 {
 	UE::TWriteScopeLock DNAScopeLock{DNAUpdateLock};
-	size_t PredictedSize = AVG_EMPTY_SIZE;
+	size_t PredictedSize = DNAAsset_Helpers::AVG_EMPTY_SIZE;
 	if(SourceDNAReader.IsValid())
 	{
-		PredictedSize = (SourceDNAReader->GetNeuralNetworkCount() != 0) ? AVG_DNA_SIZE + AVG_MACHINE_LEARNED_BEHAVIOR_SIZE : AVG_DNA_SIZE;
+		PredictedSize = (SourceDNAReader->GetNeuralNetworkCount() != 0) ? DNAAsset_Helpers::AVG_DNA_SIZE + DNAAsset_Helpers::AVG_MACHINE_LEARNED_BEHAVIOR_SIZE : DNAAsset_Helpers::AVG_DNA_SIZE;
 	}	
 
 	EDNADataLayer DNALayers = (
@@ -138,7 +142,7 @@ void UDNAAsset::SetDNAReader(TSharedPtr<IDNAReader> SourceDNAReader)
 		DNALayers |= EDNADataLayer::Geometry;
 #endif
 
-	DNAReader = CopyDNALayer(SourceDNAReader.Get(), DNALayers, PredictedSize);
+	DNAReader = DNAAsset_Helpers::CopyDNALayer(SourceDNAReader.Get(), DNALayers, PredictedSize);
 	InvalidateRigRuntimeContext();
 	InitializeRigRuntimeContext();
 }
@@ -343,11 +347,11 @@ void UDNAAsset::Serialize(FArchive& Ar)
 					EDNADataLayer::MachineLearnedBehavior |
 					EDNADataLayer::RBFBehavior
 					);
-				TSharedPtr<IDNAReader> BehaviorReader = ReadDNAFromStream(&BehaviorStream, BehaviorLayers, 0u); //0u = max LOD
+				TSharedPtr<IDNAReader> BehaviorReader = DNAAsset_Helpers::ReadDNAFromStream(&BehaviorStream, BehaviorLayers, 0u); //0u = max LOD
 				// Geometry data is always present (even if only as an empty placeholder), just so the uasset
 				// format remains consistent between editor and non-editor builds
 				FArchiveMemoryStream GeometryStream{ &Ar };
-				TSharedPtr<IDNAReader> GeometryReader = ReadDNAFromStream(&GeometryStream, EDNADataLayer::Geometry, 0u); //0u = max LOD
+				TSharedPtr<IDNAReader> GeometryReader = DNAAsset_Helpers::ReadDNAFromStream(&GeometryStream, EDNADataLayer::Geometry, 0u); //0u = max LOD
 
 				pma::ScopedPtr<dna::MemoryStream> MemoryStream = pma::makeScoped<dna::MemoryStream>();
 				pma::ScopedPtr<dna::BinaryStreamWriter> DnaWriter = pma::makeScoped<dna::BinaryStreamWriter>(MemoryStream.get());
@@ -363,7 +367,7 @@ void UDNAAsset::Serialize(FArchive& Ar)
 				MemoryStream->seek(0);
 
 				// Now with combined DNA streams we can read the provided memory stream into unified DNAReader
-				DNAReader = ReadDNAFromStream(MemoryStream.get(), EDNADataLayer::All, 0);
+				DNAReader = DNAAsset_Helpers::ReadDNAFromStream(MemoryStream.get(), EDNADataLayer::All, 0);
 			}
 
 			else
@@ -380,7 +384,7 @@ void UDNAAsset::Serialize(FArchive& Ar)
 					DNALayers |= EDNADataLayer::Geometry;
 				}
 
-				DNAReader = ReadDNAFromStream(&DNAStream, EDNADataLayer::All, 0u); //0u = max LOD
+				DNAReader = DNAAsset_Helpers::ReadDNAFromStream(&DNAStream, DNALayers, 0u); //0u = max LOD
 			}
 			InvalidateRigRuntimeContext();
 			InitializeRigRuntimeContext();
@@ -388,7 +392,7 @@ void UDNAAsset::Serialize(FArchive& Ar)
 
 		if (Ar.IsSaving())
 		{
-			TSharedPtr<IDNAReader> EmptyDNA = CreateEmptyDNA(AVG_EMPTY_SIZE);
+			TSharedPtr<IDNAReader> EmptyDNA = DNAAsset_Helpers::CreateEmptyDNA(DNAAsset_Helpers::AVG_EMPTY_SIZE);
 			IDNAReader* ReaderPtr = (DNAReader.IsValid() ? static_cast<IDNAReader*>(DNAReader.Get()) : EmptyDNA.Get());
 			FArchiveMemoryStream DNAStream{ &Ar };
 
@@ -402,7 +406,7 @@ void UDNAAsset::Serialize(FArchive& Ar)
 			DNALayers |= EDNADataLayer::Geometry;
 #endif // #if WITH_EDITORONLY_DATA
 
-			WriteDNAToStream(ReaderPtr, DNALayers, &DNAStream);
+			DNAAsset_Helpers::WriteDNAToStream(ReaderPtr, DNALayers, &DNAStream);
 		}
 	}
 }
diff --git a/Engine/Plugins/Animation/RigLogic/Source/RigLogicModule/Private/DNAToSkelMeshMap.cpp b/Engine/Plugins/Animation/RigLogic/Source/RigLogicModule/Private/DNAToSkelMeshMap.cpp
--- a/Engine/Plugins/Animation/RigLogic/Source/RigLogicModule/Private/DNAToSkelMeshMap.cpp
+++ b/Engine/Plugins/Animation/RigLogic/Source/RigLogicModule/Private/DNAToSkelMeshMap.cpp
@@ -3,7 +3,8 @@
 #include "DNAToSkelMeshMap.h"
 
 #include "DNAAsset.h"
-
+#include "DNA.h"
+#include "DNAAssetUserData.h"
 #include "Serialization/BufferArchive.h"
 #include "Serialization/MemoryReader.h"
 #include "Misc/Paths.h"
@@ -37,17 +38,36 @@ bool FDNAToSkelMeshMap::InitFromDNAAsset(USkeletalMesh* InSkelMesh)
 {
 #if WITH_EDITORONLY_DATA
 	// Fetch DNAAsset from SkelMesh and use its StreamReader
-	UAssetUserData* UserData = InSkelMesh->GetAssetUserDataOfClass(UDNAAsset::StaticClass());
-	if (UserData == nullptr)
+	// We either have DNAAssetUserData on SKM and then use it to get DNA asset or we use legacy asset user data as a fallback option in else
+	if(UDNAAssetUserData* DNAAssetUserData = InSkelMesh->GetAssetUserData<UDNAAssetUserData>())
 	{
-		FString PathName = GetPathNameSafe(InSkelMesh);
-		UE_LOG(LogDNAToSkelMeshMap, Warning, TEXT("Could not find DNAAsset user data for %s"), *PathName);
-		return false;
+		if (UDNA* DNA = DNAAssetUserData->DNAAsset)
+		{
+			InitBaseMesh(DNA->GetDNAReader().Get(), InSkelMesh);
+			return true;
+		}
+		else
+		{
+			FString PathName = GetPathNameSafe(InSkelMesh);
+			UE_LOG(LogDNAToSkelMeshMap, Warning, TEXT("No DNA asset is referenced on DNAAssetUserData for %s"), *PathName);
+			return false;
+		}
+	}
+	else
+	{
+		if(UDNAAsset* DNAAsset = InSkelMesh->GetAssetUserData<UDNAAsset>())
+		{
+			InitBaseMesh(DNAAsset->GetDNAReader().Get(), InSkelMesh);
+			return true;
+		}
+		else
+		{
+			FString PathName = GetPathNameSafe(InSkelMesh);
+			UE_LOG(LogDNAToSkelMeshMap, Warning, TEXT("Could not find DNAAsset user data for %s"), *PathName);
+			return false;
+		}
 	}
-	UDNAAsset* DNAAsset = Cast<UDNAAsset>(UserData);
-	InitBaseMesh(DNAAsset->GetDNAReader().Get(), InSkelMesh);
 
-	return true;
 #else
 	return false;
 #endif //WITH_EDITORONLY_DATA
diff --git a/Engine/Plugins/Animation/RigLogic/Source/RigLogicModule/Private/RigUnit_RigLogic.cpp b/Engine/Plugins/Animation/RigLogic/Source/RigLogicModule/Private/RigUnit_RigLogic.cpp
--- a/Engine/Plugins/Animation/RigLogic/Source/RigLogicModule/Private/RigUnit_RigLogic.cpp
+++ b/Engine/Plugins/Animation/RigLogic/Source/RigLogicModule/Private/RigUnit_RigLogic.cpp
@@ -4,6 +4,9 @@
 
 #include "ControlRig.h"
 #include "DNAReader.h"
+#include "DNA.h"
+#include "DNAAsset.h"
+#include "DNAAssetUserData.h"
 #include "Engine/SkeletalMesh.h"
 #include "RigInstance.h"
 #include "RigLogic.h"
@@ -501,13 +504,25 @@ void FRigUnit_RigLogic_Data::UpdateAnimMapCurves(URigHierarchy* InHierarchy, TAr
 
 TSharedPtr<FSharedRigRuntimeContext> FRigUnit_RigLogic::GetSharedRigRuntimeContext(USkeletalMesh* SkelMesh)
 {
-	UAssetUserData* UserData = SkelMesh->GetAssetUserDataOfClass(UDNAAsset::StaticClass());
-	if (UserData == nullptr)
+	TSharedPtr<FSharedRigRuntimeContext> SharedRigRuntimeContext = nullptr;
+
+	// We either have DNAAssetUserData on SKM and then use it to get DNA asset or we use legacy asset user data as a fallback option in else
+	if (UDNAAssetUserData* DNAAssetUserData = SkelMesh->GetAssetUserData<UDNAAssetUserData>())
 	{
-		return nullptr;
+		if (UDNA* DNA = DNAAssetUserData->DNAAsset)
+		{
+			SharedRigRuntimeContext = DNA->GetRigRuntimeContext();
+		}
 	}
-	UDNAAsset* DNAAsset = Cast<UDNAAsset>(UserData);
-	return DNAAsset->GetRigRuntimeContext();
+	else
+	{
+		if (UDNAAsset* DNAAsset = SkelMesh->GetAssetUserData<UDNAAsset>())
+		{
+			SharedRigRuntimeContext = DNAAsset->GetRigRuntimeContext();
+		}
+	}
+
+	return SharedRigRuntimeContext;
 }
 
 FRigUnit_RigLogic_Execute()
diff --git a/Engine/Plugins/Animation/RigLogic/Source/RigLogicModule/Private/SkelMeshDNAUtils.cpp b/Engine/Plugins/Animation/RigLogic/Source/RigLogicModule/Private/SkelMeshDNAUtils.cpp
--- a/Engine/Plugins/Animation/RigLogic/Source/RigLogicModule/Private/SkelMeshDNAUtils.cpp
+++ b/Engine/Plugins/Animation/RigLogic/Source/RigLogicModule/Private/SkelMeshDNAUtils.cpp
@@ -1,6 +1,7 @@
 // Copyright Epic Games, Inc. All Rights Reserved.
 #include "SkelMeshDNAUtils.h"
 #include "DNAToSkelMeshMap.h"
+#include "DNAAssetUserData.h"
 
 #include "RenderResource.h"
 #include "RHI.h"
@@ -470,19 +471,9 @@ void USkelMeshDNAUtils::UpdateJointBehavior(USkeletalMeshComponent* InSkelMeshCo
 
 UDNAAsset* USkelMeshDNAUtils::GetMeshDNA(USkeletalMesh* InSkelMesh)
 {
-	if (InSkelMesh && InSkelMesh->GetAssetUserDataArray())
+	if (InSkelMesh)
 	{
-		for (UAssetUserData* UserData : *InSkelMesh->GetAssetUserDataArray())
-		{
-			if (UserData && UserData->IsA<UDNAAsset>())
-			{
-				UDNAAsset* DNAAsset = Cast<UDNAAsset>(UserData);
-				if (DNAAsset->GetDNAReader())
-				{
-					return DNAAsset;
-				}
-			}
-		}
+		return InSkelMesh->GetAssetUserData<UDNAAsset>();
 	}
 
 	return nullptr;
diff --git a/Engine/Plugins/Animation/RigLogic/Source/RigLogicModule/Public/DNA.h b/Engine/Plugins/Animation/RigLogic/Source/RigLogicModule/Public/DNA.h
new file mode 100644
--- /dev/null
+++ b/Engine/Plugins/Animation/RigLogic/Source/RigLogicModule/Public/DNA.h
@@ -0,0 +1,160 @@
+// Copyright Epic Games, Inc. All Rights Reserved.
+
+#pragma once
+
+#include "CoreMinimal.h"
+#include "UObject/Object.h"
+#include "UObject/ObjectMacros.h"
+
+#include "Misc/TransactionallySafeRWLock.h"
+
+#include "Engine/AssetUserData.h"
+#include "Interfaces/Interface_AssetUserData.h"
+
+#include "RigLogic.h"
+
+#include "DNA.generated.h"
+
+#define UE_API RIGLOGICMODULE_API
+
+DECLARE_LOG_CATEGORY_EXTERN(LogDNA, Log, All);
+
+class IDNAReader;
+class IBehaviorReader;
+class IGeometryReader;
+class FRigLogicMemoryStream;
+class UAssetUserData;
+class USkeleton;
+class USkeletalMesh;
+class USkeletalMeshComponent;
+struct FDNAIndexMapping;
+struct FSharedRigRuntimeContext;
+enum class EDNADataLayer: uint8;
+
+ /** An asset holding the data needed to generate/update/animate a RigLogic character
+  * It is imported from character's DNA file as a bit stream, and separated out it into runtime (behavior) and design-time chunks;
+  * Currently, the design-time part still loads the geometry, as it is needed for the skeletal mesh update; once SkeletalMeshDNAReader is
+  * fully implemented, it will be able to read the geometry directly from the SkeletalMesh and won't load it into this asset 
+  **/
+UCLASS(MinimalAPI, BlueprintType, meta = (DisplayName = "DNA"))
+class UDNA : public UObject
+{
+	GENERATED_BODY()
+
+public:
+	UE_API UDNA();
+	UE_API ~UDNA();
+
+#if WITH_EDITORONLY_DATA
+	UPROPERTY(VisibleAnywhere, Instanced, Category = ImportSettings)
+	TObjectPtr<class UAssetImportData> AssetImportData;
+#endif
+
+	UE_API TSharedPtr<IDNAReader> GetDNAReader();
+
+
+	/** In non-editor builds, the DNA source data will be unloaded to save memory after the runtime
+	  * data has been initialized from it.
+	  * 
+	  * Set this property to true to keep the DNA in memory, e.g. if you need to modify it at
+	  * runtime. For most use cases, this shouldn't be needed.
+	 **/
+	UPROPERTY(EditAnywhere, Category = RuntimeSettings)
+	bool bKeepDNAAfterInitialization;
+
+	UPROPERTY(EditAnywhere, Category = RuntimeSettings)
+	uint16 MaxLOD = 0u;
+
+	/** Collection of runtime metadata related to a specific character.
+	  * The value field is a FString and requires casting for a derived types.
+	 **/
+	UPROPERTY(EditAnywhere, BlueprintReadWrite, Category = RuntimeSettings)
+	TMap<FString, FString> MetaData;
+
+	UPROPERTY(EditAnywhere, BlueprintReadWrite, Category = RuntimeSettings)
+	FRigLogicConfiguration RigLogicConfiguration;
+
+#if WITH_EDITOR
+	UE_API void PostEditChangeProperty(FPropertyChangedEvent& PropertyChangedEvent) override;
+#endif
+
+	UE_API bool Init(const FString& Filename);
+	UE_API void Serialize(FArchive& Ar) override;
+
+	/** Used when importing behavior into archetype SkelMesh in the editor,
+	* and when updating SkeletalMesh runtime with GeneSplicer
+	**/
+	UE_API void SetDNAReader(TSharedPtr<IDNAReader> SourceDNAReader);
+
+	/** Initialize this object for use at runtime from another instance that has already been
+	  * initialized. 
+	  * 
+	  * Overwrites all member variables. Only data needed for runtime evaluation will be copied.
+	  * 
+	  * Performs a shallow copy, so the runtime data is shared between the two instances and the
+	  * memory cost of the copied UDNA is very low.
+	  *
+	  * Note that the reference to the shared runtime data will be dropped if the source DNA is
+	  * modified, so the two instances are effectively independent and can safely be modified or
+	  * deleted without affecting the other.
+	 **/
+	UE_API void InitializeForRuntimeFrom(UDNA* Other);
+	UE_API TSharedPtr<FSharedRigRuntimeContext> GetRigRuntimeContext();
+	UE_API TSharedPtr<FDNAIndexMapping> GetDNAIndexMapping(const USkeleton* Skeleton, const USkeletalMesh* SkeletalMesh);
+
+private:
+	friend struct FAnimNode_RigLogic;
+	friend struct FRigUnit_RigLogic;
+
+	UE_API void InvalidateRigRuntimeContext();
+	UE_API void InitializeRigRuntimeContext();
+
+private:
+	struct FSkeletalMeshSkeletonToDNAIndexMappingKey
+	{
+		TWeakObjectPtr<const USkeletalMesh> SkeletalMesh;
+		TWeakObjectPtr<const USkeleton> Skeleton;
+
+		bool operator==(const FSkeletalMeshSkeletonToDNAIndexMappingKey& Other) const
+		{
+			return (SkeletalMesh == Other.SkeletalMesh) && (Skeleton == Other.Skeleton);
+		}
+		bool operator!=(const FSkeletalMeshSkeletonToDNAIndexMappingKey& Other) const
+		{
+			return !(*this == Other);
+		}
+		friend uint32 GetTypeHash(const FSkeletalMeshSkeletonToDNAIndexMappingKey& Instance)
+		{
+			return HashCombine(GetTypeHash(Instance.SkeletalMesh), GetTypeHash(Instance.Skeleton));
+		}
+	};
+
+private:
+	// Synchronize DNA updates
+	FTransactionallySafeRWLock DNAUpdateLock;
+
+	// Synchronize Rig Runtime Context updates
+	FTransactionallySafeRWLock RigRuntimeContextUpdateLock;
+
+	// Synchronize DNA Index Mapping updates
+	FTransactionallySafeRWLock DNAIndexMappingUpdateLock;
+
+	/**
+	*	Unified DNA Reader used for run-time execution of RigLogic and updating of SkeletalMesh geometry
+	**/
+	TSharedPtr<IDNAReader> DNAReader;
+
+	/** Runtime data necessary for rig computations that is shared between
+	  * multiple rig instances based on the same DNA.
+	**/
+	TSharedPtr<FSharedRigRuntimeContext> RigRuntimeContext;
+
+	/** Container for Skeleton <-> DNAAsset index mappings
+	  * The mapping object owners will be the SkeletalMeshes, and periodic cleanups will
+	  * ensure that dead objects are deleted from the map.
+	 **/
+	TMap<FSkeletalMeshSkeletonToDNAIndexMappingKey, TSharedPtr<FDNAIndexMapping>> DNAIndexMappingContainer;
+
+};
+
+#undef UE_API
diff --git a/Engine/Plugins/Animation/RigLogic/Source/RigLogicModule/Public/DNAAsset.h b/Engine/Plugins/Animation/RigLogic/Source/RigLogicModule/Public/DNAAsset.h
--- a/Engine/Plugins/Animation/RigLogic/Source/RigLogicModule/Public/DNAAsset.h
+++ b/Engine/Plugins/Animation/RigLogic/Source/RigLogicModule/Public/DNAAsset.h
@@ -36,7 +36,7 @@ enum class EDNADataLayer: uint8;
   * Currently, the design-time part still loads the geometry, as it is needed for the skeletal mesh update; once SkeletalMeshDNAReader is
   * fully implemented, it will be able to read the geometry directly from the SkeletalMesh and won't load it into this asset 
   **/
-UCLASS(MinimalAPI, BlueprintType, meta = (DisplayName = "MetaHuman DNA Data"))
+UCLASS(MinimalAPI, BlueprintType, meta = (DisplayName = "Legacy MetaHuman DNA Data"))
 class UDNAAsset : public UAssetUserData
 {
 	GENERATED_BODY()
diff --git a/Engine/Plugins/Animation/RigLogic/Source/RigLogicModule/Public/DNAAssetUserData.h b/Engine/Plugins/Animation/RigLogic/Source/RigLogicModule/Public/DNAAssetUserData.h
new file mode 100644
--- /dev/null
+++ b/Engine/Plugins/Animation/RigLogic/Source/RigLogicModule/Public/DNAAssetUserData.h
@@ -0,0 +1,26 @@
+// Copyright Epic Games, Inc. All Rights Reserved.
+
+#pragma once
+
+#include "CoreMinimal.h"
+#include "UObject/Object.h"
+#include "UObject/ObjectMacros.h"
+#include "Engine/AssetUserData.h"
+#include "Engine/SkeletalMesh.h"
+
+#include "DNAAssetUserData.generated.h"
+
+
+class UDNA;
+class USkeletalMesh;
+
+
+UCLASS()
+class UDNAAssetUserData: public UAssetUserData 
+{
+	GENERATED_BODY()
+public:
+
+	UPROPERTY(EditAnywhere, Category = "Asset", meta = (DisplayName = "DNA"))
+	TObjectPtr<UDNA> DNAAsset;
+};
\ No newline at end of file
diff --git a/Engine/Plugins/Animation/RigLogic/Source/RigLogicModule/RigLogicModule.Build.cs b/Engine/Plugins/Animation/RigLogic/Source/RigLogicModule/RigLogicModule.Build.cs
--- a/Engine/Plugins/Animation/RigLogic/Source/RigLogicModule/RigLogicModule.Build.cs
+++ b/Engine/Plugins/Animation/RigLogic/Source/RigLogicModule/RigLogicModule.Build.cs
@@ -28,6 +28,7 @@ namespace UnrealBuildTool.Rules
 				PrivateDependencyModuleNames.Add("SkeletalMeshUtilitiesCommon");
 				PrivateDependencyModuleNames.Add("RHI");
 				PrivateDependencyModuleNames.Add("RenderCore");
+				PrivateDependencyModuleNames.Add("AssetRegistry");
 			}
 
 			PrivateDependencyModuleNames.AddRange(
