Fix OOB access in RigLogic due to choice of bad SIMD vector in quaternion evaluator factory

#Synced-CL 49723036

diff --git a/Engine/Plugins/Animation/RigLogic/Source/RigLogicLib/Private/riglogic/joints/cpu/quaternions/QuaternionJointsBuilder.cpp b/Engine/Plugins/Animation/RigLogic/Source/RigLogicLib/Private/riglogic/joints/cpu/quaternions/QuaternionJointsBuilder.cpp
--- a/Engine/Plugins/Animation/RigLogic/Source/RigLogicLib/Private/riglogic/joints/cpu/quaternions/QuaternionJointsBuilder.cpp
+++ b/Engine/Plugins/Animation/RigLogic/Source/RigLogicLib/Private/riglogic/joints/cpu/quaternions/QuaternionJointsBuilder.cpp
@@ -49,7 +49,7 @@ struct JointGroupQuaternionStrategyFactory {
                            MemoryResource* memRes) {
 
         if (rotationType == RotationType::Quaternions) {
-            using CalculationStrategy = VectorizedJointGroupQuaternionCalculationStrategy<T, TF256, TF256, PassthroughAdapter>;
+            using CalculationStrategy = VectorizedJointGroupQuaternionCalculationStrategy<T, TF256, TF128, PassthroughAdapter>;
             return UniqueInstance<CalculationStrategy, JointGroupQuaternionCalculationStrategy>::with(memRes).create();
         }
 
@@ -57,11 +57,11 @@ struct JointGroupQuaternionStrategyFactory {
         if (rotationOrder == RotationOrder::XYZ) {
             if (rotationUnit == dna::RotationUnit::degrees) {
                 using Q2E = QuaternionsToEulerAngles<tdm::fdeg, tdm::rot_seq::xyz>;
-                using CalculationStrategy = VectorizedJointGroupQuaternionCalculationStrategy<T, TF256, TF256, Q2E>;
+                using CalculationStrategy = VectorizedJointGroupQuaternionCalculationStrategy<T, TF256, TF128, Q2E>;
                 return UniqueInstance<CalculationStrategy, JointGroupQuaternionCalculationStrategy>::with(memRes).create();
             } else {
                 using Q2E = QuaternionsToEulerAngles<tdm::frad, tdm::rot_seq::xyz>;
-                using CalculationStrategy = VectorizedJointGroupQuaternionCalculationStrategy<T, TF256, TF256, Q2E>;
+                using CalculationStrategy = VectorizedJointGroupQuaternionCalculationStrategy<T, TF256, TF128, Q2E>;
                 return UniqueInstance<CalculationStrategy, JointGroupQuaternionCalculationStrategy>::with(memRes).create();
             }
         }
@@ -71,11 +71,11 @@ struct JointGroupQuaternionStrategyFactory {
         if (rotationOrder == RotationOrder::XZY) {
             if (rotationUnit == dna::RotationUnit::degrees) {
                 using Q2E = QuaternionsToEulerAngles<tdm::fdeg, tdm::rot_seq::xzy>;
-                using CalculationStrategy = VectorizedJointGroupQuaternionCalculationStrategy<T, TF256, TF256, Q2E>;
+                using CalculationStrategy = VectorizedJointGroupQuaternionCalculationStrategy<T, TF256, TF128, Q2E>;
                 return UniqueInstance<CalculationStrategy, JointGroupQuaternionCalculationStrategy>::with(memRes).create();
             } else {
                 using Q2E = QuaternionsToEulerAngles<tdm::frad, tdm::rot_seq::xzy>;
-                using CalculationStrategy = VectorizedJointGroupQuaternionCalculationStrategy<T, TF256, TF256, Q2E>;
+                using CalculationStrategy = VectorizedJointGroupQuaternionCalculationStrategy<T, TF256, TF128, Q2E>;
                 return UniqueInstance<CalculationStrategy, JointGroupQuaternionCalculationStrategy>::with(memRes).create();
             }
         }
@@ -85,11 +85,11 @@ struct JointGroupQuaternionStrategyFactory {
         if (rotationOrder == RotationOrder::YXZ) {
             if (rotationUnit == dna::RotationUnit::degrees) {
                 using Q2E = QuaternionsToEulerAngles<tdm::fdeg, tdm::rot_seq::yxz>;
-                using CalculationStrategy = VectorizedJointGroupQuaternionCalculationStrategy<T, TF256, TF256, Q2E>;
+                using CalculationStrategy = VectorizedJointGroupQuaternionCalculationStrategy<T, TF256, TF128, Q2E>;
                 return UniqueInstance<CalculationStrategy, JointGroupQuaternionCalculationStrategy>::with(memRes).create();
             } else {
                 using Q2E = QuaternionsToEulerAngles<tdm::frad, tdm::rot_seq::yxz>;
-                using CalculationStrategy = VectorizedJointGroupQuaternionCalculationStrategy<T, TF256, TF256, Q2E>;
+                using CalculationStrategy = VectorizedJointGroupQuaternionCalculationStrategy<T, TF256, TF128, Q2E>;
                 return UniqueInstance<CalculationStrategy, JointGroupQuaternionCalculationStrategy>::with(memRes).create();
             }
         }
@@ -99,11 +99,11 @@ struct JointGroupQuaternionStrategyFactory {
         if (rotationOrder == RotationOrder::YZX) {
             if (rotationUnit == dna::RotationUnit::degrees) {
                 using Q2E = QuaternionsToEulerAngles<tdm::fdeg, tdm::rot_seq::yzx>;
-                using CalculationStrategy = VectorizedJointGroupQuaternionCalculationStrategy<T, TF256, TF256, Q2E>;
+                using CalculationStrategy = VectorizedJointGroupQuaternionCalculationStrategy<T, TF256, TF128, Q2E>;
                 return UniqueInstance<CalculationStrategy, JointGroupQuaternionCalculationStrategy>::with(memRes).create();
             } else {
                 using Q2E = QuaternionsToEulerAngles<tdm::frad, tdm::rot_seq::yzx>;
-                using CalculationStrategy = VectorizedJointGroupQuaternionCalculationStrategy<T, TF256, TF256, Q2E>;
+                using CalculationStrategy = VectorizedJointGroupQuaternionCalculationStrategy<T, TF256, TF128, Q2E>;
                 return UniqueInstance<CalculationStrategy, JointGroupQuaternionCalculationStrategy>::with(memRes).create();
             }
         }
@@ -113,11 +113,11 @@ struct JointGroupQuaternionStrategyFactory {
         if (rotationOrder == RotationOrder::ZXY) {
             if (rotationUnit == dna::RotationUnit::degrees) {
                 using Q2E = QuaternionsToEulerAngles<tdm::fdeg, tdm::rot_seq::zxy>;
-                using CalculationStrategy = VectorizedJointGroupQuaternionCalculationStrategy<T, TF256, TF256, Q2E>;
+                using CalculationStrategy = VectorizedJointGroupQuaternionCalculationStrategy<T, TF256, TF128, Q2E>;
                 return UniqueInstance<CalculationStrategy, JointGroupQuaternionCalculationStrategy>::with(memRes).create();
             } else {
                 using Q2E = QuaternionsToEulerAngles<tdm::frad, tdm::rot_seq::zxy>;
-                using CalculationStrategy = VectorizedJointGroupQuaternionCalculationStrategy<T, TF256, TF256, Q2E>;
+                using CalculationStrategy = VectorizedJointGroupQuaternionCalculationStrategy<T, TF256, TF128, Q2E>;
                 return UniqueInstance<CalculationStrategy, JointGroupQuaternionCalculationStrategy>::with(memRes).create();
             }
         }
@@ -127,11 +127,11 @@ struct JointGroupQuaternionStrategyFactory {
         if (rotationOrder == RotationOrder::ZYX) {
             if (rotationUnit == dna::RotationUnit::degrees) {
                 using Q2E = QuaternionsToEulerAngles<tdm::fdeg, tdm::rot_seq::zyx>;
-                using CalculationStrategy = VectorizedJointGroupQuaternionCalculationStrategy<T, TF256, TF256, Q2E>;
+                using CalculationStrategy = VectorizedJointGroupQuaternionCalculationStrategy<T, TF256, TF128, Q2E>;
                 return UniqueInstance<CalculationStrategy, JointGroupQuaternionCalculationStrategy>::with(memRes).create();
             } else {
                 using Q2E = QuaternionsToEulerAngles<tdm::frad, tdm::rot_seq::zyx>;
-                using CalculationStrategy = VectorizedJointGroupQuaternionCalculationStrategy<T, TF256, TF256, Q2E>;
+                using CalculationStrategy = VectorizedJointGroupQuaternionCalculationStrategy<T, TF256, TF128, Q2E>;
                 return UniqueInstance<CalculationStrategy, JointGroupQuaternionCalculationStrategy>::with(memRes).create();
             }
         }
