Implement optimized cooking process for DNA assets, such that instead of serializing DNA data, the RigLogic instance itself is dumped (and subsequently restored at runtime), bypassing the more expensive route of initialization of RigLogic from a DNA
#Synced-CL 49759155

diff --git a/Engine/Plugins/Animation/RigLogic/Source/RigLogicModule/Private/DNA.cpp b/Engine/Plugins/Animation/RigLogic/Source/RigLogicModule/Private/DNA.cpp
--- a/Engine/Plugins/Animation/RigLogic/Source/RigLogicModule/Private/DNA.cpp
+++ b/Engine/Plugins/Animation/RigLogic/Source/RigLogicModule/Private/DNA.cpp
@@ -59,6 +59,29 @@ namespace DNA_Helpers
 		DNAWriter->write();
 	}
 
+	static TSharedPtr<IDNAReader> CopyMinimumRuntimeDNA(const IDNAReader* Source)
+	{
+		auto DNAReader = Source->Unwrap();
+		auto Stream = rl4::makeScoped<rl4::MemoryStream>(FMemoryResource::Instance());
+		auto DNAWriter = rl4::makeScoped<dna::BinaryStreamWriter>(Stream.get(), FMemoryResource::Instance());
+	 
+		DNAWriter->setFrom(DNAReader, dna::DataLayer::Definition);
+	 
+		for (uint16 MeshIndex = DNAReader->getMeshCount(); MeshIndex != static_cast<uint16>(-1); --MeshIndex)
+		{
+			for (uint16 RegionIndex = DNAReader->getMeshRegionCount(MeshIndex); RegionIndex != static_cast<uint16>(-1); --RegionIndex)
+			{
+				DNAWriter->setMeshRegionName(MeshIndex, RegionIndex, DNAReader->getMeshRegionName(MeshIndex, RegionIndex));
+				const auto NetIndices = DNAReader->getNeuralNetworkIndicesForMeshRegion(MeshIndex, RegionIndex);
+				DNAWriter->setNeuralNetworkIndicesForMeshRegion(MeshIndex, RegionIndex, NetIndices.data(), static_cast<uint16>(NetIndices.size()));
+			}
+		}
+	 
+		DNAWriter->write();
+		Stream->seek(0);
+		return ReadDNAFromStream(Stream.get(), EDNADataLayer::All, 0);
+	}
+
 	static TSharedPtr<IDNAReader> CopyDNALayer(const IDNAReader* Source, EDNADataLayer DNADataLayer, uint32 PredictedSize)
 	{
 		// To avoid lots of reallocations in `FRigLogicMemoryStream`, reserve an approximate size
@@ -153,6 +176,8 @@ void UDNA::InitializeForRuntimeFrom(UDNA* Other)
 	// if bKeepDNAAfterInitialization is true, DNAReader will contain DNA and vice versa.
 	bKeepDNAAfterInitialization = Other->bKeepDNAAfterInitialization;
 
+	bUseOptimizedCooking = Other->bUseOptimizedCooking;
+
 	// This map is populated on demand, so doesn't need to be copied.
 	DNAIndexMappingContainer.Empty(1);
 
@@ -176,23 +201,15 @@ void UDNA::InitializeRigRuntimeContext()
 	LLM_SCOPE_BYNAME(TEXT("Animation/RigLogic"));
 
 	// Assumes DNAUpdateLock is locked by caller
-
-	TSharedPtr<FSharedRigRuntimeContext> NewContext = MakeShared<FSharedRigRuntimeContext>();
 	if (DNAReader.IsValid() && (DNAReader->GetJointCount() != 0))
 	{
-		NewContext->DNAReader = DNAReader;
 		// Convert behavior data from DNA to UE space on the fly as RigLogic accesses it
-		RigLogicDNAReader BehaviorReaderInUESpace{ DNAReader->Unwrap() };
-		FDNAReader<RigLogicDNAReader> BehaviorReaderInUESpaceWrapper{ &BehaviorReaderInUESpace };
-		NewContext->RigLogic = MakeShared<FRigLogic>(&BehaviorReaderInUESpaceWrapper, RigLogicConfiguration);
-		NewContext->CacheVariableJointIndices();
-		if ((DNAReader->GetRBFSolverCount() != 0) || (DNAReader->GetSwingCount() != 0) || (DNAReader->GetTwistCount() != 0))
-		{
-			NewContext->CacheInverseNeutralJointRotations();
-		}
-
+		RigLogicDNAReader BehaviorReaderInUESpace{DNAReader->Unwrap()};
+		FDNAReader<RigLogicDNAReader> BehaviorReaderInUESpaceWrapper{&BehaviorReaderInUESpace};
+		TSharedPtr<FRigLogic> RigLogic = MakeShared<FRigLogic>(&BehaviorReaderInUESpaceWrapper, RigLogicConfiguration);
+		TSharedPtr<FSharedRigRuntimeContext> NewContext = MakeShared<FSharedRigRuntimeContext>(DNAReader, RigLogic);
 		{
-			UE::TWriteScopeLock ContextScopeLock{ RigRuntimeContextUpdateLock };
+			UE::TWriteScopeLock ContextScopeLock{RigRuntimeContextUpdateLock};
 			RigRuntimeContext = NewContext;
 		}
 #if !WITH_EDITOR
@@ -339,8 +356,26 @@ void UDNA::Serialize(FArchive& Ar)
 
 				// Now with combined DNA streams we can read the provided memory stream into unified DNAReader
 				DNAReader = DNA_Helpers::ReadDNAFromStream(MemoryStream.get(), EDNADataLayer::All, MaxLOD);
+				InvalidateRigRuntimeContext();
+				InitializeRigRuntimeContext();
+			}
+			else if (bUseOptimizedCooking && Ar.IsLoadingFromCookedPackage() && (Ar.CustomVer(FDNAAssetCustomVersion::GUID) >= FDNAAssetCustomVersion::IntroduceOptimizedSerializationDuringCooking))
+			{
+				// Read already initialized RigLogic state from archive with minimum necessary DNA data
+				bool IsValid = false;
+				Ar << IsValid;
+				if (IsValid)
+				{
+					FArchiveMemoryStream DNAStream{&Ar};
+					DNAReader = DNA_Helpers::ReadDNAFromStream(&DNAStream, EDNADataLayer::All, MaxLOD);
+					TSharedPtr<FRigLogic> RigLogic = MakeShared<FRigLogic>(&Ar);
+					TSharedPtr<FSharedRigRuntimeContext> NewContext = MakeShared<FSharedRigRuntimeContext>(DNAReader, RigLogic);
+					{
+						UE::TWriteScopeLock ContextScopeLock{RigRuntimeContextUpdateLock};
+						RigRuntimeContext = NewContext;
+					}
+				}
 			}
-
 			else
 			{
 				FArchiveMemoryStream DNAStream{ &Ar };
@@ -356,28 +391,43 @@ void UDNA::Serialize(FArchive& Ar)
 				}
 
 				DNAReader = DNA_Helpers::ReadDNAFromStream(&DNAStream, DNALayers, MaxLOD);
+				InvalidateRigRuntimeContext();
+				InitializeRigRuntimeContext();
 			}
-			InvalidateRigRuntimeContext();
-			InitializeRigRuntimeContext();
 		}
-
-		if (Ar.IsSaving())
+		else if (Ar.IsSaving())
 		{
-			TSharedPtr<IDNAReader> EmptyDNA = DNA_Helpers::CreateEmptyDNA(DNA_Helpers::AVG_EMPTY_SIZE);
-			IDNAReader* ReaderPtr = (DNAReader.IsValid() ? static_cast<IDNAReader*>(DNAReader.Get()) : EmptyDNA.Get());
-			FArchiveMemoryStream DNAStream{ &Ar };
-
-			EDNADataLayer DNALayers = (
-				EDNADataLayer::Behavior |
-				EDNADataLayer::MachineLearnedBehavior |
-				EDNADataLayer::RBFBehavior
-				);
-#if WITH_EDITORONLY_DATA
-			// Geometry data is discarded unless in Editor
-			DNALayers |= EDNADataLayer::Geometry;
-#endif // #if WITH_EDITORONLY_DATA
-
-			DNA_Helpers::WriteDNAToStream(ReaderPtr, DNALayers, &DNAStream);
+			if (bUseOptimizedCooking && Ar.IsCooking() && (Ar.CustomVer(FDNAAssetCustomVersion::GUID) >= FDNAAssetCustomVersion::IntroduceOptimizedSerializationDuringCooking))
+			{
+				// Write already initialized RigLogic state into archive with minimum necessary DNA data
+				bool IsValid = DNAReader.IsValid();
+				Ar << IsValid;
+				if (IsValid)
+				{
+					FArchiveMemoryStream DNAStream{&Ar};
+					TSharedPtr<IDNAReader> MinRuntimeDNA = DNA_Helpers::CopyMinimumRuntimeDNA(DNAReader.Get());
+					DNA_Helpers::WriteDNAToStream(MinRuntimeDNA.Get(), EDNADataLayer::All, &DNAStream);
+					RigRuntimeContext->RigLogic->Dump(&Ar);
+				}
+			}
+			else
+			{
+				TSharedPtr<IDNAReader> EmptyDNA = DNA_Helpers::CreateEmptyDNA(DNA_Helpers::AVG_EMPTY_SIZE);
+				IDNAReader* ReaderPtr = (DNAReader.IsValid() ? static_cast<IDNAReader*>(DNAReader.Get()) : EmptyDNA.Get());
+				FArchiveMemoryStream DNAStream{ &Ar };
+	
+				EDNADataLayer DNALayers = (
+					EDNADataLayer::Behavior |
+					EDNADataLayer::MachineLearnedBehavior |
+					EDNADataLayer::RBFBehavior
+					);
+	#if WITH_EDITORONLY_DATA
+				// Geometry data is discarded unless in Editor
+				DNALayers |= EDNADataLayer::Geometry;
+	#endif // #if WITH_EDITORONLY_DATA
+	
+				DNA_Helpers::WriteDNAToStream(ReaderPtr, DNALayers, &DNAStream);
+			}
 		}
 	}
 }
diff --git a/Engine/Plugins/Animation/RigLogic/Source/RigLogicModule/Private/DNAAsset.cpp b/Engine/Plugins/Animation/RigLogic/Source/RigLogicModule/Private/DNAAsset.cpp
--- a/Engine/Plugins/Animation/RigLogic/Source/RigLogicModule/Private/DNAAsset.cpp
+++ b/Engine/Plugins/Animation/RigLogic/Source/RigLogicModule/Private/DNAAsset.cpp
@@ -205,21 +205,13 @@ void UDNAAsset::InitializeRigRuntimeContext()
 	LLM_SCOPE_BYNAME(TEXT("Animation/RigLogic"));
 
 	// Assumes DNAUpdateLock is locked by caller
-
-	TSharedPtr<FSharedRigRuntimeContext> NewContext = MakeShared<FSharedRigRuntimeContext>();
 	if (DNAReader.IsValid() && (DNAReader->GetJointCount() != 0))
 	{
-		NewContext->DNAReader = DNAReader;
 		// Convert behavior data from DNA to UE space on the fly as RigLogic accesses it
-		RigLogicDNAReader BehaviorReaderInUESpace{ DNAReader->Unwrap()};
+		RigLogicDNAReader BehaviorReaderInUESpace{DNAReader->Unwrap()};
 		FDNAReader<RigLogicDNAReader> BehaviorReaderInUESpaceWrapper{&BehaviorReaderInUESpace};
-		NewContext->RigLogic = MakeShared<FRigLogic>(&BehaviorReaderInUESpaceWrapper, RigLogicConfiguration);
-		NewContext->CacheVariableJointIndices();
-		if ((DNAReader->GetRBFSolverCount() != 0) || (DNAReader->GetSwingCount() != 0) || (DNAReader->GetTwistCount() != 0))
-		{
-			NewContext->CacheInverseNeutralJointRotations();
-		}
-
+		TSharedPtr<FRigLogic> RigLogic = MakeShared<FRigLogic>(&BehaviorReaderInUESpaceWrapper, RigLogicConfiguration);
+		TSharedPtr<FSharedRigRuntimeContext> NewContext = MakeShared<FSharedRigRuntimeContext>(DNAReader, RigLogic);
 		{
 			UE::TWriteScopeLock ContextScopeLock{RigRuntimeContextUpdateLock};
 			RigRuntimeContext = NewContext;
@@ -369,7 +361,6 @@ void UDNAAsset::Serialize(FArchive& Ar)
 				// Now with combined DNA streams we can read the provided memory stream into unified DNAReader
 				DNAReader = DNAAsset_Helpers::ReadDNAFromStream(MemoryStream.get(), EDNADataLayer::All, 0);
 			}
-
 			else
 			{
 				FArchiveMemoryStream DNAStream{ &Ar };
@@ -389,8 +380,7 @@ void UDNAAsset::Serialize(FArchive& Ar)
 			InvalidateRigRuntimeContext();
 			InitializeRigRuntimeContext();
 		}
-
-		if (Ar.IsSaving())
+		else if (Ar.IsSaving())
 		{
 			TSharedPtr<IDNAReader> EmptyDNA = DNAAsset_Helpers::CreateEmptyDNA(DNAAsset_Helpers::AVG_EMPTY_SIZE);
 			IDNAReader* ReaderPtr = (DNAReader.IsValid() ? static_cast<IDNAReader*>(DNAReader.Get()) : EmptyDNA.Get());
diff --git a/Engine/Plugins/Animation/RigLogic/Source/RigLogicModule/Private/DNAAssetCustomVersion.h b/Engine/Plugins/Animation/RigLogic/Source/RigLogicModule/Private/DNAAssetCustomVersion.h
--- a/Engine/Plugins/Animation/RigLogic/Source/RigLogicModule/Private/DNAAssetCustomVersion.h
+++ b/Engine/Plugins/Animation/RigLogic/Source/RigLogicModule/Private/DNAAssetCustomVersion.h
@@ -13,6 +13,7 @@ struct FDNAAssetCustomVersion
 		// Before any version changes were made in the plugin
 		BeforeCustomVersionWasAdded = 0,
 		UnifiedReaderVersion,
+		IntroduceOptimizedSerializationDuringCooking,
 		// -----<new versions can be added above this line>-------------------------------------------------
 		VersionPlusOne,
 		LatestVersion = VersionPlusOne - 1
diff --git a/Engine/Plugins/Animation/RigLogic/Source/RigLogicModule/Private/DNAIndexMapping.cpp b/Engine/Plugins/Animation/RigLogic/Source/RigLogicModule/Private/DNAIndexMapping.cpp
--- a/Engine/Plugins/Animation/RigLogic/Source/RigLogicModule/Private/DNAIndexMapping.cpp
+++ b/Engine/Plugins/Animation/RigLogic/Source/RigLogicModule/Private/DNAIndexMapping.cpp
@@ -59,14 +59,23 @@ void FDNAIndexMapping::MapMLMaskCurves(const IDNAReader* DNAReader, const USkele
 {
 	LLM_SCOPE_BYNAME(TEXT("Animation/RigLogic"));
 
+	const uint16 MeshCount = DNAReader->GetMeshCount();
+	const uint16 MappingCount = [MeshCount, DNAReader]()
+	{
+		uint32 Count = 0;
+		for (uint16 MeshIndex = 0; MeshIndex < MeshCount; ++MeshIndex)
+		{
+			Count += DNAReader->GetMeshRegionCount(MeshIndex);
+		}
+		return static_cast<uint16>(Count);
+	}();
 	MLMaskCurves.Empty();
+	MLMaskCurves.Reserve(MappingCount);
 
 	uint32 MLMaskIndex = 0;
-	const uint16 MeshCount = DNAReader->GetMeshCount();
 	for (uint16 MeshIndex = 0; MeshIndex < MeshCount; ++MeshIndex)
 	{
 		const uint16 MeshRegionCount = DNAReader->GetMeshRegionCount(MeshIndex);
-		MLMaskCurves.Reserve(MLMaskCurves.Num() + MeshRegionCount);
 		for (uint16 RegionIndex = 0; RegionIndex < MeshRegionCount; ++RegionIndex, ++MLMaskIndex)
 		{
 			const FString& MeshRegionName = DNAReader->GetMeshRegionName(MeshIndex, RegionIndex);
@@ -147,6 +156,12 @@ void FDNAIndexMapping::MapDriverJoints(const IDNAReader* DNAReader, const USkele
 			Mapping.RotationW = ControlIndex;
 		}
 	}
+
+	// Don't keep reservation if no mappings were created
+	if (DriverJointsToControlAttributesMap.Num() == 0)
+	{
+		DriverJointsToControlAttributesMap.Empty();
+	}
 }
 
 void FDNAIndexMapping::MapJoints(const IDNAReader* DNAReader, const USkeletalMesh* SkeletalMesh)
diff --git a/Engine/Plugins/Animation/RigLogic/Source/RigLogicModule/Private/RigLogic.cpp b/Engine/Plugins/Animation/RigLogic/Source/RigLogicModule/Private/RigLogic.cpp
--- a/Engine/Plugins/Animation/RigLogic/Source/RigLogicModule/Private/RigLogic.cpp
+++ b/Engine/Plugins/Animation/RigLogic/Source/RigLogicModule/Private/RigLogic.cpp
@@ -1,7 +1,8 @@
 // Copyright Epic Games, Inc. All Rights Reserved.
 
 #include "RigLogic.h"
-
+ 
+#include "ArchiveMemoryStream.h"
 #include "DNAReader.h"
 #include "FMemoryResource.h"
 #include "RigInstance.h"
@@ -39,7 +40,10 @@ DEFINE_STAT(STAT_RigLogic_CalculateAnimatedMapsTime);
 
 void FRigLogic::FRigLogicDeleter::operator()(rl4::RigLogic* Pointer)
 {
-	rl4::RigLogic::destroy(Pointer);
+	if (Pointer != nullptr)
+	{
+		rl4::RigLogic::destroy(Pointer);
+	}
 }
 
 static rl4::Configuration AdaptToRigLogicConfig(const FRigLogicConfiguration& Config)
@@ -89,8 +93,24 @@ FRigLogic::FRigLogic(const IDNAReader* Reader, const FRigLogicConfiguration& Con
 {
 }
 
+FRigLogic::FRigLogic(FArchive* Archive) :
+	MemoryResource{FMemoryResource::SharedInstance()},
+	RigLogic{nullptr},
+	Configuration{}
+{
+	FArchiveMemoryStream Stream(Archive);
+	RigLogic.Reset(rl4::RigLogic::restore(&Stream, FMemoryResource::Instance()));
+	Configuration = AdaptFromRigLogicConfig(RigLogic->getConfiguration());
+}
+
 FRigLogic::~FRigLogic() = default;
 
+void FRigLogic::Dump(FArchive* Archive) const
+{
+	FArchiveMemoryStream Stream(Archive);
+	RigLogic->dump(&Stream);
+}
+
 const FRigLogicConfiguration& FRigLogic::GetConfiguration() const
 {
 	return Configuration;
@@ -153,6 +173,16 @@ uint16 FRigLogic::GetRBFSolverCount() const
 	return RigLogic->getRBFSolverCount();
 }
 
+uint16 FRigLogic::GetTwistCount() const
+{
+	return RigLogic->getTwistCount();
+}
+
+uint16 FRigLogic::GetSwingCount() const
+{
+	return RigLogic->getSwingCount();
+}
+
 uint16 FRigLogic::GetMeshCount() const
 {
 	return RigLogic->getMeshCount();
diff --git a/Engine/Plugins/Animation/RigLogic/Source/RigLogicModule/Private/SharedRigRuntimeContext.cpp b/Engine/Plugins/Animation/RigLogic/Source/RigLogicModule/Private/SharedRigRuntimeContext.cpp
--- a/Engine/Plugins/Animation/RigLogic/Source/RigLogicModule/Private/SharedRigRuntimeContext.cpp
+++ b/Engine/Plugins/Animation/RigLogic/Source/RigLogicModule/Private/SharedRigRuntimeContext.cpp
@@ -7,26 +7,37 @@
 #include "DNAReader.h"
 #include "RigLogic.h"
 
+FSharedRigRuntimeContext::FSharedRigRuntimeContext() = default;
+
+FSharedRigRuntimeContext::FSharedRigRuntimeContext(TSharedPtr<IDNAReader> InDNAReader, TSharedPtr<FRigLogic> InRigLogic) :
+	DNAReader{InDNAReader},
+	RigLogic{InRigLogic}
+{
+	CacheVariableJointIndices();
+	if ((RigLogic->GetRBFSolverCount() != 0) || (RigLogic->GetSwingCount() != 0) || (RigLogic->GetTwistCount() != 0))
+	{
+		CacheInverseNeutralJointRotations();
+	}
+}
+
 void FSharedRigRuntimeContext::CacheVariableJointIndices()
 {
 	LLM_SCOPE_BYNAME(TEXT("Animation/RigLogic"));
 
-	const uint16 LODCount = DNAReader->GetLODCount();
+	const FRigLogicConfiguration& Config = RigLogic->GetConfiguration();
+	const uint32 AttrCount = static_cast<uint8>(Config.TranslationType) + static_cast<uint8>(Config.RotationType) + static_cast<uint8>(Config.ScaleType);
+	const uint16 LODCount = RigLogic->GetLODCount();
 	VariableJointIndicesPerLOD.Reset();
 	VariableJointIndicesPerLOD.AddDefaulted(LODCount);
 	TSet<uint16> DistinctVariableJointIndices;
 	for (uint16 LODIndex = 0; LODIndex < LODCount; ++LODIndex)
 	{
-		TArrayView<const uint16> VariableAttributeIndices = DNAReader->GetJointVariableAttributeIndices(LODIndex);
+		TArrayView<const uint16> VariableAttributeIndices = RigLogic->GetJointVariableAttributeIndices(LODIndex);
 		DistinctVariableJointIndices.Reset();
 		DistinctVariableJointIndices.Reserve(VariableAttributeIndices.Num());
 		for (const uint16 AttrIndex : VariableAttributeIndices)
 		{
-			// In DNA, the number of joint attributes is always 9 (only RigLogic has the ability to switch this)
-			// and since the variable indices are queried from the DNA here, we deal with 9 as well, regardless of
-			// the state elsewhere where we switched to 10 since the introduction of quaternion outputs from RigLogic.
-			static constexpr uint16 NUM_ATTRS_PER_JOINT = 9;
-			const uint16 JointIndex = AttrIndex / NUM_ATTRS_PER_JOINT;
+			const uint16 JointIndex = AttrIndex / AttrCount;
 			DistinctVariableJointIndices.Add(JointIndex);
 		}
 		VariableJointIndicesPerLOD[LODIndex].Values = DistinctVariableJointIndices.Array();
diff --git a/Engine/Plugins/Animation/RigLogic/Source/RigLogicModule/Public/DNA.h b/Engine/Plugins/Animation/RigLogic/Source/RigLogicModule/Public/DNA.h
--- a/Engine/Plugins/Animation/RigLogic/Source/RigLogicModule/Public/DNA.h
+++ b/Engine/Plugins/Animation/RigLogic/Source/RigLogicModule/Public/DNA.h
@@ -62,6 +62,9 @@ public:
 	UPROPERTY(EditAnywhere, Category = RuntimeSettings)
 	bool bKeepDNAAfterInitialization;
 
+	UPROPERTY(EditAnywhere, Category = RuntimeSettings)
+	bool bUseOptimizedCooking = true;
+
 	UPROPERTY(EditAnywhere, Category = RuntimeSettings)
 	uint16 MaxLOD = 0u;
 
diff --git a/Engine/Plugins/Animation/RigLogic/Source/RigLogicModule/Public/RigLogic.h b/Engine/Plugins/Animation/RigLogic/Source/RigLogicModule/Public/RigLogic.h
--- a/Engine/Plugins/Animation/RigLogic/Source/RigLogicModule/Public/RigLogic.h
+++ b/Engine/Plugins/Animation/RigLogic/Source/RigLogicModule/Public/RigLogic.h
@@ -182,6 +182,7 @@ class FRigLogic
 {
 public:
 	UE_API explicit FRigLogic(const IDNAReader* Reader, const FRigLogicConfiguration& Config = FRigLogicConfiguration());
+	UE_API explicit FRigLogic(FArchive* Archive);
 	UE_API ~FRigLogic();
 
 	FRigLogic(const FRigLogic&) = delete;
@@ -190,6 +191,7 @@ public:
 	FRigLogic(FRigLogic&&) = default;
 	FRigLogic& operator=(FRigLogic&&) = default;
 
+	UE_API void Dump(FArchive* Archive) const;
 	UE_API const FRigLogicConfiguration& GetConfiguration() const;
 	UE_API uint16 GetLODCount() const;
 	UE_API TArrayView<const uint16> GetRBFSolverIndicesForLOD(uint16_t LOD) const;
@@ -201,6 +203,8 @@ public:
 	UE_API TArrayView<const uint16> GetJointVariableAttributeIndices(uint16 LOD) const;
 	UE_API uint16 GetJointGroupCount() const;
 	UE_API uint16 GetRBFSolverCount() const;
+	UE_API uint16 GetTwistCount() const;
+	UE_API uint16 GetSwingCount() const;
 	UE_API uint16 GetMeshCount() const;
 	UE_API uint16 GetMeshRegionCount(uint16 MeshIndex) const;
 	UE_API uint16 GetMLTypeCount() const;
diff --git a/Engine/Plugins/Animation/RigLogic/Source/RigLogicModule/Public/SharedRigRuntimeContext.h b/Engine/Plugins/Animation/RigLogic/Source/RigLogicModule/Public/SharedRigRuntimeContext.h
--- a/Engine/Plugins/Animation/RigLogic/Source/RigLogicModule/Public/SharedRigRuntimeContext.h
+++ b/Engine/Plugins/Animation/RigLogic/Source/RigLogicModule/Public/SharedRigRuntimeContext.h
@@ -27,6 +27,11 @@ struct FSharedRigRuntimeContext
 
 	TArray<FQuat> InverseNeutralJointRotations;
 
+	// Deprecated default constructor for ABI compatibility
+	UE_DEPRECATED(5.8, "Default constructor is deprecated. Use the explicit constructor instead.")
+	FSharedRigRuntimeContext();
+
+	FSharedRigRuntimeContext(TSharedPtr<IDNAReader> InBehaviorReader, TSharedPtr<FRigLogic> InRigLogic);
 	void CacheVariableJointIndices();
 	void CacheInverseNeutralJointRotations();
 
