Resolve error C4458: declaration of 'remappedIndices' hides class member

#rnx

#Synced-CL 50064991

diff --git a/Engine/Plugins/Animation/RigLogic/Source/RigLogicLib/Private/dna/stream/FilteredBinaryInputArchive.cpp b/Engine/Plugins/Animation/RigLogic/Source/RigLogicLib/Private/dna/stream/FilteredBinaryInputArchive.cpp
--- a/Engine/Plugins/Animation/RigLogic/Source/RigLogicLib/Private/dna/stream/FilteredBinaryInputArchive.cpp
+++ b/Engine/Plugins/Animation/RigLogic/Source/RigLogicLib/Private/dna/stream/FilteredBinaryInputArchive.cpp
@@ -367,16 +367,16 @@ void FilteredBinaryInputArchive::process(RawMachineLearnedBehavior& dest) {
 
         dest.lodNeuralNetworkMapping.discardLODs(lodConstraint);
         const auto passingNetIndices = dest.lodNeuralNetworkMapping.getCombinedDistinctIndices(memRes);
-        UnorderedMap<std::uint16_t, std::uint16_t> remappedIndices{memRes};
-        remap(neuralNetCount, passingNetIndices, remappedIndices);
-        dest.lodNeuralNetworkMapping.mapIndices([&remappedIndices](std::uint16_t value) { return remappedIndices.at(value); });
+        UnorderedMap<std::uint16_t, std::uint16_t> remappedIndicesLocal{memRes};
+        remap(neuralNetCount, passingNetIndices, remappedIndicesLocal);
+        dest.lodNeuralNetworkMapping.mapIndices([&remappedIndicesLocal](std::uint16_t value) { return remappedIndicesLocal.at(value); });
 
         // Delete and remap neural network indices in neuralNetworkToMeshRegion mapping
         for (auto& mesh : dest.neuralNetworkToMeshRegion.indices) {
             for (auto& region : mesh) {
                 extd::filter(region, extd::byValue(passingNetIndices));
                 for (auto& netIndex : region) {
-                    netIndex = remappedIndices.at(netIndex);
+                    netIndex = remappedIndicesLocal.at(netIndex);
                 }
             }
         }
@@ -427,9 +427,9 @@ void FilteredBinaryInputArchive::process(RawMachineLearnedBehaviorExt& dest) {
                 const auto mlOperationCount = static_cast<std::uint16_t>(mlOperations.size());
                 lodMapping.discardLODs(lodConstraint);
                 const auto distinctIndices = lodMapping.getCombinedDistinctIndices(memRes);
-                UnorderedMap<std::uint16_t, std::uint16_t> remappedIndices{memRes};
-                remap(mlOperationCount, distinctIndices, remappedIndices);
-                lodMapping.mapIndices([&remappedIndices](std::uint16_t value) { return remappedIndices.at(value); });
+                UnorderedMap<std::uint16_t, std::uint16_t> remappedIndicesLocal{memRes};
+                remap(mlOperationCount, distinctIndices, remappedIndicesLocal);
+                lodMapping.mapIndices([&remappedIndicesLocal](std::uint16_t value) { return remappedIndicesLocal.at(value); });
 
                 extd::filter(mlOperations, [&distinctIndices](const RawMLOperation& operation, std::size_t index) {
                     if (!extd::contains(distinctIndices, index)) {
@@ -447,8 +447,8 @@ void FilteredBinaryInputArchive::process(RawMachineLearnedBehaviorExt& dest) {
                     for (auto& operation : mlbType.operations[mlNextOperationSetIndex]) {
                         for (std::size_t i = {}; i < operation.mlOperationDependencyIndices.size(); ++i) {
                             if (operation.mlOperationSetDependencyIndices[i] == mlOperationSetIndex) {
-                                auto it = remappedIndices.find(operation.mlOperationDependencyIndices[i]);
-                                if (it != remappedIndices.end()) {
+                                auto it = remappedIndicesLocal.find(operation.mlOperationDependencyIndices[i]);
+                                if (it != remappedIndicesLocal.end()) {
                                     operation.mlOperationDependencyIndices[i] = it->second;
                                 }
                             }
@@ -531,9 +531,9 @@ void FilteredBinaryInputArchive::process(RawRBFBehavior& dest) {
     dest.lodSolverMapping.discardLODs(lodConstraint);
 
     const auto passingSolverIndices = dest.lodSolverMapping.getCombinedDistinctIndices(memRes);
-    UnorderedMap<std::uint16_t, std::uint16_t> remappedIndices{memRes};
-    remap(solverCount, passingSolverIndices, remappedIndices);
-    dest.lodSolverMapping.mapIndices([&remappedIndices](std::uint16_t value) { return remappedIndices.at(value); });
+    UnorderedMap<std::uint16_t, std::uint16_t> remappedIndicesLocal{memRes};
+    remap(solverCount, passingSolverIndices, remappedIndicesLocal);
+    dest.lodSolverMapping.mapIndices([&remappedIndicesLocal](std::uint16_t value) { return remappedIndicesLocal.at(value); });
 
     dest.solvers.reserve(passingSolverIndices.size());
     for (std::uint16_t i = {}; i < solverCount; ++i) {
