Fix/suppress warnings exposed by newer MSVC
Add missing breaks
#rnx
#Synced-CL 50140652

diff --git a/Engine/Plugins/Animation/RigLogic/Source/RigLogicLib/Public/terse/archives/binary/OutputArchive.h b/Engine/Plugins/Animation/RigLogic/Source/RigLogicLib/Public/terse/archives/binary/OutputArchive.h
--- a/Engine/Plugins/Animation/RigLogic/Source/RigLogicLib/Public/terse/archives/binary/OutputArchive.h
+++ b/Engine/Plugins/Animation/RigLogic/Source/RigLogicLib/Public/terse/archives/binary/OutputArchive.h
@@ -312,6 +312,7 @@ protected:
     processElements(const TContainer& source) {
         using ValueType = typename TContainer::value_type;
         const auto size = source.size();
+		#pragma warning(suppress : 6326)
         if (size != 0ul) {
             // NOLINTNEXTLINE(cppcoreguidelines-pro-type-reinterpret-cast)
             stream->write(reinterpret_cast<const char*>(&source[0]), size * sizeof(ValueType));
