[RigLogic] Expose DNA Config options as EditAnywhere

#rnx
#Synced-CL 52096074

diff --git a/Engine/Plugins/Animation/RigLogic/Source/RigLogicModule/Public/DNACommon.h b/Engine/Plugins/Animation/RigLogic/Source/RigLogicModule/Public/DNACommon.h
--- a/Engine/Plugins/Animation/RigLogic/Source/RigLogicModule/Public/DNACommon.h
+++ b/Engine/Plugins/Animation/RigLogic/Source/RigLogicModule/Public/DNACommon.h
@@ -204,11 +204,11 @@ struct FRotationSign
 		return !(*this == Other);
 	}
 
-	UPROPERTY(BlueprintReadOnly, Category = "RigLogic")
+	UPROPERTY(EditAnywhere, BlueprintReadOnly, Category = "RigLogic")
 	ERotationDirection XAxis;
-	UPROPERTY(BlueprintReadOnly, Category = "RigLogic")
+	UPROPERTY(EditAnywhere, BlueprintReadOnly, Category = "RigLogic")
 	ERotationDirection YAxis;
-	UPROPERTY(BlueprintReadOnly, Category = "RigLogic")
+	UPROPERTY(EditAnywhere, BlueprintReadOnly, Category = "RigLogic")
 	ERotationDirection ZAxis;
 };
 
@@ -235,11 +235,11 @@ struct FCoordinateSystem
 		return !(*this == Other);
 	}
 
-	UPROPERTY(BlueprintReadOnly, Category = "RigLogic")
+	UPROPERTY(EditAnywhere, BlueprintReadOnly, Category = "RigLogic")
 	EDirection XAxis;
-	UPROPERTY(BlueprintReadOnly, Category = "RigLogic")
+	UPROPERTY(EditAnywhere, BlueprintReadOnly, Category = "RigLogic")
 	EDirection YAxis;
-	UPROPERTY(BlueprintReadOnly, Category = "RigLogic")
+	UPROPERTY(EditAnywhere, BlueprintReadOnly, Category = "RigLogic")
 	EDirection ZAxis;
 };
 
@@ -272,31 +272,31 @@ struct FDNAConfig
 		return DNAConfig;
 	}
 
-	UPROPERTY(BlueprintReadOnly, Category = "RigLogic")
+	UPROPERTY(EditAnywhere, BlueprintReadOnly, Category = "RigLogic")
 	EDNADataLayer Layer = EDNADataLayer::All;
 
-	UPROPERTY(BlueprintReadOnly, Category = "RigLogic")
+	UPROPERTY(EditAnywhere, BlueprintReadOnly, Category = "RigLogic")
 	uint8 MaxLOD = 0;
 
-	UPROPERTY(BlueprintReadOnly, Category = "RigLogic")
+	UPROPERTY(EditAnywhere, BlueprintReadOnly, Category = "RigLogic")
 	uint8 MinLOD = static_cast<uint8>(-1);
 
-	UPROPERTY(BlueprintReadOnly, Category = "RigLogic")
+	UPROPERTY(EditAnywhere, BlueprintReadOnly, Category = "RigLogic")
 	TArray<uint8> ExactLODs;
 	// This is a hidden copy of the above blueprint-exposed exact LOD list so it can be safely converted to the lower level config
 	UPROPERTY(Transient)
 	mutable TArray<uint16> InternalExactLODs;
 
-	UPROPERTY(BlueprintReadOnly, Category = "RigLogic")
+	UPROPERTY(EditAnywhere, BlueprintReadOnly, Category = "RigLogic")
 	ECoordinateSystemTransformPolicy CoordinateSystemTransformPolicy = ECoordinateSystemTransformPolicy::Preserve;
 
-	UPROPERTY(BlueprintReadOnly, Category = "RigLogic")
+	UPROPERTY(EditAnywhere, BlueprintReadOnly, Category = "RigLogic")
 	FCoordinateSystem CoordinateSystem = {EDirection::Left, EDirection::Front, EDirection::Up};
 
-	UPROPERTY(BlueprintReadOnly, Category = "RigLogic")
+	UPROPERTY(EditAnywhere, BlueprintReadOnly, Category = "RigLogic")
 	FRotationSign RotationSign = {ERotationDirection::Negative, ERotationDirection::Negative, ERotationDirection::Positive};
 
-	UPROPERTY(BlueprintReadOnly, Category = "RigLogic")
+	UPROPERTY(EditAnywhere, BlueprintReadOnly, Category = "RigLogic")
 	ERotationSequence RotationSequence = ERotationSequence::XYZ;
 };
 
