Make DNA Layer checkbox follow the implicit rules already defined by the underlying enum type
#Synced-CL 52290889

diff --git a/Engine/Plugins/Animation/RigLogic/Source/RigLogicModule/Private/DNA.cpp b/Engine/Plugins/Animation/RigLogic/Source/RigLogicModule/Private/DNA.cpp
--- a/Engine/Plugins/Animation/RigLogic/Source/RigLogicModule/Private/DNA.cpp
+++ b/Engine/Plugins/Animation/RigLogic/Source/RigLogicModule/Private/DNA.cpp
@@ -103,6 +103,36 @@ namespace DNA_Helpers
 		MemoryStream->seek(0);
 		return ReadDNAFromStream(MemoryStream.get(), DNAConfig);
 	}
+
+	static int32 AutoSelectImplicitlyLoadedLayers(int32 Layers)
+	{
+		// GeometryWithoutBlendShapes takes precedence over Geometry
+		constexpr int32 GeometryBit = 1 << static_cast<uint8>(EDNADataLayerUIProxy::Geometry);
+		constexpr int32 GeometryNoBSBit = 1 << static_cast<uint8>(EDNADataLayerUIProxy::GeometryNoBlendShapes);
+		if (Layers & GeometryNoBSBit)
+		{
+			Layers &= ~GeometryBit;
+		}
+
+		static constexpr EDNADataLayer Mapping[] = {
+			EDNADataLayer::Descriptor,                 // UIProxy::Descriptor = 0
+			EDNADataLayer::Definition,                 // UIProxy::Definition = 1
+			EDNADataLayer::Behavior,                   // UIProxy::Behavior = 2
+			EDNADataLayer::Geometry,                   // UIProxy::Geometry = 3
+			EDNADataLayer::GeometryWithoutBlendShapes, // UIProxy::GeometryNoBlendShapes = 4
+			EDNADataLayer::MachineLearnedBehavior,     // UIProxy::MLBehavior = 5
+			EDNADataLayer::RBFBehavior,                // UIProxy::RBFBehavior = 6
+		};
+		int32 Result = 0;
+		for (int32 i = 0; i < UE_ARRAY_COUNT(Mapping); ++i)
+		{
+			if (Layers & (1 << i))
+			{
+				Result |= static_cast<int32>(Mapping[i]);
+			}
+		}
+		return Result;
+	}
 }
 
 
@@ -520,6 +550,8 @@ void UDNA::PostEditChangeProperty(FPropertyChangedEvent& PropertyChangedEvent)
 	}
 	else if (MemberPropertyName == GET_MEMBER_NAME_CHECKED(UDNA, DNAConfig))
 	{
+		DNAConfig.Layers = DNA_Helpers::AutoSelectImplicitlyLoadedLayers(DNAConfig.Layers);
+
 		UE::TWriteScopeLock DNAScopeLock(DNAUpdateLock);
 		pma::ScopedPtr<dna::MemoryStream> MemoryStream = pma::makeScoped<dna::MemoryStream>();
 		pma::ScopedPtr<dna::BinaryStreamWriter> DNAWriter = pma::makeScoped<dna::BinaryStreamWriter>(MemoryStream.get());
