Render custom widget for DNA coordinate systems so no invalid, intermediate settings can be applied to it.
#Synced-CL 52297420

diff --git a/Engine/Plugins/Animation/RigLogic/Source/RigLogicEditor/Private/CoordinateSystemCustomization.cpp b/Engine/Plugins/Animation/RigLogic/Source/RigLogicEditor/Private/CoordinateSystemCustomization.cpp
new file mode 100644
--- /dev/null
+++ b/Engine/Plugins/Animation/RigLogic/Source/RigLogicEditor/Private/CoordinateSystemCustomization.cpp
@@ -0,0 +1,272 @@
+// Copyright Epic Games, Inc. All Rights Reserved.
+
+#include "CoordinateSystemCustomization.h"
+#include "DetailLayoutBuilder.h"
+#include "DetailWidgetRow.h"
+#include "IDetailChildrenBuilder.h"
+#include "Widgets/Input/SComboBox.h"
+#include "Widgets/Input/SButton.h"
+#include "Widgets/Text/STextBlock.h"
+#include "ScopedTransaction.h"
+
+#define LOCTEXT_NAMESPACE "CoordinateSystemCustomization"
+
+namespace
+{
+
+	const FString DirectionNames[] = {
+		TEXT("Left"), TEXT("Right"), TEXT("Up"), TEXT("Down"), TEXT("Front"), TEXT("Back")
+	};
+	constexpr int32 NumDirections = UE_ARRAY_COUNT(DirectionNames);
+
+	/** Which spatial dimension does this direction belong to? (0=Left/Right, 1=Up/Down, 2=Front/Back) */
+	int32 GetDimensionIndex(EDirection Dir)
+	{
+		switch (Dir)
+		{
+		case EDirection::Left:  case EDirection::Right: return 0;
+		case EDirection::Up:    case EDirection::Down:  return 1;
+		case EDirection::Front: case EDirection::Back:  return 2;
+		default: return -1;
+		}
+	}
+
+	EDirection DirectionFromString(const FString& Str)
+	{
+		for (int32 i = 0; i < NumDirections; ++i)
+		{
+			if (Str == DirectionNames[i])
+			{
+				return static_cast<EDirection>(i);
+			}
+		}
+		return EDirection::Left;
+	}
+
+	const FString& DirectionToString(EDirection Dir)
+	{
+		return DirectionNames[static_cast<int32>(Dir)];
+	}
+
+	EDirection ReadDirectionFromHandle(const TSharedPtr<IPropertyHandle>& Handle)
+	{
+		uint8 Value = 0;
+		Handle->GetValue(Value);
+		return static_cast<EDirection>(Value);
+	}
+
+}  // namespace
+
+TSharedRef<IPropertyTypeCustomization> FCoordinateSystemCustomization::MakeInstance()
+{
+	return MakeShareable(new FCoordinateSystemCustomization);
+}
+
+void FCoordinateSystemCustomization::CustomizeHeader(TSharedRef<IPropertyHandle> PropertyHandle, FDetailWidgetRow& HeaderRow, IPropertyTypeCustomizationUtils& Utils)
+{
+	StructPropertyHandle = PropertyHandle;
+	DetailFont = IDetailLayoutBuilder::GetDetailFont();
+
+	// Get child handles (used for reading committed state)
+	AxisHandles[0] = PropertyHandle->GetChildHandle(GET_MEMBER_NAME_CHECKED(FCoordinateSystem, XAxis));
+	AxisHandles[1] = PropertyHandle->GetChildHandle(GET_MEMBER_NAME_CHECKED(FCoordinateSystem, YAxis));
+	AxisHandles[2] = PropertyHandle->GetChildHandle(GET_MEMBER_NAME_CHECKED(FCoordinateSystem, ZAxis));
+
+	// Build direction options list
+	for (int32 i = 0; i < NumDirections; ++i)
+	{
+		DirectionOptions.Add(MakeShareable(new FString(DirectionNames[i])));
+	}
+
+	// Read current committed values into staged state
+	ReadFromProperty();
+
+	static const FText AxisLabels[] = {
+		LOCTEXT("XAxis", "X"),
+		LOCTEXT("YAxis", "Y"),
+		LOCTEXT("ZAxis", "Z")
+	};
+
+	// Helper to build one axis: [Label] [ComboBox]
+	auto MakeAxisSlot = [this](int32 Index) -> TSharedRef<SWidget>
+		{
+			return SNew(SHorizontalBox)
+				+ SHorizontalBox::Slot()
+				.AutoWidth()
+				.VAlign(VAlign_Center)
+				.Padding(0.f, 0.f, 4.f, 0.f)
+				[
+					SNew(STextBlock)
+						.Text(AxisLabels[Index])
+						.Font(DetailFont)
+				]
+				+ SHorizontalBox::Slot()
+				.AutoWidth()
+				[
+					SNew(SComboBox<TSharedPtr<FString>>)
+						.OptionsSource(&DirectionOptions)
+						.OnSelectionChanged_Lambda([this, Index](TSharedPtr<FString> NewVal, ESelectInfo::Type SelectInfo)
+							{
+								OnAxisSelectionChanged(NewVal, SelectInfo, Index);
+							})
+						.OnGenerateWidget(this, &FCoordinateSystemCustomization::GenerateDirectionWidget)
+						.Content()
+						[
+							SNew(STextBlock)
+								.Text(this, &FCoordinateSystemCustomization::GetCurrentText, Index)
+								.Font(DetailFont)
+						]
+				];
+		};
+
+	HeaderRow
+		.NameContent()
+		[
+			PropertyHandle->CreatePropertyNameWidget()
+		]
+		.ValueContent()
+		.MinDesiredWidth(400.f)
+		[
+			SNew(SHorizontalBox)
+				+ SHorizontalBox::Slot()
+				.AutoWidth()
+				.Padding(2.f, 0.f)
+				[
+					MakeAxisSlot(0)
+				]
+				+ SHorizontalBox::Slot()
+				.AutoWidth()
+				.Padding(2.f, 0.f)
+				[
+					MakeAxisSlot(1)
+				]
+				+ SHorizontalBox::Slot()
+				.AutoWidth()
+				.Padding(2.f, 0.f)
+				[
+					MakeAxisSlot(2)
+				]
+				+ SHorizontalBox::Slot()
+				.AutoWidth()
+				.Padding(8.f, 0.f, 0.f, 0.f)
+				[
+					SNew(SButton)
+						.OnClicked(this, &FCoordinateSystemCustomization::OnApplyClicked)
+						.IsEnabled_Lambda([this]() { return IsValid() && HasPendingChanges(); })
+						.HAlign(HAlign_Center)
+						.VAlign(VAlign_Center)
+						.ToolTipText_Lambda([this]()
+							{
+								if (!IsValid())
+								{
+									return LOCTEXT("InvalidTooltip",
+										"Each axis must use a different spatial dimension (Left/Right, Up/Down, Front/Back)");
+								}
+								if (!HasPendingChanges())
+								{
+									return LOCTEXT("NoChangesTooltip", "No pending changes");
+								}
+								return LOCTEXT("ApplyTooltip", "Apply coordinate system changes");
+							})
+						[
+							SNew(STextBlock)
+								.Text(LOCTEXT("Apply", "Apply"))
+								.Font(DetailFont)
+						]
+				]
+		];
+}
+
+void FCoordinateSystemCustomization::CustomizeChildren(TSharedRef<IPropertyHandle> PropertyHandle, IDetailChildrenBuilder& ChildBuilder, IPropertyTypeCustomizationUtils& Utils)
+{
+	// Everything is in the header row - no child rows needed
+}
+
+void FCoordinateSystemCustomization::ReadFromProperty()
+{
+	for (int32 i = 0; i < 3; ++i)
+	{
+		if (AxisHandles[i].IsValid())
+		{
+			StagedAxes[i] = ReadDirectionFromHandle(AxisHandles[i]);
+		}
+	}
+}
+
+bool FCoordinateSystemCustomization::HasPendingChanges() const
+{
+	for (int32 i = 0; i < 3; ++i)
+	{
+		if (AxisHandles[i].IsValid())
+		{
+			if (StagedAxes[i] != ReadDirectionFromHandle(AxisHandles[i]))
+			{
+				return true;
+			}
+		}
+	}
+	return false;
+}
+
+bool FCoordinateSystemCustomization::IsValid() const
+{
+	const int32 DimX = GetDimensionIndex(StagedAxes[0]);
+	const int32 DimY = GetDimensionIndex(StagedAxes[1]);
+	const int32 DimZ = GetDimensionIndex(StagedAxes[2]);
+	return (DimX != DimY) && (DimX != DimZ) && (DimY != DimZ);
+}
+
+FReply FCoordinateSystemCustomization::OnApplyClicked()
+{
+	FScopedTransaction Transaction(LOCTEXT("SetCoordinateSystem", "Set Coordinate System"));
+
+	// Single pre-change notification
+	StructPropertyHandle->NotifyPreChange();
+
+	// Write all three axes at once via raw memory - no per-property notifications
+	TArray<void*> RawData;
+	StructPropertyHandle->AccessRawData(RawData);
+	for (void* Data : RawData)
+	{
+		if (FCoordinateSystem* CS = static_cast<FCoordinateSystem*>(Data))
+		{
+			CS->XAxis = StagedAxes[0];
+			CS->YAxis = StagedAxes[1];
+			CS->ZAxis = StagedAxes[2];
+		}
+	}
+
+	// Single post-change notification - one event to UDNA
+	StructPropertyHandle->NotifyPostChange(EPropertyChangeType::ValueSet);
+
+	return FReply::Handled();
+}
+
+void FCoordinateSystemCustomization::OnAxisSelectionChanged(
+	TSharedPtr<FString> NewValue,
+	ESelectInfo::Type SelectInfo,
+	int32 AxisIndex)
+{
+	if (NewValue.IsValid() && AxisIndex >= 0 && AxisIndex < 3)
+	{
+		StagedAxes[AxisIndex] = DirectionFromString(*NewValue);
+	}
+}
+
+TSharedRef<SWidget> FCoordinateSystemCustomization::GenerateDirectionWidget(TSharedPtr<FString> InItem)
+{
+	return SNew(STextBlock)
+		.Text(FText::FromString(*InItem))
+		.Font(DetailFont);
+}
+
+FText FCoordinateSystemCustomization::GetCurrentText(int32 AxisIndex) const
+{
+	if (AxisIndex >= 0 && AxisIndex < 3)
+	{
+		return FText::FromString(DirectionToString(StagedAxes[AxisIndex]));
+	}
+	return FText::GetEmpty();
+}
+
+#undef LOCTEXT_NAMESPACE
diff --git a/Engine/Plugins/Animation/RigLogic/Source/RigLogicEditor/Private/CoordinateSystemCustomization.h b/Engine/Plugins/Animation/RigLogic/Source/RigLogicEditor/Private/CoordinateSystemCustomization.h
new file mode 100644
--- /dev/null
+++ b/Engine/Plugins/Animation/RigLogic/Source/RigLogicEditor/Private/CoordinateSystemCustomization.h
@@ -0,0 +1,48 @@
+// Copyright Epic Games, Inc. All Rights Reserved.
+
+#pragma once
+
+#include "IPropertyTypeCustomization.h"
+#include "Fonts/SlateFontInfo.h"
+#include "Input/Reply.h"
+#include "DNACommon.h"
+
+class FCoordinateSystemCustomization : public IPropertyTypeCustomization
+{
+public:
+	static TSharedRef<IPropertyTypeCustomization> MakeInstance();
+
+	virtual void CustomizeHeader(TSharedRef<IPropertyHandle> PropertyHandle, FDetailWidgetRow& HeaderRow, IPropertyTypeCustomizationUtils& Utils) override;
+	virtual void CustomizeChildren(TSharedRef<IPropertyHandle> PropertyHandle, IDetailChildrenBuilder& ChildBuilder, IPropertyTypeCustomizationUtils& Utils) override;
+
+private:
+	/** Read the current committed values from the property handles */
+	void ReadFromProperty();
+
+	/** Is the staged state different from the committed property? */
+	bool HasPendingChanges() const;
+
+	/** Is the staged coordinate system valid? (3 distinct spatial dimensions) */
+	bool IsValid() const;
+
+	FReply OnApplyClicked();
+
+	TSharedRef<SWidget> GenerateDirectionWidget(TSharedPtr<FString> InItem);
+	FText GetCurrentText(int32 AxisIndex) const;
+	void OnAxisSelectionChanged(TSharedPtr<FString> NewValue, ESelectInfo::Type SelectInfo, int32 AxisIndex);
+
+	/** The struct-level property handle */
+	TSharedPtr<IPropertyHandle> StructPropertyHandle;
+
+	/** Child handles for XAxis, YAxis, ZAxis */
+	TSharedPtr<IPropertyHandle> AxisHandles[3];
+
+	/** Staged values (local edits before Apply) */
+	EDirection StagedAxes[3] = { EDirection::Left, EDirection::Front, EDirection::Up };
+
+	/** Options for the combo boxes */
+	TArray<TSharedPtr<FString>> DirectionOptions;
+
+	/** Cached detail panel font */
+	FSlateFontInfo DetailFont;
+};
diff --git a/Engine/Plugins/Animation/RigLogic/Source/RigLogicEditor/Private/RigLogicEditor.cpp b/Engine/Plugins/Animation/RigLogic/Source/RigLogicEditor/Private/RigLogicEditor.cpp
--- a/Engine/Plugins/Animation/RigLogic/Source/RigLogicEditor/Private/RigLogicEditor.cpp
+++ b/Engine/Plugins/Animation/RigLogic/Source/RigLogicEditor/Private/RigLogicEditor.cpp
@@ -9,6 +9,7 @@
 #include "ISettingsModule.h"
 #include "DNAAssetUserData.h"
 #include "DNA.h"
+#include "CoordinateSystemCustomization.h"
 #include "PropertyEditorModule.h"
 #include "Misc/Paths.h"
 #include "AssetRegistry/IAssetRegistry.h"
@@ -25,6 +26,12 @@ DEFINE_LOG_CATEGORY_STATIC(LogRigLogicEditor, Log, All);
 void FRigLogicEditor::StartupModule()
 {
 	UObject::FAssetRegistryTag::OnGetExtraObjectTagsWithContext.AddStatic(&GetAssetRegistryTagsForDNA);
+
+	FPropertyEditorModule& PropertyModule = FModuleManager::LoadModuleChecked<FPropertyEditorModule>("PropertyEditor");
+	PropertyModule.RegisterCustomPropertyTypeLayout(
+		TEXT("CoordinateSystem"),
+		FOnGetPropertyTypeCustomizationInstance::CreateStatic(&FCoordinateSystemCustomization::MakeInstance)
+	);
 }
 
 void FRigLogicEditor::GetAssetRegistryTagsForDNA(FAssetRegistryTagsContext Context)
@@ -75,6 +82,11 @@ void FRigLogicEditor::GetAssetRegistryTagsForDNA(FAssetRegistryTagsContext Conte
 
 void FRigLogicEditor::ShutdownModule()
 {
+	if (FModuleManager::Get().IsModuleLoaded("PropertyEditor"))
+	{
+		FPropertyEditorModule& PropertyModule = FModuleManager::GetModuleChecked<FPropertyEditorModule>("PropertyEditor");
+		PropertyModule.UnregisterCustomPropertyTypeLayout(TEXT("CoordinateSystem"));
+	}
 }
 
 
