Fix crash when no DNA is chosen for import when importing on SkeletalMesh

#Synced-CL 52343460

diff --git a/Engine/Plugins/Animation/RigLogic/Source/RigLogicEditor/Private/AssetDefinitions/AssetDefinition_DNAAsset.cpp b/Engine/Plugins/Animation/RigLogic/Source/RigLogicEditor/Private/AssetDefinitions/AssetDefinition_DNAAsset.cpp
--- a/Engine/Plugins/Animation/RigLogic/Source/RigLogicEditor/Private/AssetDefinitions/AssetDefinition_DNAAsset.cpp
+++ b/Engine/Plugins/Animation/RigLogic/Source/RigLogicEditor/Private/AssetDefinitions/AssetDefinition_DNAAsset.cpp
@@ -33,7 +33,7 @@ namespace MenuExtention_SkelMeshDNAAsset
 		if(!DNAImporter->ImportDNAWithPrompt(SkeletalMeshes))
 		{
 			const FText Message = LOCTEXT("DNAImportFromSkeletalMeshMenuFail", "Importing of DNA failed");
-			UE_LOGFMT(LogDNAAssetDefinition, Error, "{}", *Message.ToString());
+			UE_LOGFMT(LogDNAAssetDefinition, Error, "{Name}", *Message.ToString());
 		}
 	}
 
